[Harmonix] Move Harmonix plugin into Engine/Plugins/Runtime to make it available to licensees in UE 5.4

diff --git a/Engine/Plugins/Runtime/Harmonix/Source/HarmonixDsp/Private/HarmonixDsp/Effects/Vocoder.cpp b/Engine/Plugins/Runtime/Harmonix/Source/HarmonixDsp/Private/HarmonixDsp/Effects/Vocoder.cpp
new file mode 100644
--- /dev/null
+++ b/Engine/Plugins/Runtime/Harmonix/Source/HarmonixDsp/Private/HarmonixDsp/Effects/Vocoder.cpp
@@ -0,0 +1,514 @@
+// Copyright Epic Games, Inc. All Rights Reserved.
+#include "HarmonixDsp/Effects/Vocoder.h"
+#include "HarmonixDsp/Effects/FourBiquads.h"
+
+#include "HAL/PlatformMemory.h"
+#include "HAL/PlatformMath.h"
+
+#define FILTER_COEF_RAMP_MS   15.0f
+
+
+namespace Harmonix::Dsp::Effects
+{
+
+void FVocoder::Init()
+{
+}
+
+FVocoder::FVocoder()
+{}
+
+FVocoder::~FVocoder()
+{
+	DeleteFilters();
+	DeleteBandFrequencies();
+	DeleteCarrierFilterCoefficients();
+	DeleteModulatorFilterCoefficients();
+}
+
+void FVocoder::Setup(int32 ChannelCount, float SampleRate, int32 FrameCount, EVocoderBandConfig BandConfig)
+{
+	TargetSettings.ChannelCount = ChannelCount;
+	TargetSettings.SampleRate = SampleRate;
+	TargetSettings.FrameCount = FrameCount;
+	TargetSettings.BandConfig = BandConfig;
+	ResolveSettings(true);
+}
+
+void FVocoder::Process(const TAudioBuffer<float>& InCarrier, const TAudioBuffer<float>& InModulator, TAudioBuffer<float>& OutBuffer)
+{
+	if (!CurrentSettings.IsEnabled)
+	{
+		OutBuffer.Copy(InCarrier);
+		return;
+	}
+
+	// Validate channel counts of received buffers...
+	checkSlow(CurrentSettings.ChannelCount == InModulator.GetNumValidChannels() &&
+		CurrentSettings.ChannelCount == OutBuffer.GetNumValidChannels() &&
+		CurrentSettings.ChannelCount == InCarrier.GetNumValidChannels());
+
+	// Validate frame counts of received buffers...
+	int32 FrameCount = InCarrier.GetLengthInFrames();
+	checkSlow(FrameCount == InModulator.GetLengthInFrames() && FrameCount == OutBuffer.GetLengthInFrames());
+
+	WorkingCarrier.Copy(InCarrier, CurrentSettings.CarrierGain, -1.0f, 1.0f);
+	WorkingModulator.Copy(InModulator, CurrentSettings.ModulatorGain, -1.0f, 1.0f);
+	OutBuffer.ZeroData();
+
+	ProcessTimeDomain(FrameCount);
+
+	if (OutBuffer.GetIsInterleaved())
+	{
+		OutBuffer.Interleave(WorkingOutput);
+	}
+	else
+	{
+		OutBuffer.Copy(WorkingOutput);
+	}
+}
+
+void FVocoder::ProcessTimeDomain(int32 FrameCount)
+{
+	ResolveSettings();
+
+	const int32& BandCount = CurrentSettings.GetBandConfig().BandCount;
+	float GeneralMultiplier = CurrentSettings.Wet * CurrentSettings.OutputGain;
+#if !UE_BUILD_DEBUG || SIMD_IN_DEBUG
+	//In-loop variables (declared here to avoid stack allocations in loop)
+	VectorRegister4Float GeneralMultiplierVec = VectorSetFloat1(GeneralMultiplier);
+	VectorRegister4Float AttackVec = VectorSetFloat1(CurrentSettings.Attack);
+	VectorRegister4Float ReleaseVec = VectorSetFloat1(CurrentSettings.Release);
+	float FrameOutput;
+	for (int32 BandIdx = 0; BandIdx < BandCount; BandIdx++)
+	{
+		FVocoderBand Band = CurrentSettings.Bands[BandIdx];
+		WorkingBandGains[BandIdx] = Band.Gain * ((!CurrentSettings.Soloing || Band.Solo) ? 1.0f : 0.0f);
+	}
+
+	for (int32 ChannelIdx = 0; ChannelIdx < CurrentSettings.ChannelCount; ChannelIdx++)
+	{
+		float* CarrierChannelData = WorkingCarrier.GetValidChannelData(ChannelIdx);
+		float* ModulatorChannelData = WorkingModulator.GetValidChannelData(ChannelIdx);
+		float* OutputChannelData = WorkingOutput.GetValidChannelData(ChannelIdx);
+		LowReductionFilters[ChannelIdx].Process(CarrierChannelData, CarrierChannelData, FrameCount, LowReductionFilterCoefs);
+		HighIncreaseFilters[ChannelIdx].Process(CarrierChannelData, CarrierChannelData, FrameCount, HighIncreaseFilterCoefs);
+
+		for (int32 FrameIdx = 0; FrameIdx < FrameCount; FrameIdx++)
+		{
+			FrameOutput = 0.0f;
+			for (int32 BandIdx = 0; BandIdx < BandCount; BandIdx += 4)
+			{
+				VectorRegister4Float ModulatorFilteredVec = ModulatorFilters[BandIdx / 4][ChannelIdx].ProcessOne(VectorSetFloat1(ModulatorChannelData[FrameIdx]));
+
+				VectorRegister4Float EnvelopeVec = VectorMax(VectorMultiply(ModulatorFilteredVec, VectorSetFloat1(-1.0f)), ModulatorFilteredVec);
+				VectorRegister4Float PrevEnvelopeVec = VectorLoadAligned(&EnvelopeSamples[ChannelIdx][BandIdx]);
+				VectorRegister4Float AlphaVec = VectorSelect(VectorCompareLT(EnvelopeVec, PrevEnvelopeVec), ReleaseVec, AttackVec);
+				EnvelopeVec = VectorAdd(VectorMultiply(PrevEnvelopeVec, VectorSubtract(GlobalVectorConstants::FloatOne, AlphaVec)), VectorMultiply(EnvelopeVec, AlphaVec));
+				VectorStoreAligned(EnvelopeVec, &EnvelopeSamples[ChannelIdx][BandIdx]);
+
+				VectorRegister4Float CarrierFilteredVec = CarrierFilters[BandIdx / 4][ChannelIdx].ProcessOne(VectorSetFloat1(CarrierChannelData[FrameIdx]));
+				union { VectorRegister4Float OutputVec; float OutBuffer[4]; };
+				OutputVec = VectorMultiply(CarrierFilteredVec, EnvelopeVec);
+				OutputVec = VectorMultiply(OutputVec, GeneralMultiplierVec);
+				OutputVec = VectorMultiply(OutputVec, VectorLoadAligned(&WorkingBandGains[BandIdx]));
+
+				for (int32 LaneIdx = 0; LaneIdx < 4; LaneIdx++)
+				{
+					FrameOutput += OutBuffer[LaneIdx];
+				}
+			}
+			OutputChannelData[FrameIdx] = FrameOutput;
+		}
+	}
+#else
+	float** CarrierData = WorkingCarrier.GetData();
+
+	for (int32 ChannelIdx = 0; ChannelIdx < CurrentSettings.ChannelCount; ChannelIdx++)
+	{
+		LowReductionFilters[ChannelIdx].Process(CarrierData[ChannelIdx], CarrierData[ChannelIdx], FrameCount, LowReductionFilterCoefs);
+		HighIncreaseFilters[ChannelIdx].Process(CarrierData[ChannelIdx], CarrierData[ChannelIdx], FrameCount, HighIncreaseFilterCoefs);
+	}
+
+	for (int32 BandIdx = 0; BandIdx < CurrentSettings.GetBandConfig().BandCount; BandIdx++)
+	{
+		//if(CurrentSettings.SingleBand && Band != CurrentSettings.FocusedBand)
+		if (CurrentSettings.Soloing && !CurrentSettings.Bands[BandIdx].Solo)
+		{
+			continue;
+		}
+
+		// Prepare a work buffer for this single Band of the modulator...
+		BandWork.Copy(WorkingModulator);
+
+		for (int32 ChannelIdx = 0; ChannelIdx < CurrentSettings.ChannelCount; ChannelIdx++)
+		{
+			float* BandWorkChannelData = BandWork.GetValidChannelData(ChannelIdx);
+			ModulatorFilters[BandIdx][ChannelIdx].Process(BandWorkChannelData, BandWorkChannelData, FrameCount, ModulatorFilterCoefs[Band]);
+
+			// Absolute value is helpful for envelope detection...
+			union 
+			{ 
+				VectorRegister4Float AbsWorkVec; 
+				float AbsWork[4]; 
+			};
+			for (int32 FrameIdx = 0; FrameIdx < FrameCount; FrameIdx += 4)
+			{
+				AbsWorkVec = MakeVectorRegisterFloat(BandWorkChannelData[FrameIdx], BandWorkChannelData[FrameIdx + 1], BandWorkChannelData[FrameIdx + 2], BandWorkChannelData[FrameIdx + 3]);
+				AbsWorkVec = VectorMax(VectorMultiply(AbsWorkVec, VectorSetFloat1(-1.0f)), AbsWorkVec);
+				for (int32 LaneIdx = 0; LaneIdx < 4; LaneIdx++)
+				{
+					BandWorkChannelData[FrameIdx + LaneIdx] = AbsWork[LaneIdx];
+				}
+			}
+
+			float& EnvelopeSample = EnvelopeSamples[BandIdx][ChannelIdx];
+			for (int32 FrameIdx = 0; FrameIdx < FrameCount; FrameIdx++)
+			{
+				// Choose our alpha, then update our envelope follower for this Band...
+				float Alpha = (BandWorkChannelData[FrameIdx] >= EnvelopeSample) ? CurrentSettings.Attack : CurrentSettings.Release;
+				EnvelopeSample = (EnvelopeSample * (1.0f - Alpha)) + (BandWorkChannelData[FrameIdx] * Alpha);
+				// Now that BandWorkChannelData[FrameIdx] is done being used for envelope detection, we will dump output into it...
+				CarrierFilters[BandIdx][ChannelIdx].Process(&CarrierData[ChannelIdx][FrameIdx], &BandWorkChannelData[FrameIdx], 1, CarrierFilterCoefs[BandIdx], GeneralMultiplier * EnvelopeSample * CurrentSettings.Bands[BandIdx].Gain);
+			}
+		}
+		//This Band is done, so we'll accumulate it into our output...
+		WorkingOutput.Accumulate(BandWork);
+	}
+#endif
+
+	// If we want some of the original carrier mixed in, scale it appropriately
+	// and accumulate it into our output
+	if (CurrentSettings.Wet != 1.0f)
+	{
+		WorkingCarrier.Scale(1.0f - CurrentSettings.Wet);
+		WorkingOutput.Accumulate(WorkingCarrier);
+	}
+}
+
+void FVocoder::ResolveSettings(bool InForceAll)
+{
+	// Make sure no one else can mess with settings while we rebuild things...
+	FScopeLock Lock(&SettingsLock);
+
+	// Rebuild only what components are necessary based on what has changed...
+
+	bool ChannelCountChanged = TargetSettings.ChannelCount != CurrentSettings.ChannelCount;
+	bool SampleRateChanged = TargetSettings.SampleRate != CurrentSettings.SampleRate;
+	bool BandConfigChanged = TargetSettings.BandConfig != CurrentSettings.BandConfig;
+	bool FrameCountChanged = TargetSettings.FrameCount != CurrentSettings.FrameCount;
+
+	if (ChannelCountChanged || SampleRateChanged || FrameCountChanged || InForceAll)
+	{
+		InitScratchBuffers();
+	}
+
+	if (ChannelCountChanged || BandConfigChanged || InForceAll)
+	{
+		DeleteFilters();
+		CreateFilters();
+	}
+
+	if (BandConfigChanged || InForceAll)
+	{
+		CreateBandFrequencies();
+	}
+
+	if (BandConfigChanged || SampleRateChanged || TargetSettings.CarrierThin != CurrentSettings.CarrierThin || InForceAll)
+	{
+		DeleteCarrierFilterCoefficients();
+		CreateCarrierFilterCoefficients();
+	}
+
+	if (BandConfigChanged || SampleRateChanged || TargetSettings.ModulatorThin != CurrentSettings.ModulatorThin || InForceAll)
+	{
+		DeleteModulatorFilterCoefficients();
+		CreateModulatorFilterCoefficients();
+	}
+
+#if !UE_BUILD_DEBUG || SIMD_IN_DEBUG
+	if (BandConfigChanged || SampleRateChanged || TargetSettings.CarrierThin != CurrentSettings.CarrierThin || ChannelCountChanged || InForceAll)
+	{
+		SetCarrierSets();
+	}
+	if (BandConfigChanged || SampleRateChanged || TargetSettings.ModulatorThin != CurrentSettings.ModulatorThin || ChannelCountChanged || InForceAll)
+	{
+		SetModulatorSets();
+	}
+#endif
+
+	if (BandConfigChanged || SampleRateChanged || TargetSettings.HighEmphasis != CurrentSettings.HighEmphasis || InForceAll)
+	{
+		InitEmphasisFilterCoefficients();
+	}
+
+	CurrentSettings = TargetSettings;
+}
+
+void FVocoder::InitScratchBuffers()
+{
+	WorkingCarrier.Configure(TargetSettings.ChannelCount, TargetSettings.FrameCount, EAudioBufferCleanupMode::Delete, TargetSettings.SampleRate, false);
+	WorkingCarrier.ZeroData();
+	WorkingModulator.Configure(TargetSettings.ChannelCount, TargetSettings.FrameCount, EAudioBufferCleanupMode::Delete, TargetSettings.SampleRate, false);
+	WorkingModulator.ZeroData();
+	BandWork.Configure(TargetSettings.ChannelCount, TargetSettings.FrameCount, EAudioBufferCleanupMode::Delete, TargetSettings.SampleRate, false);
+	BandWork.ZeroData();
+	WorkingOutput.Configure(TargetSettings.ChannelCount, TargetSettings.FrameCount, EAudioBufferCleanupMode::Delete, TargetSettings.SampleRate, false);
+	WorkingOutput.ZeroData();
+}
+
+void FVocoder::CreateFilters()
+{
+	const int32& BandCount = TargetSettings.GetBandConfig().BandCount;
+
+#if !UE_BUILD_DEBUG || SIMD_IN_DEBUG
+	CarrierFilters = new FFourBiquads * [BandCount / 4];
+	ModulatorFilters = new FFourBiquads * [BandCount / 4];
+	EnvelopeSamples = new float* [TargetSettings.ChannelCount];
+
+	for (int32 BandIdx = 0; BandIdx < BandCount / 4; BandIdx++)
+	{
+		CarrierFilters[BandIdx] = new FFourBiquads[TargetSettings.ChannelCount];
+		ModulatorFilters[BandIdx] = new FFourBiquads[TargetSettings.ChannelCount];
+	}
+	for (int32 ChannelIdx = 0; ChannelIdx < TargetSettings.ChannelCount; ChannelIdx++)
+	{
+		EnvelopeSamples[ChannelIdx] = (float*)FMemory::Malloc(BandCount * sizeof(float), 16);
+		FMemory::Memset(EnvelopeSamples[ChannelIdx], 0, sizeof(float) * BandCount);
+	}
+#else
+	CarrierFilters = new FBiquadFilter * [BandCount];
+	ModulatorFilters = new FBiquadFilter * [BandCount];
+	EnvelopeSamples = new float* [BandCount];
+
+	for (int32 Banddx = 0; BandIdx < BandCount; BandIdx++)
+	{
+		CarrierFilters[BandIdx] = new FBiquadFilter[TargetSettings.ChannelCount];
+		ModulatorFilters[BandIdx] = new FBiquadFilter[TargetSettings.ChannelCount];
+		EnvelopeSamples[BandIdx] = new float[TargetSettings.ChannelCount];
+		FMemory::Memset(EnvelopeSamples[BandIdx], 0, sizeof(float) * TargetSettings.ChannelCount);
+	}
+
+#endif
+	LowReductionFilters = new FBiquadFilter[TargetSettings.ChannelCount];
+	HighIncreaseFilters = new FBiquadFilter[TargetSettings.ChannelCount];
+}
+
+void FVocoder::CreateBandFrequencies()
+{
+	const int32& BandCount = TargetSettings.GetBandConfig().BandCount;
+	BandFrequencies = new float[BandCount];
+	float Ratio = TargetSettings.GetBandConfig().FrequencyRatio;
+
+	// Should yield approximately the same results without w0/sin(w0) and sinh
+	float Q = 0.5f / (0.3465736f * (FMath::Log2(Ratio) * 0.5f)); 
+	CarrierQ = Q * 3.0f;
+	ModulatorQ = Q * 1.5f;
+
+	for (int32 Band = 0; Band < BandCount; Band++)
+	{
+		if (Band == 0)
+		{
+			BandFrequencies[Band] = 80.0f;
+		}
+		else
+		{
+			BandFrequencies[Band] = BandFrequencies[Band - 1] * Ratio;
+		}
+	}
+
+#if !UE_BUILD_DEBUG || SIMD_IN_DEBUG
+	WorkingBandGains = (float*)FMemory::Malloc(BandCount * sizeof(float), 16);
+#endif
+}
+
+void FVocoder::CreateCarrierFilterCoefficients()
+{
+	const int32& BandCount = TargetSettings.GetBandConfig().BandCount;
+
+	CarrierFilterCoefs = new FBiquadFilterCoefs[BandCount];
+
+	FBiquadFilterSettings NewSettings;
+	NewSettings.IsEnabled = true;
+	NewSettings.Type = EBiquadFilterType::BandPass;
+	NewSettings.Q = CarrierQ * TargetSettings.CarrierThin;
+	for (int32 BandIdx = 0; BandIdx < BandCount; BandIdx++)
+	{
+		NewSettings.Freq = BandFrequencies[BandIdx];
+		CarrierFilterCoefs[BandIdx] = FBiquadFilterCoefs(NewSettings, TargetSettings.SampleRate);
+	}
+}
+
+void FVocoder::CreateModulatorFilterCoefficients()
+{
+	const int32& BandCount = TargetSettings.GetBandConfig().BandCount;
+
+	ModulatorFilterCoefs = new FBiquadFilterCoefs[BandCount];
+
+	FBiquadFilterSettings NewSettings;
+	NewSettings.IsEnabled = true;
+	NewSettings.Type = EBiquadFilterType::BandPass;
+	NewSettings.Q = ModulatorQ * TargetSettings.ModulatorThin;
+	for (int32 BandIdx = 0; BandIdx < BandCount; BandIdx++)
+	{
+		NewSettings.Freq = BandFrequencies[BandIdx];
+		ModulatorFilterCoefs[BandIdx] = FBiquadFilterCoefs(NewSettings, TargetSettings.SampleRate);
+	}
+}
+
+#if !UE_BUILD_DEBUG || SIMD_IN_DEBUG
+
+void FVocoder::SetCarrierSets()
+{
+	for (int32 BandIdx = 0; BandIdx < TargetSettings.GetBandConfig().BandCount; BandIdx++)
+	{
+		for (int32 ChannelIdx = 0; ChannelIdx < TargetSettings.ChannelCount; ChannelIdx++)
+		{
+			CarrierFilters[BandIdx / 4][ChannelIdx].SetFilter(BandIdx % 4,
+				(float)CarrierFilterCoefs[BandIdx].B0, (float)CarrierFilterCoefs[BandIdx].B1, (float)CarrierFilterCoefs[BandIdx].B2,
+				(float)CarrierFilterCoefs[BandIdx].A1, (float)CarrierFilterCoefs[BandIdx].A2);
+		}
+	}
+}
+
+void FVocoder::SetModulatorSets()
+{
+	for (int32 BandIdx = 0; BandIdx < TargetSettings.GetBandConfig().BandCount; BandIdx++)
+	{
+		for (int32 ChannelIdx = 0; ChannelIdx < TargetSettings.ChannelCount; ChannelIdx++)
+		{
+			ModulatorFilters[BandIdx / 4][ChannelIdx].SetFilter(BandIdx % 4,
+				(float)ModulatorFilterCoefs[BandIdx].B0, (float)ModulatorFilterCoefs[BandIdx].B1, (float)ModulatorFilterCoefs[BandIdx].B2,
+				(float)ModulatorFilterCoefs[BandIdx].A1, (float)ModulatorFilterCoefs[BandIdx].A2);
+		}
+	}
+}
+
+#endif
+
+void FVocoder::InitEmphasisFilterCoefficients()
+{
+	FBiquadFilterSettings FilterSettings;
+	FilterSettings.IsEnabled = true;
+	FilterSettings.Type = EBiquadFilterType::HighShelf;
+	FilterSettings.Freq = 2000.0f;
+	FilterSettings.Q = 1.0f;
+	FilterSettings.DesignedDBGain = TargetSettings.HighEmphasis * 18.0f;
+	HighIncreaseFilterCoefs = FBiquadFilterCoefs(FilterSettings, TargetSettings.SampleRate);
+
+	FilterSettings.Type = EBiquadFilterType::LowShelf;
+	FilterSettings.Freq = 1000.0f;
+	FilterSettings.Q = 1.0f;
+	FilterSettings.DesignedDBGain = TargetSettings.HighEmphasis * -6.0f;
+	LowReductionFilterCoefs = FBiquadFilterCoefs(FilterSettings, TargetSettings.SampleRate);
+}
+
+void FVocoder::DeleteFilters()
+{
+#if !UE_BUILD_DEBUG || SIMD_IN_DEBUG
+	if (CarrierFilters)
+	{
+		for (int32 BandIdx = 0; BandIdx < GetBandCount() / 4; BandIdx++)
+		{
+			delete[] CarrierFilters[BandIdx];
+		}
+		delete[] CarrierFilters;
+		CarrierFilters = nullptr;
+	}
+
+	if (ModulatorFilters)
+	{
+		for (int32 BandIdx = 0; BandIdx < GetBandCount() / 4; BandIdx++)
+		{
+			delete[] ModulatorFilters[BandIdx];
+		}
+		delete[] ModulatorFilters;
+		ModulatorFilters = nullptr;
+	}
+
+	if (EnvelopeSamples)
+	{
+		for (int32 ChannelIdx = 0; ChannelIdx < CurrentSettings.ChannelCount; ChannelIdx++)
+		{
+			FMemory::Free(EnvelopeSamples[ChannelIdx]);
+		}
+		delete[] EnvelopeSamples;
+		EnvelopeSamples = nullptr;
+	}
+#else
+	if (CarrierFilters)
+	{
+		for (int32 Band = 0; Band < GetBandCount(); Band++)
+		{
+			delete[] CarrierFilters[Band];
+		}
+		delete[] CarrierFilters;
+		CarrierFilters = nullptr;
+	}
+
+	if (ModulatorFilters)
+	{
+		for (int32 Band = 0; Band < GetBandCount(); Band++)
+		{
+			delete[] ModulatorFilters[Band];
+		}
+		delete[] ModulatorFilters;
+		ModulatorFilters = nullptr;
+	}
+
+	if (EnvelopeSamples)
+	{
+		for (int32 Band = 0; Band < GetBandCount(); Band++)
+		{
+			delete[] EnvelopeSamples[Band];
+		}
+		delete[] EnvelopeSamples;
+		EnvelopeSamples = nullptr;
+	}
+#endif
+
+	if (LowReductionFilters)
+	{
+		delete[] LowReductionFilters;
+		LowReductionFilters = nullptr;
+	}
+	if (HighIncreaseFilters)
+	{
+		delete[] HighIncreaseFilters;
+		HighIncreaseFilters = nullptr;
+	}
+}
+
+void FVocoder::DeleteBandFrequencies()
+{
+	if (BandFrequencies)
+	{
+		delete[] BandFrequencies;
+		BandFrequencies = nullptr;
+	}
+
+#if !UE_BUILD_DEBUG || SIMD_IN_DEBUG
+	FMemory::Free(WorkingBandGains);
+	WorkingBandGains = nullptr;
+#endif
+}
+
+void FVocoder::DeleteCarrierFilterCoefficients()
+{
+	if (CarrierFilterCoefs)
+	{
+		delete[] CarrierFilterCoefs;
+		CarrierFilterCoefs = nullptr;
+	}
+}
+
+void FVocoder::DeleteModulatorFilterCoefficients()
+{
+	if (ModulatorFilterCoefs)
+	{
+		delete[] ModulatorFilterCoefs;
+		ModulatorFilterCoefs = nullptr;
+	}
+}
+
+};
\ No newline at end of file
diff --git a/Engine/Plugins/Runtime/Harmonix/Source/HarmonixDsp/Private/HarmonixDsp/FusionSampler/FusionPatch.cpp b/Engine/Plugins/Runtime/Harmonix/Source/HarmonixDsp/Private/HarmonixDsp/FusionSampler/FusionPatch.cpp
new file mode 100644
--- /dev/null
+++ b/Engine/Plugins/Runtime/Harmonix/Source/HarmonixDsp/Private/HarmonixDsp/FusionSampler/FusionPatch.cpp
@@ -0,0 +1,317 @@
+// Copyright Epic Games, Inc. All Rights Reserved.
+
+#include "HarmonixDsp/FusionSampler/FusionPatch.h"
+#include "HarmonixDsp/FusionSampler/FusionPatchCustomVersion.h"
+#include "HarmonixDsp/FusionSampler/SingletonFusionVoicePool.h"
+#include "HarmonixDsp/StretcherAndPitchShifterFactoryConfig.h"
+#include "EditorFramework/AssetImportData.h"
+
+#include "Harmonix/PropertyUtility.h"
+
+#include "UObject/AssetRegistryTagsContext.h"
+#include "UObject/UnrealType.h"
+
+#include "Logging/LogMacros.h"
+#include "Misc/AssertionMacros.h"
+#include "Misc/DataValidation.h"
+
+DEFINE_LOG_CATEGORY_STATIC(LogFusionPatch, Log, All)
+
+
+TSharedPtr<Audio::IProxyData> UFusionPatch::CreateProxyData(const Audio::FProxyDataInitParams& InitParams)
+{
+	if (!RenderableFusionPatchData)
+	{
+		RenderableFusionPatchData = MakeShared<FFusionPatchDataProxy::QueueType>(FusionPatchData);
+
+		(*RenderableFusionPatchData)->InitProxyData(InitParams);
+	}
+	return MakeShared<FFusionPatchDataProxy>(RenderableFusionPatchData);
+}
+
+UFusionPatch::UFusionPatch() : Super()
+{
+	FusionPatchData.CurrentPresetIndex = 0;
+	FFusionPatchSettings& DefaultSettings = FusionPatchData.Presets.Emplace_GetRef();
+	FAdsrSettings& VolumeAdsr = DefaultSettings.Adsrs.Volume;
+	VolumeAdsr.Target = EAdsrTarget::Volume;
+	VolumeAdsr.Depth = 1.0f;
+	VolumeAdsr.IsEnabled = true;
+}
+
+void UFusionPatch::SetCurrentPresetIndex(int32 Index)
+{
+	FusionPatchData.SetCurrentPresetIndex(Index);
+
+	UpdateRenderableForNonTrivialChange();
+}
+
+void UFusionPatch::AddPreset(FFusionPatchSettings& Preset)
+{
+	FusionPatchData.AddPreset(Preset);
+
+	UpdateRenderableForNonTrivialChange();
+}
+
+void UFusionPatch::UpdatePreset(int32 Index, FFusionPatchSettings& Preset)
+{
+	FusionPatchData.UpdatePreset(Index, Preset);
+
+	UpdateRenderableForNonTrivialChange();
+}
+
+void UFusionPatch::UpdateRenderableForNonTrivialChange()
+{
+	if (!RenderableFusionPatchData)
+	{
+		return;
+	}
+
+	Audio::FProxyDataInitParams InitParams;
+
+	auto NewSettings = FFusionPatchDataProxy::QueueType::SharedNodePtrType::CreateSharedRenderable(FusionPatchData);
+	NewSettings->InitProxyData(InitParams);
+	RenderableFusionPatchData->SetNewSettings(NewSettings);
+}
+
+#if WITH_EDITORONLY_DATA
+
+void UFusionPatch::PostEditChangeProperty(FPropertyChangedEvent& PropertyChangedEvent)
+{
+	Super::PostEditChangeProperty(PropertyChangedEvent);
+}
+
+void UFusionPatch::PostEditChangeChainProperty(FPropertyChangedChainEvent& PropertyChangedChainEvent)
+{
+	Super::PostEditChangeChainProperty(PropertyChangedChainEvent);
+
+	// we only care about changes to FusionPatchData
+	const FProperty* MemberChanged = PropertyChangedChainEvent.PropertyChain.GetHead()->GetValue();
+	if (MemberChanged->GetFName() != GET_MEMBER_NAME_CHECKED(UFusionPatch, FusionPatchData))
+	{
+		return;
+	}
+
+	for (FKeyzoneSettings& Keyzone : FusionPatchData.GetKeyzones())
+	{
+		int8 MinNote = Keyzone.MinNote;
+		int8 MaxNote = Keyzone.MaxNote;
+		int8 MinVelocity = Keyzone.MinVelocity;
+		int8 MaxVelocity = Keyzone.MaxVelocity;
+		
+		//clamp note numbers:
+		//MinNote <= MaxNote
+		if (MinNote > MaxNote)
+		{
+			Keyzone.MinNote = MaxNote;
+		}
+
+		//clamp velocities
+		//MinVelocity < MaxVelocity
+ 		if (MinVelocity > MaxVelocity)
+		{
+			Keyzone.MinVelocity = MaxVelocity;
+		}
+	}
+
+	// don't need to do anything if there isn't a RenderableFusionPatchData to copy data to
+	if (!RenderableFusionPatchData)
+	{
+		return;
+	}
+
+	const FProperty* PropertyChanged = PropertyChangedChainEvent.Property;
+	const EPropertyChangeType::Type PropertyChangeType = PropertyChangedChainEvent.ChangeType;
+	Harmonix::EPostEditAction PostEditAction = Harmonix::GetPropertyPostEditAction(PropertyChanged, PropertyChangeType);
+
+	if (PostEditAction == Harmonix::EPostEditAction::UpdateTrivial)
+	{
+		FFusionPatchData* CopyToStruct = *RenderableFusionPatchData;
+		Harmonix::CopyStructProperty(CopyToStruct, &FusionPatchData, PropertyChangedChainEvent);
+	}
+	else if (PostEditAction == Harmonix::EPostEditAction::UpdateNonTrivial)
+	{
+		UpdateRenderableForNonTrivialChange();
+	}
+
+	// otherwise, do nothing
+}
+
+void UFusionPatch::PostInitProperties()
+{
+	if (!HasAnyFlags(RF_ClassDefaultObject) && !AssetImportData)
+	{
+		AssetImportData = NewObject<UAssetImportData>(this, TEXT("AssetImportData"));
+	}
+
+	Super::PostInitProperties();
+}
+
+void UFusionPatch::GetAssetRegistryTags(FAssetRegistryTagsContext Context) const
+{
+	Super::GetAssetRegistryTags(Context);
+
+	if (AssetImportData)
+	{
+		Context.AddTag(FAssetRegistryTag(SourceFileTagName(), AssetImportData->GetSourceData().ToJson(), FAssetRegistryTag::TT_Hidden));
+	}
+}
+
+#endif // WITH_EDITOR_ONLY_DATA
+
+#if WITH_EDITOR
+
+EDataValidationResult UFusionPatch::IsDataValid(FDataValidationContext& Context) const
+{
+	EDataValidationResult Result = Super::IsDataValid(Context);
+	const int32 Version = GetLinkerCustomVersion(FFusionPatchCustomVersion::GUID);
+
+	if (Version < FFusionPatchCustomVersion::LatestVersion)
+	{
+		Context.AddWarning(FText::Format(INVTEXT("Asset saved with outdated version: {0}. Resave asset to latest version: {1}."),
+			Version, FFusionPatchCustomVersion::LatestVersion));
+
+		Result = EDataValidationResult::Invalid;
+	}
+
+	return Result;
+}
+
+#endif // WITH_EDITOR
+
+void UFusionPatch::Serialize(FArchive& Ar)
+{
+	Super::Serialize(Ar);
+
+	Ar.UsingCustomVersion(FFusionPatchCustomVersion::GUID);
+	const int32 Version = Ar.CustomVer(FFusionPatchCustomVersion::GUID);
+
+	if (Ar.IsLoading() && Version < FFusionPatchCustomVersion::PitchShifterNameRedirects)
+	{
+		UE_LOG(LogFusionPatch, Warning, TEXT("Fusion patch was loaded with an outdated version. Asset requires reimport: %s"), *GetPathName());
+	}
+
+	const UStretcherAndPitchShifterFactoryConfig* FactoryConfig = GetDefault<UStretcherAndPitchShifterFactoryConfig>();
+	for (FKeyzoneSettings& Keyzone : FusionPatchData.Keyzones)
+	{
+		if (const FPitchShifterNameRedirect* Redirect = FactoryConfig->FindFactoryNameRedirect(Keyzone.TimeStretchConfig.PitchShifter))
+		{
+			Keyzone.TimeStretchConfig.PitchShifter = Redirect->NewName;
+		}
+	}
+}
+
+void UFusionPatch::PostLoad()
+{
+	Super::PostLoad();
+}
+
+TArray<FKeyzoneSettings>& FFusionPatchData::GetKeyzones()
+{
+	return Keyzones;
+}
+
+const TArray<FKeyzoneSettings>& FFusionPatchData::GetKeyzones() const
+{
+	return Keyzones;
+}
+
+int32 FFusionPatchData::GetCurrentPresetIndex() const
+{
+	return CurrentPresetIndex;
+}
+
+void FFusionPatchData::SetCurrentPresetIndex(int32 Index)
+{
+	ensureAlwaysMsgf(Presets.Num() > 0, TEXT("Fusion Patch Data presets have not been initialized"));
+	ensureAlwaysMsgf(Presets.IsValidIndex(Index), TEXT("assigning Fusion Patch Data to an invalid preset"));
+	CurrentPresetIndex = FMath::Clamp(Index, 0, Presets.Num());
+}
+
+int32 FFusionPatchData::GetNumPresets() const
+{
+	return Presets.Num();
+}
+
+void FFusionPatchData::InitProxyData(const Audio::FProxyDataInitParams& InitParams)
+{
+	for (auto& Keyzone : Keyzones)
+	{
+		Keyzone.InitProxyData(InitParams);
+	}
+}
+
+void FFusionPatchData::AddPreset(const FFusionPatchSettings& Preset)
+{
+	// There shouldn't be that many presets, so this should be a pretty fast search
+	int32 Index = FindPresetIndex(Preset.Name);
+
+	// replace the preset of the same name with the new preset
+	if (Index != INDEX_NONE)
+	{
+		Presets[Index] = Preset;
+	}
+	else
+	{
+		Presets.Add(Preset);
+	}
+}
+
+int32 FFusionPatchData::FindPresetIndex(FName PresetName) const
+{
+	return Presets.IndexOfByPredicate([&PresetName](const FFusionPatchSettings& Elem)
+		{
+			return Elem.Name == PresetName;
+		});
+}
+
+void FFusionPatchData::UpdatePreset(int32 Index, const FFusionPatchSettings& Preset)
+{
+	check(Presets.IsValidIndex(Index));
+	Presets[Index] = Preset;
+}
+
+void FFusionPatchData::DisconnectSampler(const FFusionSampler* Sampler)
+{
+	for (FKeyzoneSettings& Keyzone : Keyzones)
+	{
+		if (Keyzone.SingletonFusionVoicePool.IsValid())
+		{
+			Keyzone.SingletonFusionVoicePool->SamplerDisconnecting(Sampler);
+		}
+	}
+}
+
+bool FFusionPatchData::HasPreset(FName PresetName) const
+{
+	return Presets.ContainsByPredicate([&PresetName](const FFusionPatchSettings& Elem)
+		{
+			return Elem.Name == PresetName;
+		});
+}
+
+const FFusionPatchSettings& FFusionPatchData::GetPreset(int32 Index) const
+{
+	check(Presets.IsValidIndex(Index));
+	return Presets[Index];
+}
+
+FFusionPatchSettings& FFusionPatchData::GetPreset(int32 Index)
+{
+	check(Presets.IsValidIndex(Index));
+	return Presets[Index];
+}
+
+const FFusionPatchSettings& FFusionPatchData::GetCurrentPreset() const
+{
+	check(Presets.Num() > 0);
+	check(Presets.IsValidIndex(CurrentPresetIndex));
+	return Presets[GetCurrentPresetIndex()];
+}
+
+FFusionPatchSettings& FFusionPatchData::GetCurrentPreset()
+{
+	check(Presets.Num() > 0);
+	check(Presets.IsValidIndex(CurrentPresetIndex));
+	return Presets[CurrentPresetIndex];
+}
\ No newline at end of file
diff --git a/Engine/Plugins/Runtime/Harmonix/Source/HarmonixDsp/Private/HarmonixDsp/FusionSampler/FusionSampler.cpp b/Engine/Plugins/Runtime/Harmonix/Source/HarmonixDsp/Private/HarmonixDsp/FusionSampler/FusionSampler.cpp
new file mode 100644
--- /dev/null
+++ b/Engine/Plugins/Runtime/Harmonix/Source/HarmonixDsp/Private/HarmonixDsp/FusionSampler/FusionSampler.cpp
@@ -0,0 +1,2285 @@
+// Copyright Epic Games, Inc. All Rights Reserved.
+#include "HarmonixDsp/FusionSampler/FusionSampler.h"
+#include "HarmonixDsp/FusionSampler/FusionVoice.h"
+#include "HarmonixDsp/FusionSampler/FusionVoicePool.h"
+#include "HarmonixDsp/TimeSyncOption.h"
+
+#include "HarmonixDsp/AudioUtility.h"
+#include "HarmonixDsp/Math/Interp.h"
+
+DEFINE_LOG_CATEGORY(LogFusionSampler);
+
+namespace Defaults
+{
+	static const float kVolumeDb = 12.0f;
+	static const float kPortamentoTimeMs = 300.0f;
+	
+	static const float kMinPitchBend = -200.0f;
+	static const float kMaxPitchBend = 200.0f;
+
+	static const int32 kMaxNumVoices = 100;
+}
+
+// interpolates for mapping [0, 1] to reasonable ranges for each parameter type
+static FLerp gStartPointTimeInterp(0.0f, 1000.0f);
+static FLerp gPortamentoTimeInterp(0.0f, 10000.0f);
+static FInterpEaseIn gFreqInterp(20.0f, 20000.0f, 4.335f);
+static FInterpEaseIn gFilterQInterp(0.1f, 10.0f, 4.03f);
+static FInterpEaseIn gLfoFreqInterp(0.1f, 30.0f, 2.74f);
+static FInterpEaseIn gDelayTimeInterp(0.0001f, 4.0f, 2.0f);
+
+FFusionSampler::FFusionSampler()
+	: RampCallRateHz(1)
+	, Speed(1.0f)
+	, MidiChannelVolume(Defaults::kVolumeDb)
+	, MidiChannelMuted(false)
+	, TrimVolume(Defaults::kVolumeDb)
+	, ExpressionGain(1.0f)
+	, IsPortamentoEnabled(false)
+	, PortamentoMode(EPortamentoMode::Legato)
+	, PortamentoTimeMs(Defaults::kPortamentoTimeMs)
+	, IsPortamentoActive(false)
+	, ExtraPitchBend(0)
+	, PitchBendFactor(1.0f)
+	, FineTuneCents(0.0f)
+	, StartPointMs(0.0f)
+	, MaxNumVoices(Defaults::kMaxNumVoices)
+	, MinPitchBendCents(Defaults::kMinPitchBend)
+	, MaxPitchBendCents(Defaults::kMaxPitchBend)
+	, VoicePool(nullptr)
+	, DelayEnabled(false)
+	, BitCrusherEnabled(false)
+	, DistortionEnabled(false)
+	, VocoderEnabled(false)
+	, VocoderModulatorIndex(0)
+	, CurrentTempoBPM(120.0f)
+	, RawPitchMultiplier(1.0f)
+	, Transposition(0)
+	, PresetIdx(0)
+	, KeyzoneSelectMode(EKeyzoneSelectMode::Layers)
+{
+	// Can't use 'alignas' because we may have been allocated in an array of FusionSamplers.
+	// In that case alignas can't be honored :-(
+	float* BufferStart = VoiceWorkBuffer;
+	if ((size_t)VoiceWorkBuffer & 0xF)
+	{
+		int32 OffsetFloats = (16 - ((size_t)VoiceWorkBuffer % 16)) / sizeof(float);
+		check (OffsetFloats > 0 && OffsetFloats < 4);
+		BufferStart += OffsetFloats;
+	}
+
+	for (int32 Idx = 0; Idx < FAudioBufferConfig::kMaxAudioBufferChannels; ++Idx)
+	{
+		VoiceWorkBufferChannels[Idx] = BufferStart + (kScratchBufferFrames * Idx);
+	}
+
+	PendingNoteActions.Reserve(16);
+	ResetNoteActions(true);
+	ResetNoteStatus();
+
+	SetSampleRate(0.0f);
+	ResetInstrumentState();
+	ResetPatchRelatedState();
+}
+
+FFusionSampler::~FFusionSampler()
+{
+	DumpAllChildren();
+	if (FusionPatchData)
+	{
+		FusionPatchData->DisconnectSampler(this);
+	}
+	SetVoicePool(nullptr, true);
+}
+
+void FFusionSampler::SetVoicePool(FSharedFusionVoicePoolPtr InPool, bool NoCallbacks)
+{
+	if (VoicePool == InPool)
+	{
+		return;
+	}
+
+	FScopeLock Lock(&GetBusLock());
+
+	if (InPool)
+	{
+		InPool->AddClient(this);
+	}
+
+	if (VoicePool)
+	{
+		VoicePool->KillVoices(this, NoCallbacks);
+		VoicePool->RemoveClient(this);
+	}
+
+	VoicePool = InPool;
+
+	if (VoicePool)
+	{
+		SetSampleRate(VoicePool->GetSampleRate());
+	}
+}
+
+void FFusionSampler::VoicePoolWillDestruct(const FFusionVoicePool* InPool)
+{
+	if (VoicePool.IsValid() && VoicePool.Get() == InPool)
+	{
+		FScopeLock Lock(&GetBusLock());
+		SetVoicePool(nullptr, false);
+	}
+}
+
+void FFusionSampler::NoteOn(FMidiVoiceId InVoiceId, int8 InMidiNoteNumber, int8 InVelocity, int8 InMidiChannel, int32 InEventTick, int32 InCurrentTick, float InOffsetMs)
+{
+	// we ignore InMidiChannel because this is a single channel instrument!
+	if (InMidiNoteNumber > MidiConstants::kMaxNote)
+	{
+		return;
+	}
+
+	FScopeLock Lock(&sNoteActionCritSec);
+	
+	FPendingNoteAction* ExitingAction = PendingNoteActions.FindByPredicate([=](const FPendingNoteAction& Action){ return Action.VoiceId == InVoiceId;});
+	int8 PendingVelocity = ExitingAction ? ExitingAction->Velocity : 0;
+
+	if ( InVelocity > PendingVelocity || InVelocity == kNoteOff)
+	{
+		PendingNoteActions.Add(
+			{ 
+				/* .MidiNote    = */ InMidiNoteNumber,
+				/* .Velocity    = */ InVelocity,
+				/* .EventTick   = */ InEventTick,
+				/* .TriggerTick = */ InCurrentTick,
+				/* .OffsetMs    = */ InOffsetMs,
+				/* .FrameOffset = */ 0,
+				/* .VoiceId     = */ InVoiceId
+			});
+	}
+}
+
+void FFusionSampler::NoteOnWithFrameOffset(FMidiVoiceId InVoiceId, int8 InMidiNoteNumber, int8 InVelocity, int8 InMidiChannel, int32 InNumFrames)
+{
+	// we ignore InMidiChannel because this is a single channel instrument!
+	if (InMidiNoteNumber > MidiConstants::kMaxNote)
+	{
+		return;
+	}
+
+	FScopeLock Lock(&sNoteActionCritSec);
+
+	FPendingNoteAction* ExitingAction = PendingNoteActions.FindByPredicate([=](const FPendingNoteAction& Action) { return Action.VoiceId == InVoiceId; });
+	int8 PendingVelocity = ExitingAction ? ExitingAction->Velocity : 0;
+
+	if (InVelocity > PendingVelocity || InVelocity == kNoteOff)
+	{
+		PendingNoteActions.Add(
+			{
+				/* .MidiNote    = */ InMidiNoteNumber,
+				/* .Velocity    = */ InVelocity,
+				/* .EventTick   = */ 0,
+				/* .TriggerTick = */ 0,
+				/* .OffsetMs    = */ 0.0f,
+				/* .FrameOffset = */ InNumFrames,
+				/* .VoiceId     = */ InVoiceId
+			});
+	}
+}
+
+bool FFusionSampler::NoteIsOn(int8 InMidiNoteNumber, int8 InMidiChannel /*= 0*/)
+{
+	if (!VoicePool)
+	{
+		return false;
+	}
+
+	{
+		FScopeLock Lock(&sNoteActionCritSec);
+		if (PendingNoteActions.FindByPredicate([=](const FPendingNoteAction& Action){ return Action.MidiNote == InMidiNoteNumber; }))
+		{
+			return true;
+		}
+	}
+
+	{
+		FScopeLock Lock(&sNoteStatusCritSec);
+		if (NoteStatus[InMidiNoteNumber].KeyedOn)
+		{
+			return true;
+		}
+		else if (NoteStatus[InMidiNoteNumber].NumActiveVoices > 0)
+		{
+			return true;
+		}
+	}
+
+	return false;
+}
+
+void FFusionSampler::NoteOff(FMidiVoiceId InVoiceId, int8 InMidiNoteNumber, int8 InMidiChannel /*= 0*/)
+{
+	// we ignore InMidiChannel because this is a single channel instrument!
+	NoteOn(InVoiceId, InMidiNoteNumber, kNoteOff);
+}
+
+void FFusionSampler::NoteOffWithFrameOffset(FMidiVoiceId InVoiceId, int8 InMidiNoteNumber, int8 InMidiChannel /*= 0*/, int32 InNumFrames /*= 0*/)
+{
+	// we ignore InMidiChannel because this is a single channel instrument!
+	NoteOnWithFrameOffset(InVoiceId, InMidiNoteNumber, kNoteOff, InMidiChannel, InNumFrames);
+}
+
+void FFusionSampler::SetSpeed(float InSpeed, bool InMaintainPitch)
+{
+	FScopeLock Lock(&GetBusLock());
+	if (InSpeed <= .0f)
+	{
+		UE_LOG(LogFusionSampler, Warning, TEXT("Speed value cannot be zero or negative. Forced to 0.0001"));
+		// lower boundary
+		Speed = .0001f;	
+	}
+	else
+	{
+		Speed = InSpeed;
+	}
+
+	// this line applies Speed to portamento time
+	SetPortamentoTime(GetPortamentoTime());   
+
+	Delay.SetSpeed(Speed);
+	Delay.SetDelaySeconds(Delay.GetDelaySeconds());
+
+	MaintainPitchWhenSpeedChanges = InMaintainPitch;
+}
+
+void FFusionSampler::SetTempo(float InBPM)
+{
+	if (CurrentTempoBPM == InBPM)
+	{
+		return;
+	}
+
+	CurrentTempoBPM = InBPM;
+
+	for (ELfoIndex LfoIdx : TEnumRange<ELfoIndex>())
+	{
+		LfoSettings[LfoIdx].TempoBPM = InBPM * Speed;
+		Lfos[LfoIdx].UseSettings(&LfoSettings[LfoIdx]);
+	}
+	
+	UpdateVoiceLfos();
+	Delay.SetTempo(InBPM * Speed);
+}
+
+void FFusionSampler::SetBeat(float Beat)
+{
+	FScopeLock Lock(&GetBusLock());
+	CurrentBeat = Beat;
+}
+
+void FFusionSampler::ResetNoteActions(bool ClearNotes)
+{
+	FScopeLock Lock(&GetBusLock());
+	
+	if (ClearNotes)
+	{
+		PendingNoteActions.Empty(16);
+		return;
+	}
+
+	PendingNoteActions.RemoveAll([](const FPendingNoteAction& Action){ return Action.FrameOffset <= 0; });
+}
+
+void FFusionSampler::ResetNoteStatus()
+{
+	FScopeLock Lock(&GetBusLock());
+	for (int32 NoteIdx = 0; NoteIdx < MidiConstants::kMaxNumNotes; ++NoteIdx)
+	{
+		NoteStatus[NoteIdx].KeyedOn = false;
+		NoteStatus[NoteIdx].NumActiveVoices = 0;
+	}
+}
+
+void FFusionSampler::ProcessNoteActions(int32 InNumFrames)
+{
+	// Do a quick test here without locking...
+	if (!PendingNoteActions.Num())
+	{
+		return;
+	}
+
+	// finalize the list of note actions to take
+	TArray<FPendingNoteAction> NoteActionsToDo;
+	{
+		FScopeLock Lock(&sNoteActionCritSec);
+		for (FPendingNoteAction& Action : PendingNoteActions)
+		{
+			if (Action.FrameOffset > 0)
+			{
+				Action.FrameOffset -= InNumFrames;
+				if (Action.FrameOffset < 0)
+				{
+					Action.FrameOffset = 0;
+				}
+			}
+
+			if (Action.FrameOffset == 0)
+			{
+				NoteActionsToDo.Add(Action);
+			}
+		}
+		ResetNoteActions(false);
+	}
+
+	// process each note action individually
+	int32 AvailableNoteOns = MaxNumVoices;
+	for (FPendingNoteAction& Action : NoteActionsToDo)
+	{
+		if (Action.Velocity != kNoteIgnore && Action.FrameOffset == 0)
+		{
+			if (Action.Velocity == kNoteOff)
+			{
+				TArray<uint8> NotesNeedingStopKeyzoneCheck = StopNote(Action.VoiceId, LastVelocity[Action.MidiNote]);
+				// if Action.VoiceId == kMidiVoiceIdAny then NewStopNote may have stopped many notes, and we mayhave many 
+				// notes to check for "on note-off keyzones"...
+				if (Action.VoiceId == FMidiVoiceId::Any() && NotesNeedingStopKeyzoneCheck.Num())
+				{
+					for (uint8 MidiNote : NotesNeedingStopKeyzoneCheck)
+					{
+						if (LastVelocity[MidiNote] > 0)
+						{
+							AvailableNoteOns -= StartNote(Action.VoiceId, MidiNote, LastVelocity[MidiNote], false, 0, 0, 0.0f);
+						}
+					}
+				}
+				// if Action.VoiceId != kMidiVoiceIdAny then are can do our normal check for "on note-off keyzones"...
+				if (Action.VoiceId != FMidiVoiceId::Any() && LastVelocity[Action.MidiNote] > 0)
+				{
+					AvailableNoteOns -= StartNote(Action.VoiceId, Action.MidiNote, LastVelocity[Action.MidiNote], false, 0, 0, 0.0f);
+				}
+
+				LastVelocity[Action.MidiNote] = 0;
+				FScopeLock Lock(&sNoteStatusCritSec);
+				NoteStatus[Action.MidiNote].KeyedOn = false;
+			}
+			else if (AvailableNoteOns > 0)
+			{
+				AvailableNoteOns -= StartNote(Action.VoiceId, Action.MidiNote, Action.Velocity, true, Action.EventTick, Action.TriggerTick, Action.OffsetMs);
+				LastVelocity[Action.MidiNote] = Action.Velocity;
+				FScopeLock Lock(&sNoteStatusCritSec);
+				NoteStatus[Action.MidiNote].KeyedOn = true;
+			}
+		}
+	}
+}
+
+int32 FFusionSampler::GatherMatchingKeyzones(uint8 TransposedNote, uint8 InVelocity, bool IsNoteOn, const TArray<FKeyzoneSettings>& InKeyzones, uint16* OutMathingZones)
+{
+	int32 OutNumMatchingZones = 0;
+	for (int32 KeyzoneIdx = 0; KeyzoneIdx < InKeyzones.Num(); ++KeyzoneIdx)
+	{
+		const FKeyzoneSettings& Keyzone = InKeyzones[KeyzoneIdx];
+		if (!Keyzone.AudioSample)
+		{
+			continue;
+		}
+
+		// When we look for a matching Keyzone, we use the transposed note...
+		if (!Keyzone.ContainsNoteAndVelocity(TransposedNote, InVelocity))
+		{
+			continue;
+		}
+
+		if (Keyzone.IsNoteOnZone() != IsNoteOn)
+		{
+			continue;
+		}
+
+		// completely exclude any Keyzones with a random weight of 0.
+		if ((KeyzoneSelectMode == EKeyzoneSelectMode::Random ||
+			KeyzoneSelectMode == EKeyzoneSelectMode::RandomWithRepetition) &&
+			Keyzone.RandomWeight == 0)
+		{
+			continue;
+		}
+		OutMathingZones[OutNumMatchingZones++] = (uint16)KeyzoneIdx;
+	}
+	return OutNumMatchingZones;
+}
+
+
+
+int32 FFusionSampler::StartNote(FMidiVoiceId InVoiceId, uint8 InMidiNoteNumber, uint8 InVelocity, bool IsNoteOn, int32 InEventTick, int32 InTriggerTick, float InOffsetMs)
+{
+	check(InMidiNoteNumber <= MidiConstants::kMaxNote);
+
+	if (VoicePool == nullptr)
+	{
+		UE_LOG(LogFusionSampler, Warning, TEXT("Cannot \"start note \" on a sampler with no voice pool!"));
+		return 0;
+	}
+
+	int32 NumVoicesStarted = 0;
+
+	uint8 TransposedNote = InMidiNoteNumber + Transposition;
+
+	// more than one voice may trigger with a note on.
+	// we need to test all the PatchKeyzones to determine which get triggered.
+	if (FusionPatchData)
+	{
+		int8* LayerSelector = (IsNoteOn) ? LastStartLayerSelect : LastStopLayerSelect;
+
+		// maybe too many hops
+  		const TArray<FKeyzoneSettings>& PatchKeyzones = FusionPatchData->GetKeyzones();
+		uint16 MatchingZones[kMaxLayersPerNote];
+		int32  NumMatchingZones = GatherMatchingKeyzones(TransposedNote, InVelocity, IsNoteOn, PatchKeyzones, MatchingZones);
+
+		if (NumMatchingZones != 0)
+		{
+			if (NumMatchingZones == 1 || KeyzoneSelectMode == EKeyzoneSelectMode::Layers)
+			{
+				for (int32 Idx = 0; Idx < NumMatchingZones; Idx++)
+				{
+					if (TryKeyOnZone(InVoiceId, InMidiNoteNumber, TransposedNote, InVelocity, InEventTick, InTriggerTick, InOffsetMs, &PatchKeyzones[MatchingZones[Idx]]))
+					{
+						++NumVoicesStarted;
+					}
+				}
+			}
+			else
+			{
+				switch (KeyzoneSelectMode)
+				{
+				case EKeyzoneSelectMode::Random:
+				case EKeyzoneSelectMode::RandomWithRepetition:
+				{
+					check(NumMatchingZones > 1);
+
+					uint32 TriggerKeyzone = 0;
+					int32 ExcludedKeyzone = -1;
+					if (KeyzoneSelectMode == EKeyzoneSelectMode::Random)
+					{
+						ExcludedKeyzone = LayerSelector[InMidiNoteNumber];
+					}
+
+					// gather total weight of all PatchKeyzones except the excluded one
+					float TotalWeight = 0;
+					for (int32 Idx = 0; Idx < NumMatchingZones; ++Idx)
+					{
+						if (Idx != ExcludedKeyzone)
+						{
+							TotalWeight += PatchKeyzones[MatchingZones[Idx]].RandomWeight;
+						}
+					}
+
+					// pick randomly based on total weight
+					float Rnd = FMath::FRand() * TotalWeight;
+					for (int32 Idx = 0; Idx < NumMatchingZones; ++Idx)
+					{
+						if (Idx != ExcludedKeyzone)
+						{
+							// ensure we always give a result
+							TriggerKeyzone = Idx; 
+							Rnd -= PatchKeyzones[MatchingZones[Idx]].RandomWeight;
+							if (Rnd < 0)
+							{
+								break;
+							}
+						}
+					}
+
+					LayerSelector[InMidiNoteNumber] = (int8)TriggerKeyzone;
+					if (TryKeyOnZone(InVoiceId, InMidiNoteNumber, TransposedNote, InVelocity, InEventTick, InTriggerTick, InOffsetMs, &PatchKeyzones[MatchingZones[TriggerKeyzone]]))
+					{
+						++NumVoicesStarted;
+					}
+				}
+				break;
+				case EKeyzoneSelectMode::Cycle:
+				{
+					uint32 TriggerKeyzone = ((uint32)LayerSelector[InMidiNoteNumber] + 1) % NumMatchingZones;
+					LayerSelector[InMidiNoteNumber] = (int8)TriggerKeyzone;
+					if (TryKeyOnZone(InVoiceId, InMidiNoteNumber, TransposedNote, InVelocity, InEventTick, InTriggerTick, InOffsetMs, &PatchKeyzones[MatchingZones[TriggerKeyzone]]))
+					{
+						++NumVoicesStarted;
+					}
+				}
+				break;
+				}
+			}
+		}
+	}
+
+	VoicePool->FastReleaseExcessVoices();
+	VoicePool->FastReleaseExcessVoices(this);
+	VoiceUsage.Merge((double)ActiveVoices.Num());
+
+	return NumVoicesStarted;
+}
+
+bool FFusionSampler::TryKeyOnZone(FMidiVoiceId InVoiceId, uint8 InTriggeredNote, uint8 InTransposedNote, uint8 InVelocity, int32 InEventTick, int32 InTriggerTick, float InOffsetMs, const FKeyzoneSettings* InKeyzone)
+{
+	if (!ensure(VoicePool))
+		return false;
+
+	check(InKeyzone);
+
+	static float sVelocityNormalizer = 1.0f / 127.0f;
+
+	// We send in the original, triggered note so that the allocator
+	// properly stops other sounding voices that started as a result of
+	// a trigger on the same note...
+	FFusionVoice* Voice = VoicePool->GetFreeVoice(this, InVoiceId, InKeyzone, [this](FFusionVoice* Voice) { return RelinquishVoice(Voice); });
+	if (Voice == nullptr)
+	{
+		UE_LOG(LogFusionSampler, Warning, TEXT("Voice pool failed to create voice, can't play note!"), InTriggeredNote);
+		return false;
+	}
+
+	{
+		FScopeLock Lock(&sNoteStatusCritSec);
+		NoteStatus[InTriggeredNote].NumActiveVoices += 1;
+	}
+
+	ActiveVoices.AddTail(Voice);
+
+	// From here on out we use the transposed note to that the correct pitch is played...
+
+	for (uint8 ModulatorIdx = 0; ModulatorIdx < (uint8)EModulatorTarget::Num; ++ModulatorIdx)
+	{
+		Randomizer[ModulatorIdx].Modulate(FMath::FRand());
+	}
+
+
+	float VelocityNorm = (float)InVelocity * sVelocityNormalizer;
+	for (uint8 ModulatorIdx = 0; ModulatorIdx < (uint32)EModulatorTarget::Num; ++ModulatorIdx)
+	{
+		VelocityModulator[ModulatorIdx].Modulate(VelocityNorm);
+	}
+
+	float AttackGain = VelocityNorm; // should this really just be linear?
+	if (!InKeyzone->bVelocityToGain)
+	{
+		AttackGain = 1.0f;
+	}
+
+	if (GetIsPortamentoEnabled())
+	{
+		PortamentoPitchRamper.SetTarget(InTransposedNote);
+		if (!IsPortamentoActive)
+		{
+			PortamentoPitchRamper.SnapToTarget();
+		}
+
+		IsPortamentoActive = true;
+	}
+
+	double AttackStartPointMs = StartPointMs;
+	AttackStartPointMs += (double)InOffsetMs;
+	if (AttackStartPointMs < 0.0)
+	{
+		AttackStartPointMs = 0.0;
+	}
+
+	Voice->SetPitchOffset(FineTuneCents);
+
+// TODO:
+//#if FUSION_VOICE_DEBUG_DUMP_ENABLED
+//	Voice->AttackWithTargetNote(InTransposedNote, AttackGain, InEventTick, InTriggerTick, AttackStartPointMs, InKeyzone->AudioSample->GetFile().Str());
+//#else
+	Voice->AttackWithTargetNote(InTransposedNote, AttackGain, InEventTick, InTriggerTick, AttackStartPointMs);
+//#endif
+
+	return true;
+}
+
+TArray<uint8> FFusionSampler::StopNote(FMidiVoiceId InVoiceId, uint8 LastTriggeredVelocity)
+{
+	TArray<uint8> MidiNotesToCheckForStopKeyzones;
+
+	FFusionVoice* Voice = nullptr;
+	if (VoicePool == nullptr)
+	{
+		return MidiNotesToCheckForStopKeyzones;
+	}
+
+	// find Voice with the ID (there can be more than one since keyzones can overlap)
+	for (auto Node = ActiveVoices.GetHead(); Node; Node = Node->GetNextNode())
+	{
+		Voice = Node->GetValue();
+
+		// ignore notes already being released or idle
+		if (Voice->GetAdsrStage() == Harmonix::Dsp::Modulators::EAdsrStage::Release
+		 || Voice->GetAdsrStage() == Harmonix::Dsp::Modulators::EAdsrStage::Idle)
+		{
+			continue;
+		}
+
+		if (Voice->MatchesIDs(this, InVoiceId))
+		{
+			MidiNotesToCheckForStopKeyzones.Add(Voice->GetTriggeredMidiNote());
+
+			Voice->Release();
+
+			// possibly disable portamento pitch bends
+			if (PortamentoMode == EPortamentoMode::Legato)
+			{
+				IsPortamentoActive = false;
+			}
+		}
+	}
+	return MidiNotesToCheckForStopKeyzones;
+}
+
+void FFusionSampler::KillAllVoices()
+{
+	FScopeLock Lock(&GetBusLock());
+
+	// isolated scope locks for each of the reset methods
+	{
+		FScopeLock NoteActionLock(&sNoteActionCritSec);
+		ResetNoteActions(true);
+	}
+	{
+		FScopeLock NoteStatusLock(&sNoteStatusCritSec);
+		ResetNoteStatus();
+	}
+	
+	auto Node = ActiveVoices.GetHead();
+	while (Node)
+	{
+		// kill might delete the one we are sitting on~
+		auto NextNode = Node->GetNextNode();
+		FFusionVoice* Voice = Node->GetValue();
+		Voice->Kill();
+		Node = NextNode;
+	}
+
+	ActiveVoices.Empty();
+}
+
+
+void FFusionSampler::AllNotesOff()
+{
+	FScopeLock Lock(&GetBusLock());
+	NoteOff(FMidiVoiceId::Any(), 0);
+}
+
+void FFusionSampler::AllNotesOffWithFrameOffset(int32 InNumFrames)
+{
+	FScopeLock Lock(&GetBusLock());
+	NoteOffWithFrameOffset(FMidiVoiceId::Any(), 0, 0, InNumFrames);
+}
+
+void FFusionSampler::SetPatch(FFusionPatchData* PatchData)
+{
+	if (PatchData == FusionPatchData)
+	{
+		return;
+	}
+
+	FScopeLock Lock(&GetBusLock());
+
+	if (FusionPatchData)
+	{
+		FusionPatchData->DisconnectSampler(this);
+	}
+		
+	FusionPatchData = PatchData;
+
+	if (FusionPatchData)
+	{
+		PresetIdx = FusionPatchData->GetCurrentPresetIndex();
+		ApplyPatchPresetSettings();
+	}
+}
+
+void FFusionSampler::ApplyPatchPresetSettings()
+{
+	check(FusionPatchData);
+
+	check(FusionPatchData->GetNumPresets() > 0);
+	PresetIdx = FMath::Clamp<size_t>(PresetIdx, 0, FusionPatchData->GetNumPresets() - 1);
+
+	const FFusionPatchSettings& Preset = FusionPatchData->GetPreset(PresetIdx);
+	SetTrimVolume(Preset.VolumeDb);
+	SetMinPitchBendCents(Preset.DownPitchBendCents);
+	SetMaxPitchBendCents(Preset.UpPitchBendCents);
+	SetMaxNumVoices(Preset.MaxVoices);
+	SetStartPointMs(Preset.StartPointOffsetMs);
+	SetFineTuneCents(Preset.FineTuneCents);
+	SetPan(Preset.PannerDetails);
+	KeyzoneSelectMode = Preset.KeyzoneSelectMode;
+
+	for (EAdsrIndex Idx : TEnumRange<EAdsrIndex>())
+	{
+		AdsrSettings[Idx].CopySettings(Preset.Adsrs[Idx]);
+		AdsrSettings[Idx].CopyCurveTables(Preset.Adsrs[Idx]);
+	}
+
+	//-------------------------------------
+	// get the filter info
+	//------------------------------------
+	FilterSettings = Preset.Filter;
+
+	//------------------------------------
+	// Read in Lfo settings
+	//------------------------------------
+
+	for (ELfoIndex Idx : TEnumRange<ELfoIndex>())
+	{
+		LfoSettings[Idx].CopySettings(Preset.Lfos[Idx]);
+	}
+
+	//------------------------------------
+	// Read in portamento settings
+	//------------------------------------
+	{
+		const FPortamentoSettings& Portamento = Preset.Portamento;
+		SetIsPortamentoEnabled(Portamento.IsEnabled);
+		SetPortamentoMode(Portamento.Mode);
+		SetPortamentoTime(Portamento.Seconds);
+	}
+	//-------------------------------------
+	// read in delay settings
+	//-------------------------------------
+	{
+		const FDelaySettings& DelaySettings = Preset.Delay;
+		SetDelayEnabled(DelaySettings.IsEnabled);
+		Delay.SetTimeSyncOption(DelaySettings.TimeSyncOption);
+		Delay.SetDelaySeconds(DelaySettings.TimeSeconds);
+		Delay.SetDryGain(DelaySettings.DryGain);
+		Delay.SetWetGain(DelaySettings.WetGain);
+		Delay.SetFeedbackGain(DelaySettings.FeedbackGain);
+		Delay.SetFilterFreq(DelaySettings.EQFreq);
+		Delay.SetFilterType(DelaySettings.EQType);
+		Delay.SetWetFilterEnabled(DelaySettings.EQEnabled);
+		Delay.SetFilterQ(DelaySettings.EQQ);
+		Delay.SetLfoEnabled(DelaySettings.LfoEnabled);
+		Delay.SetLfoTimeSyncOption(DelaySettings.LfoTimeSyncOption);
+		Delay.SetLfoFreq(DelaySettings.LfoRate);
+		Delay.SetLfoDepth(DelaySettings.LfoDepth);
+		Delay.SetStereoSpreadLeft(DelaySettings.PanLeft);
+		Delay.SetStereoSpreadRight(DelaySettings.PanRight);
+		Delay.SetStereoType(DelaySettings.StereoType);
+		// snap settings immediately
+		Delay.SetParamsToTargets();
+	}
+
+	//-------------------------------------
+	// read in bitcrusher settings
+	//-------------------------------------
+	{
+		const FBitCrusherSettings& CrushSettings = Preset.BitCrusher;
+		SetBitCrusherEnabled(CrushSettings.IsEnabled);
+		BitCrusher.Reset();
+		BitCrusher.SetWetGain(CrushSettings.WetGain);
+		BitCrusher.SetCrush(CrushSettings.Crush);
+		BitCrusher.SetSampleHoldFactor(CrushSettings.SampleHoldFactor);
+	}
+
+	//-------------------------------------
+	// read in vocoder settings
+	//-------------------------------------
+	{
+		const FVocoderSettings& VocoderSettings = Preset.Vocoder;
+		VocoderEnabled = VocoderSettings.IsEnabled;
+		Vocoder.SetBandConfig(VocoderSettings.BandConfig);
+		Vocoder.SetSoloing(VocoderSettings.Soloing);
+		Vocoder.SetWet(VocoderSettings.Wet);
+		Vocoder.SetCarrierGain(VocoderSettings.CarrierGain);
+		Vocoder.SetModulatorGain(VocoderSettings.ModulatorGain);
+		Vocoder.SetCarrierThin(VocoderSettings.CarrierThin);
+		Vocoder.SetModulatorThin(VocoderSettings.ModulatorThin);
+		Vocoder.SetAttack(VocoderSettings.Attack);
+		Vocoder.SetRelease(VocoderSettings.Release);
+		Vocoder.SetHighEmphasis(VocoderSettings.HighEmphasis);
+		Vocoder.SetOutputGain(VocoderSettings.OutputGain);
+		Vocoder.SetPerBandSettings(VocoderSettings.Bands);
+	}
+
+	//-------------------------------------
+	// read in distortion settings
+	//-------------------------------------
+	{
+		const FDistortionSettingsV1& DistortionSettings = Preset.Distortion;
+		DistortionEnabled = DistortionSettings.IsEnabled;
+		Distortion.SetSampleRate((int32)GetSamplesPerSecond());
+		Distortion.SetInputGainDb(DistortionSettings.InputGainDb);
+		Distortion.SetOutputGainDb(DistortionSettings.OutputGainDb);
+		Distortion.SetType(DistortionSettings.Type);
+		Distortion.SetOversample(DistortionSettings.Oversample, AudioRendering::kFramesPerRenderBuffer);
+		for (int32 Idx = 0; Idx < FDistortionSettingsV1::kNumFilters; ++Idx)
+		{
+			Distortion.SetupFilter(Idx, DistortionSettings.Filters[Idx]);
+		}
+	}
+}
+
+void FFusionSampler::Prepare(float InSampleRateHz, EAudioBufferChannelLayout InChannelLayout, uint32 InMaxSamples, bool bInAllocateBuffer /*= true*/)
+{
+	FMusicalAudioBus::Prepare(1.0f, InChannelLayout, InMaxSamples, bInAllocateBuffer);
+	SetSpeed(1.0f);
+	SetSampleRate(InSampleRateHz);
+}
+
+void FFusionSampler::SetGainTable(const FGainTable* InGainTable)
+{
+	GainTable = InGainTable;
+}
+
+void FFusionSampler::SetSampleRate(float InSampleRateHz)
+{
+	if ((float)GetSamplesPerSecond() == InSampleRateHz)
+		return;
+
+	// call super class
+	FMusicalAudioBus::SetSampleRate(InSampleRateHz);
+
+	RampCallRateHz = InSampleRateHz / (float)FMusicalAudioBus::GetMaxFramesPerProcessCall();
+
+	ExpressionGainRamper.SetRampTimeMs(RampCallRateHz, AudioRendering::kMicroFadeMs);
+	ExpressionGainRamper.SnapToTarget();
+
+	MidiChannelGainRamper.SetRampTimeMs(RampCallRateHz, AudioRendering::kMicroFadeMs);
+	MidiChannelGainRamper.SnapToTarget();
+
+	MidiChannelMuteGainRamper.SetRampTimeMs(RampCallRateHz, AudioRendering::kMicroFadeMs);
+	MidiChannelMuteGainRamper.SetTarget(MidiChannelMuted ? 0.0f : 1.0f);
+	MidiChannelMuteGainRamper.SnapToTarget();
+
+	PitchBendRamper.SetRampTimeMs(RampCallRateHz, 5.0f);
+	PitchBendRamper.SetTarget(0.0f);
+	PitchBendRamper.SnapToTarget();
+
+	for (ELfoIndex Idx : TEnumRange<ELfoIndex>())
+	{
+		Lfos[Idx].Prepare(InSampleRateHz);
+		Lfos[Idx].UseSettings(&LfoSettings[Idx]);
+	}
+
+	SetPortamentoTime(GetPortamentoTime());
+
+	if (DelayEnabled)
+	{
+		Delay.Prepare(InSampleRateHz, kScratchBufferChannels, 2500.0f);
+	}
+
+	BitCrusher.Setup(kScratchBufferChannels, InSampleRateHz);
+	Vocoder.Setup(kScratchBufferChannels, InSampleRateHz, AudioRendering::kFramesPerRenderBuffer);
+
+}
+
+void FFusionSampler::ResetInstrumentStateImpl()
+{
+	DumpAllChildren();
+	// TODO: Move some reset stuff from ResetPatchRelatedState to here!!!!
+	SetMidiChannelGain(1.0, AudioRendering::kMicroFadeSec);
+	SetMidiExpressionGain(1.0, AudioRendering::kMicroFadeSec);
+	SetMidiChannelMute(false);
+
+	Speed = 1.0f;
+	CurrentBeat = 0.0f;
+	CurrentTempoBPM = 120.0f;
+	RawPitchMultiplier = 1.0f;
+	Transposition = 0;
+	for (int32 Idx = 0; Idx < kMaxSubstreams; ++Idx)
+	{
+		SubstreamGain[Idx] = 1.0f;
+	}
+
+}
+
+void FFusionSampler::ResetMidiStateImpl()
+{
+	AllNotesOff();
+	SetPitchBend(0.0f);
+}
+
+void FFusionSampler::ResetPatchRelatedState()
+{
+	FScopeLock Lock(&GetBusLock());
+
+	ResetNoteActions(true);
+
+	PanSettings.Reset();
+
+	SetTrimVolume(Defaults::kVolumeDb);
+	SetMidiExpressionGain(1.0f);
+
+	SetPitchBend(0.0f);
+	PitchBendRamper.SnapToTarget();
+	SetMaxPitchBendCents(Defaults::kMaxPitchBend);
+	SetMinPitchBendCents(Defaults::kMinPitchBend);
+
+	FilterSettings.ResetToDefaults();
+
+	// set up the default Adsr setting
+
+	for (EAdsrIndex Idx : TEnumRange<EAdsrIndex>())
+	{
+		AdsrSettings[Idx].ResetToDefaults();
+	}
+
+	AdsrSettings.Volume.IsEnabled = true;
+	AdsrSettings.Volume.Target = EAdsrTarget::Volume;
+	AdsrSettings.Assignable.IsEnabled = false;
+	AdsrSettings.Assignable.Target = EAdsrTarget::None;
+
+	// set up the default Lfo settings
+	LfoSettings.Pan.ResetToDefaults();
+	LfoSettings.Pitch.ResetToDefaults();
+
+	LfoSettings.Pan.Target = ELfoTarget::Pan;
+	LfoSettings.Pitch.Target = ELfoTarget::Pitch;
+
+	// set up the default randomizer settings
+	for (uint8 ModulatorIdx = 0; ModulatorIdx < (uint8)EModulatorTarget::Num; ++ModulatorIdx)
+	{
+		Randomizers[ModulatorIdx].Reset();
+		Randomizers[ModulatorIdx].SetTarget(&Harmonix::Dsp::Modulators::FModulatorTarget::kDummyTarget);
+	}
+
+
+	// set up the default randomizer settings
+	for (uint8 ModulatorIdx = 0; ModulatorIdx < (uint8)EModulatorTarget::Num; ++ModulatorIdx)
+	{
+		VelocityModulator[ModulatorIdx].Reset();
+		VelocityModulator[ModulatorIdx].SetTarget(&Harmonix::Dsp::Modulators::FModulatorTarget::kDummyTarget);
+	}
+
+
+	MaxNumVoices = Defaults::kMaxNumVoices;
+
+	SetIsPortamentoEnabled(false);
+	PortamentoPitchRamper.SnapTo(60);
+	SetPortamentoMode(EPortamentoMode::Legato);
+	SetPortamentoTime(Defaults::kPortamentoTimeMs / 1000.0f);
+
+	SetDelayEnabled(false);
+	SetBitCrusherEnabled(false);
+	DistortionEnabled = false;
+
+	FMemory::Memset(LastStartLayerSelect, -1, sizeof(int8) * MidiConstants::kMaxNumNotes);
+	FMemory::Memset(LastStopLayerSelect, -1, sizeof(int8) * MidiConstants::kMaxNumNotes);
+	FMemory::Memset(LastVelocity, 0, sizeof(int8) * MidiConstants::kMaxNumNotes);
+
+	TimeStretchEnvelopeOverride = -1;
+
+	SetPatch(nullptr);
+}
+
+int32 FFusionSampler::GetMaxNumVoices() const
+{
+	if (GetIsPortamentoEnabled())
+	{
+		return 1;
+	}
+
+	return MaxNumVoices;
+}
+
+void FFusionSampler::SetMaxNumVoices(int32 Num)
+{
+	FScopeLock Lock(&GetBusLock());
+
+	MaxNumVoices = Num;
+}
+
+bool FFusionSampler::RelinquishVoice(FFusionVoice* InVoice)
+{
+	FScopeLock Lock(&GetBusLock());
+
+	check(InVoice);
+
+	uint32 NoteNumber = InVoice->GetTargetMidiNote();
+
+	if (auto Node = ActiveVoices.FindNode(InVoice))
+	{
+		ActiveVoices.RemoveNode(Node);
+		FScopeLock LockNoteStatus(&sNoteStatusCritSec);
+		NoteStatus[NoteNumber].NumActiveVoices -= 1;
+		return true;
+	}
+
+	return false;
+}
+
+//No improvements to be made here since we have to fill the array with all the voices
+
+//TODO: to test, load up flow dev solution that will build Fusion VST, then in the fusion VST put in some code that calls NoteIsOn and prints the result.
+//then load the fusion VST in reaper, go play midi note 10, see it print out yes. then print out something with a release
+int32 FFusionSampler::GetNumVoicesInUse(FFusionVoice** VoiceArray) const
+{
+	// TODO, mimic lock logging
+	//HmxThreading::ScopedCritSec scs((const_cast<FusionSampler*>(this)->GetBusLock()), HX_FILE_LINE);
+
+	FScopeLock Lock(&const_cast<FFusionSampler*>(this)->GetBusLock());
+
+	int32 OutNum = 0;
+	for (auto Node = ActiveVoices.GetHead(); Node; Node = Node->GetNextNode())
+	{
+		VoiceArray[OutNum] = Node->GetValue();
+		++OutNum;
+	}
+	return OutNum;
+}
+
+void FFusionSampler::SetIsPortamentoEnabled(bool InEnabled)
+{
+	FScopeLock Lock(&GetBusLock());
+	IsPortamentoEnabled = InEnabled;
+}
+
+void FFusionSampler::SetPortamentoMode(EPortamentoMode InMode)
+{
+	FScopeLock Lock(&GetBusLock());
+	PortamentoMode = InMode;
+}
+
+EPortamentoMode FFusionSampler::GetPortamentoMode() const
+{
+	return PortamentoMode;
+}
+
+void FFusionSampler::SetPortamentoTime(float InTimeSec)
+{
+	FScopeLock Lock(&GetBusLock());
+	PortamentoTimeMs = InTimeSec * 1000.0f;
+	PortamentoPitchRamper.SetRampTimeMs(RampCallRateHz, PortamentoTimeMs / Speed);
+}
+
+float FFusionSampler::GetPortamentoTime() const
+{
+	return PortamentoTimeMs * 0.001f;
+}
+
+float FFusionSampler::GetCurrentPortamentoPitch() const
+{
+	return PortamentoPitchRamper.GetCurrent();
+}
+
+const FPannerDetails& FFusionSampler::GetPan() const
+{
+	return PanSettings;
+}
+
+void FFusionSampler::SetPan(const FPannerDetails& InPannerDetails)
+{
+	FScopeLock Lock(&GetBusLock());
+	PanSettings = InPannerDetails;
+}
+
+void FFusionSampler::SetFineTuneCents(float InCents)
+{
+	FScopeLock Lock(&GetBusLock());
+	FineTuneCents = InCents;
+	for (auto Node = ActiveVoices.GetHead(); Node; Node = Node->GetNextNode())
+	{
+		FFusionVoice* Voice = Node->GetValue();
+		check(Voice);
+		Voice->SetPitchOffset(FineTuneCents);
+	}
+}
+
+float FFusionSampler::GetFineTuneCents() const
+{
+	return FineTuneCents;
+}
+
+float FFusionSampler::GetPitchBend(int8 InMidiChannel /*= 0*/) const
+{
+	// we ignore InMidiChannel because this is a single channel instrument!
+	return PitchBendRamper.GetTarget();
+}
+
+void FFusionSampler::SetPitchBend(float Value, int8 InMidiChannel /*= 0*/)
+{
+	FScopeLock Lock(&GetBusLock());
+	// we ignore InMidiChannel because this is a single channel instrument!
+	check((- 1.0f <= Value) & (Value <= 1.0f));
+	PitchBendRamper.SetTarget(Value);
+}
+
+float FFusionSampler::GetMinPitchBendCents() const
+{
+	return MinPitchBendCents;
+}
+
+void FFusionSampler::SetMinPitchBendCents(float Value)
+{
+	FScopeLock Lock(&GetBusLock());
+	MinPitchBendCents = Value;
+	UpdatePitchBendFactor();
+}
+
+float FFusionSampler::GetMaxPitchBendCents() const
+{
+	return MaxPitchBendCents;
+}
+
+void FFusionSampler::SetMaxPitchBendCents(float Value)
+{
+	// TODO
+	//HmxThreading::ScopedCritSec cs(GetBusLock(), HXFILELINE);
+	FScopeLock Lock(&GetBusLock());
+	MaxPitchBendCents = Value;
+	UpdatePitchBendFactor();
+}
+
+// returns the multiplicative factor used to change the rate of sample frame increment
+// to hit the current pitch bend settings
+float FFusionSampler::GetPitchBendFactor() const
+{
+	return PitchBendFactor;
+}
+
+void FFusionSampler::UpdatePitchBendFactor()
+{
+	FScopeLock Lock(&GetBusLock());
+	float bend = GetRampedPitchBend();
+	float Cents = 0.0f;
+	if (bend > 0.0f)
+	{
+		Cents = bend * MaxPitchBendCents;
+	}
+	else
+	{
+		Cents = -bend * MinPitchBendCents;
+	}
+
+	Cents += ExtraPitchBend * 100;
+	PitchBendFactor = FMath::Pow(2.0f, Cents * HarmonixDsp::kOctavesPerCent);
+}
+
+void FFusionSampler::SetExtraPitchBend(float InSemitones, int8 InChannel)
+{
+	// we ignore InMidiChannel because this is a single channel instrument!
+	ExtraPitchBend = InSemitones;
+}
+
+void  FFusionSampler::SetMidiChannelVolume(float InVolume, float InSeconds, int8 InChannel)
+{
+	FScopeLock Lock(&GetBusLock());
+	MidiChannelVolume = InVolume;
+	float gain = HarmonixDsp::DBToLinear(MidiChannelVolume);
+	float ms = InSeconds * 1000.0f;
+	if (ms < AudioRendering::kMicroFadeMs)
+	{
+		ms = AudioRendering::kMicroFadeMs;
+	}
+	MidiChannelGainRamper.SetRampTimeMs(RampCallRateHz, ms);
+	MidiChannelGainRamper.SetTarget(gain);
+}
+
+float FFusionSampler::GetMidiChannelVolume(int8 InMidiChannel) const
+{
+	float gain = MidiChannelGainRamper.GetCurrent();
+	float db = HarmonixDsp::LinearToDB(gain);
+	return db;
+}
+
+void  FFusionSampler::SetMidiExpressionGain(float gain, float seconds, int8 InMidiChannel)
+{
+	FScopeLock Lock(&GetBusLock());
+	float ms = seconds * 1000.0f;
+	if (ms < AudioRendering::kMicroFadeMs)
+		ms = AudioRendering::kMicroFadeMs;
+	ExpressionGain = gain;
+	ExpressionGainRamper.SetRampTimeMs(RampCallRateHz, ms);
+	ExpressionGainRamper.SetTarget(ExpressionGain);
+}
+
+void  FFusionSampler::SetMidiChannelGain(float gain, float seconds, int8 InMidiChannel)
+{
+	FScopeLock Lock(&GetBusLock());
+	MidiChannelVolume = HarmonixDsp::LinearToDB(gain);
+	float ms = seconds * 1000.0f;
+	if (ms < AudioRendering::kMicroFadeMs)
+		ms = AudioRendering::kMicroFadeMs;
+	MidiChannelGainRamper.SetRampTimeMs(RampCallRateHz, ms);
+	MidiChannelGainRamper.SetTarget(gain);
+}
+
+float FFusionSampler::GetMidiChannelGain(int8 InMidiChannel) const
+{
+	return MidiChannelGainRamper.GetCurrent();
+}
+
+void FFusionSampler::SetMidiChannelMute(bool InMute, int8 InMidiChannel)
+{
+	if (InMute == MidiChannelMuted) return;
+
+	FScopeLock Lock(&GetBusLock());
+	MidiChannelMuted = InMute;
+	if (InMute)
+	{
+		MidiChannelMuteGainRamper.SetTarget(0.0);
+	}
+	else
+	{
+		MidiChannelMuteGainRamper.SetTarget(1.0);
+	}
+}
+
+float FFusionSampler::GetTrimVolume() const
+{
+	return TrimVolume;
+}
+
+void FFusionSampler::SetTrimVolume(float InVolume)
+{
+	FScopeLock Lock(&GetBusLock());
+	TrimVolume = InVolume;
+	TrimGain = HarmonixDsp::DBToLinear(TrimVolume);
+}
+
+float FFusionSampler::GetTrimGain() const
+{
+	return TrimGain;
+}
+
+void FFusionSampler::SetStartPointMs(float InStartPointOffsetMs)
+{
+	FScopeLock Lock(&GetBusLock());
+	StartPointMs = InStartPointOffsetMs;
+}
+
+float FFusionSampler::GetStartPointMs() const
+{
+	return StartPointMs;
+}
+
+void FFusionSampler::PrepareToProcess(uint32 InNumSamples)
+{
+	ExpressionGainRamper.Ramp();
+	PortamentoPitchRamper.Ramp();
+	UpdatePitchBendFactor();
+
+	Lfos.Pan.Advance(static_cast<uint32_t>(InNumSamples * Speed));
+	Lfos.Pitch.Advance(static_cast<uint32_t>(InNumSamples * Speed));
+}
+
+void FFusionSampler::Process(uint32 InSliceIdx, uint32 InSubSliceIdx, TAudioBuffer<float>& Output)
+{
+	FScopeLock Lock(&GetBusLock());
+
+	CallPreProcessCallbacks((int32)Output.GetLengthInFrames(), (float)GetSamplesPerSecond(), -1, -1, true);
+
+	// There need to ramp whether or not we are currently playing any voices...
+	// so it must go before any 'early outs' below!
+	MidiChannelGainRamper.Ramp();
+	MidiChannelMuteGainRamper.Ramp();
+	PitchBendRamper.Ramp();
+
+	Output.SetIsSilent(true);
+
+	ProcessNoteActions(Output.GetMaxNumFrames());
+
+	// TODO: currently we never early out if we have children.
+	if (ActiveVoices.Num() == 0)
+	{
+		VoiceUsage.Merge(0.0);
+		Output.FillData(0);
+
+		// don't cut off delay or buffers, even if no voices are running
+		if (DelayEnabled && DelaySilenceCountdown > 0.0f)
+		{
+			Output.SetIsSilent(false);
+			Delay.Process(Output);
+			DelaySilenceCountdown -= GetSecondsPerSample() * (double)Output.GetNumValidFrames();
+		}
+
+		if (BitCrusherEnabled)
+		{
+			BitCrusher.Reset();
+		}
+
+		return;
+	}
+
+
+	// make sure the temporary voice Output buffer is
+	// the same size as the final Output buffer
+	uint32 NumChannels = Output.GetMaxNumChannels();
+	uint32 NumSamplesToProcess = Output.GetMaxNumFrames();
+	checkSlow(NumChannels <= kScratchBufferChannels);
+	checkSlow(NumSamplesToProcess <= kScratchBufferFrames);
+
+	if (!VoicePool)
+	{
+		return;
+	}
+
+	Output.ZeroData();
+
+	// call the internal process call
+	PrepareToProcess(NumSamplesToProcess);
+
+	for (auto Node = ActiveVoices.GetHead(); Node;)
+	{
+		FFusionVoice* Voice = Node->GetValue();
+		check(Voice);
+
+		// increment iterator here...
+		Node = Node->GetNextNode();
+
+		// ... because next call may result in a request back to us to 
+		// relinquish the voice....
+		uint32 NumSamplesProduced = Voice->Process(InSliceIdx, InSubSliceIdx, VoiceWorkBufferChannels, NumChannels, NumSamplesToProcess, Speed, CurrentTempoBPM, MaintainPitchWhenSpeedChanges);
+
+		if (NumSamplesProduced > 0)
+		{
+			// accumulate this voice's samples into the channel Output buffer
+			Output.Accumulate(VoiceWorkBufferChannels, NumChannels, NumSamplesToProcess);
+		}
+	}
+
+	// set up AudioBuffer alias for mixing children
+	TAudioBuffer<float> ChildMixBuffer;
+	ChildMixBuffer.Configure(
+		Output.GetChannelLayout(), 
+		NumSamplesToProcess,
+		EAudioBufferCleanupMode::DontDelete,
+		(float)GetSamplesPerSecond(), 
+		false);
+
+	for (uint32 Ch = 0; Ch < NumChannels; ++Ch)
+	{
+		ChildMixBuffer.SetChannelData(Ch, VoiceWorkBufferChannels[Ch]);
+	}
+
+	// mix in any pre-fx children now
+	for (int32 ChildIdx = 0; ChildIdx < FusionPreChildren.Num(); ++ChildIdx)
+	{
+		if (VocoderEnabled && ChildIdx == VocoderModulatorIndex)
+		{
+			continue;
+		}
+		FusionPreChildren[ChildIdx]->Process(InSliceIdx, InSubSliceIdx, ChildMixBuffer);
+		Output.Accumulate(ChildMixBuffer);
+	}
+
+	if (!IsChildRenderer() && FusionPreChildren.Num() > 0 && VocoderEnabled)
+	{
+		FusionPreChildren[VocoderModulatorIndex]->Process(InSliceIdx, InSubSliceIdx, ChildMixBuffer);
+		Vocoder.Process(Output, ChildMixBuffer);
+	}
+
+	// apply effects
+	if (DistortionEnabled)
+	{
+		Distortion.Process(Output);
+	}
+
+	if (BitCrusherEnabled)
+	{
+		BitCrusher.Process(Output);
+	}
+
+	if (DelayEnabled)
+	{
+		Delay.Process(Output);
+		DelaySilenceCountdown = Delay.CalculateSecsToIdle();
+		Output.SetIsSilent(false);
+	}
+
+	// de-activate portamento if we're in legato mode
+	// and there are no voices playing
+
+	if (PortamentoMode == EPortamentoMode::Legato && ActiveVoices.Num() == 0)
+	{
+		IsPortamentoActive = false;
+	}
+
+	// mix in any post-fx children now
+	for (int32 ChildIdx = 0; ChildIdx < FusionPostChildren.Num(); ++ChildIdx)
+	{
+		FusionPostChildren[ChildIdx]->Process(InSliceIdx, InSubSliceIdx, ChildMixBuffer);
+		Output.Accumulate(ChildMixBuffer);
+	}
+}
+
+void FFusionSampler::GetController(MidiConstants::EControllerID InController, int8& Msb, int8& Lsb, int8 InMidiChannel) const
+{
+	// we ignore InMidiChannel because this is a single channel instrument!
+
+	float Min;
+	float Max;
+	float Range;
+
+
+	float ValueF;
+
+	Lsb = 0; // default case.
+
+	float ValueP;
+	switch (InController)
+	{
+	case MidiConstants::EControllerID::BankSelection:
+	{
+		Msb = (int8)PresetIdx;
+		break;
+	}
+
+	case MidiConstants::EControllerID::Volume:
+	{
+		float Exp = HarmonixDsp::DBToMidiLinear(MidiChannelVolume);
+		Msb = (int8)(Exp * 127.0f);
+		break;
+	}
+
+	case MidiConstants::EControllerID::Expression:
+	{
+		float Exp = HarmonixDsp::LinearToDB(ExpressionGain);
+		Exp = HarmonixDsp::DBToMidiLinear(Exp);
+		Msb = (int8)(Exp * 127.0f);
+		break;
+	}
+	case MidiConstants::EControllerID::PanRight:
+	{
+		Msb = (int8)((PanSettings.Detail.Pan + 1.0f) * 64.0f);
+		break;
+	}
+	case MidiConstants::EControllerID::Release:
+	{
+		Min = 0.005f;
+		Max = 2.000f;
+		Range = Max - Min;
+		ValueP = AdsrSettings.Volume.ReleaseTime;
+		ValueF = ((ValueP - Min) / Range);
+		Msb = (int8)(ValueF * 127.0f);
+		break;
+	}
+
+	case MidiConstants::EControllerID::Attack:
+	{
+		Min = 0.005f;
+		Max = 2.000f;
+		Range = Max - Min;
+		ValueP = AdsrSettings.Volume.AttackTime;
+		ValueF = ((ValueP - Min) / Range);
+		Msb = (int8)(ValueF * 127.0f);
+		break;
+	}
+
+	case MidiConstants::EControllerID::PortamentoSwitch:
+	{
+		Msb = GetIsPortamentoEnabled() ? 127 : 0;
+		break;
+	}
+
+	case MidiConstants::EControllerID::PortamentoTime:
+	{
+		Msb = (int8)(gPortamentoTimeInterp.InverseClamped(PortamentoTimeMs) * 127.0f);
+		break;
+	}
+
+	case MidiConstants::EControllerID::FilterFrequency:
+	{
+		Msb = (int8)(gFreqInterp.InverseClamped(FilterSettings.Freq) * 127.0f);
+		break;
+	}
+
+	case MidiConstants::EControllerID::FilterQ:
+	{
+		Msb = (int8)(gFilterQInterp.InverseClamped(FilterSettings.Q) * 127.0f);
+		break;
+	}
+
+	case MidiConstants::EControllerID::CoarsePitchBend:
+	{
+		Msb = (int8)(GetPitchBend() * 127.0f);
+		break;
+	}
+
+	case MidiConstants::EControllerID::SampleStartTime:
+	{
+		Msb = (int8)(gStartPointTimeInterp.InverseClamped(GetStartPointMs()) * 127.0f);
+		break;
+	}
+
+	case MidiConstants::EControllerID::LFO0Frequency:
+	case MidiConstants::EControllerID::LFO1Frequency:
+	{
+		ELfoIndex LfoIdx = (InController == MidiConstants::EControllerID::LFO0Frequency) ? ELfoIndex::Pan : ELfoIndex::Pitch;
+
+		float Freq = LfoSettings[LfoIdx].Freq;
+		Msb = (int8)(gLfoFreqInterp.InverseClamped(Freq) * 127.0f);
+		break;
+	}
+
+	case MidiConstants::EControllerID::LFO0Depth:
+	case MidiConstants::EControllerID::LFO1Depth:
+	{
+		ELfoIndex LfoIdx = (InController == MidiConstants::EControllerID::LFO0Depth) ? ELfoIndex::Pan : ELfoIndex::Pitch;
+
+		Msb = (int8)(LfoSettings[LfoIdx].Depth * 127.0f);
+		break;
+	}
+
+	case MidiConstants::EControllerID::BitCrushWetMix:
+	{
+		Msb = (int8)(BitCrusher.GetWetGain() * 127.0f);
+		break;
+	}
+
+	case MidiConstants::EControllerID::BitCrushLevel:
+	{
+		Msb = (int8)(((float)BitCrusher.GetCrush() / 15.0f) * 127.0f);
+		break;
+	}
+
+	case MidiConstants::EControllerID::BitCrushSampleHold:
+	{
+		Msb = (int8)((((float)BitCrusher.GetSampleHoldFactor() - 1.0f) / 15.0f) * 127.0f);
+		break;
+	}
+
+	case MidiConstants::EControllerID::DelayTime:
+	{
+		Msb = (int8)(gDelayTimeInterp.InverseClamped(Delay.GetDelaySeconds()) * 127.0f);
+		break;
+	}
+
+	case MidiConstants::EControllerID::DelayDryGain:
+	{
+		Msb = (int8)(Delay.GetDryGain() * 127.0f);
+		break;
+	}
+
+	case MidiConstants::EControllerID::DelayWetGain:
+	{
+		Msb = (int8)(Delay.GetWetGain() * 127.0f);
+		break;
+	}
+
+	case MidiConstants::EControllerID::DelayFeedback:
+	{
+		Msb = (int8)(Delay.GetFeedbackGain() * 127.0f);
+		break;
+	}
+	case MidiConstants::EControllerID::DelayEQEnabled:
+	{
+		Msb = Delay.GetWetFilterEnabled() ? 127 : 0;
+		break;
+	}
+	case MidiConstants::EControllerID::DelayEQType:
+	{
+		Msb = (int8)Delay.GetFilterType();
+		break;
+	}
+	case MidiConstants::EControllerID::DelayEQFreq:
+	{
+		Msb = (int8)(Delay.GetFilterFreq() * 127.0f);
+		break;
+	}
+	case MidiConstants::EControllerID::DelayEQQ:
+	{
+		Msb = (int8)(Delay.GetFilterQ() * 127.0f);
+		break;
+	}
+	case MidiConstants::EControllerID::DelayLFOEnabled:
+	{
+		Msb = Delay.GetLfoEnabled() ? 127 : 0;
+		break;
+	}
+	case MidiConstants::EControllerID::DelayLFOBeatSync:
+	{
+		Msb = Delay.GetLfoTimeSyncOption() == ETimeSyncOption::TempoSync ? 127 : 0;
+		break;
+	}
+	case MidiConstants::EControllerID::DelayLFORate:
+	{
+		Msb = (int8)(Delay.GetLfoFreq() * 127.0f);
+		break;
+	}
+	case MidiConstants::EControllerID::DelayLFODepth:
+	{
+		Msb = (int8)(Delay.GetLfoDepth() * 127.0f);
+		break;
+	}
+	case MidiConstants::EControllerID::DelayStereoType:
+	{
+		Msb = (int8)Delay.GetStereoType();
+		break;
+	}
+	case MidiConstants::EControllerID::DelayPanLeft:
+	{
+		Msb = (int8)(Delay.GetStereoSpreadLeft() * 127.0f);
+		break;
+	}
+	case MidiConstants::EControllerID::DelayPanRight:
+	{
+		Msb = (int8)(Delay.GetStereoSpreadRight() * 127.0f);
+	}
+	default:
+		break;
+	}
+}
+
+void FFusionSampler::SetController(MidiConstants::EControllerID InController, float InValue)
+{
+	FScopeLock Lock(&GetBusLock());
+	// we ignore InMidiChannel because this is a single channel instrument!
+
+	switch (InController)
+	{
+	case MidiConstants::EControllerID::BankSelection: 
+	{ 
+		PresetIdx = (int32)InValue; ApplyPatchPresetSettings(); 
+		break;
+	}
+	case MidiConstants::EControllerID::Volume: 
+	{ 
+		SetMidiChannelVolume(InValue); 
+		break;
+	}
+	case MidiConstants::EControllerID::Expression: 
+	{ 
+		SetMidiExpressionGain(InValue); 
+		break;
+	}
+	case MidiConstants::EControllerID::PanRight: 
+	{ 
+		PanSettings.Detail.Pan = InValue; 
+		break; 
+	}
+	case MidiConstants::EControllerID::Release: 
+	{ 
+		AdsrSettings.Volume.ReleaseTime = InValue; 
+		AdsrSettings.Volume.BuildReleaseTable(); 
+		break; 
+	}
+	case MidiConstants::EControllerID::Attack: 
+	{ 
+		AdsrSettings.Volume.AttackTime = InValue; 
+		AdsrSettings.Volume.BuildAttackTable(); 
+		break; 
+	}
+	case MidiConstants::EControllerID::PortamentoSwitch: 
+	{ 
+		SetIsPortamentoEnabled(InValue > 0.0f); 
+		break; 
+	}
+	case MidiConstants::EControllerID::PortamentoTime: 
+	{ 
+		SetPortamentoTime(gPortamentoTimeInterp.EvalClamped(InValue) / 1000);
+		break; 
+	}
+	case MidiConstants::EControllerID::FilterFrequency: 
+	{ 
+		FilterSettings.Freq = gFreqInterp.EvalClamped(InValue); 
+		break; 
+	}
+	case MidiConstants::EControllerID::FilterQ: 
+	{ 
+		FilterSettings.Q = gFilterQInterp.EvalClamped(InValue); 
+		break; 
+	}
+	case MidiConstants::EControllerID::CoarsePitchBend: 
+	{ 
+		SetPitchBend(InValue); 
+		break; 
+	}
+	case MidiConstants::EControllerID::SampleStartTime: 
+	{ 
+		SetStartPointMs(gStartPointTimeInterp.EvalClamped(InValue));
+		break; 
+	}
+	case MidiConstants::EControllerID::LFO0Frequency:
+	case MidiConstants::EControllerID::LFO1Frequency: 
+	{ 
+		ELfoIndex LfoIdx = (InController == MidiConstants::EControllerID::LFO0Frequency) ? ELfoIndex::Pan : ELfoIndex::Pitch;
+
+		if (LfoSettings[LfoIdx].Freq != InValue)
+		{
+			LfoSettings[LfoIdx].Freq = InValue;
+			Lfos[LfoIdx].UseSettings(&LfoSettings[LfoIdx]);
+			UpdateVoiceLfos();
+		}
+		break; 
+	}
+	case MidiConstants::EControllerID::LFO0Depth:
+	case MidiConstants::EControllerID::LFO1Depth: 
+	{ 
+		ELfoIndex LfoIdx = (InController == MidiConstants::EControllerID::LFO0Depth) ? ELfoIndex::Pan : ELfoIndex::Pitch;
+		if (LfoSettings[LfoIdx].Depth != InValue)
+		{
+			LfoSettings[LfoIdx].Depth = InValue;
+			Lfos[LfoIdx].UseSettings(&LfoSettings[LfoIdx]);
+			UpdateVoiceLfos();
+		}
+		break; 
+	}
+	case MidiConstants::EControllerID::BitCrushWetMix: 
+	{ 
+		BitCrusher.SetWetGain(InValue);
+		break; 
+	}
+	case MidiConstants::EControllerID::BitCrushLevel: 
+	{ 
+		BitCrusher.SetCrush(InValue); 
+		break; 
+	}
+	case MidiConstants::EControllerID::BitCrushSampleHold: 
+	{ 
+		BitCrusher.SetSampleHoldFactor((uint16)InValue); 
+		break; 
+	}
+	case MidiConstants::EControllerID::DelayTime: 
+	{ 
+		Delay.SetDelaySeconds(gDelayTimeInterp.EvalClamped(InValue)); 
+		break; 
+	}
+	case MidiConstants::EControllerID::DelayDryGain: 
+	{ 
+		Delay.SetDryGain(InValue); 
+		break; 
+	}
+	case MidiConstants::EControllerID::DelayWetGain: 
+	{ 
+		Delay.SetWetGain(InValue); 
+		break; 
+	}
+	case MidiConstants::EControllerID::DelayFeedback: 
+	{ 
+		Delay.SetFeedbackGain(InValue); 
+		break; 
+	}
+	case MidiConstants::EControllerID::DelayEQEnabled: 
+	{ 
+		Delay.SetWetFilterEnabled(InValue > 0); 
+		break; 
+	}
+	case MidiConstants::EControllerID::DelayEQType: 
+	{ 
+		Delay.SetFilterType((EBiquadFilterType)(InValue * 2.0f)); 
+		break; 
+	}
+	case MidiConstants::EControllerID::DelayEQFreq: 
+	{ 
+		Delay.SetFilterFreq(20.0f + 17980.0f * InValue); 
+		break; 
+	}
+	case MidiConstants::EControllerID::DelayEQQ: 
+	{ 
+		Delay.SetFilterQ(0.1f + 29.9f * InValue); 
+		break; 
+	}
+	case MidiConstants::EControllerID::DelayLFOEnabled: 
+	{ 
+		Delay.SetLfoEnabled(InValue > 0); 
+		break; 
+	}
+	case MidiConstants::EControllerID::DelayLFOBeatSync: 
+	{ 
+		Delay.SetLfoTimeSyncOption(InValue > 0 ? ETimeSyncOption::TempoSync : ETimeSyncOption::None); 
+		break; 
+	}
+	case MidiConstants::EControllerID::DelayLFORate: 
+	{ 
+		Delay.SetLfoFreq(InValue * 5.0f); 
+		break; 
+	}
+	case MidiConstants::EControllerID::DelayLFODepth: 
+	{ 
+		Delay.SetLfoDepth(InValue * 40.0f);
+		break; 
+	}
+	case MidiConstants::EControllerID::DelayStereoType: 
+	{ 
+		Delay.SetStereoType(EDelayStereoType(InValue * 4.0f));
+		break; 
+	}
+	case MidiConstants::EControllerID::DelayPanLeft: 
+	{ 
+		Delay.SetStereoSpreadLeft(InValue * 100.0f); 
+		break; 
+	}
+	case MidiConstants::EControllerID::DelayPanRight: 
+	{
+		Delay.SetStereoSpreadRight(InValue * 100.0f); 
+		break; 
+	}
+	case MidiConstants::EControllerID::TimeStretchEnvelopeOrder:
+	{
+		TimeStretchEnvelopeOverride = (int16)InValue;
+		UpdateVoicesForEnvelopeOrderChange();
+		break;
+	}
+	case MidiConstants::EControllerID::SubStreamVol1:
+	{ 
+		SetSubstreamMidiGain(0, (int32)InValue); 
+		break; 
+	}
+	case MidiConstants::EControllerID::SubStreamVol2: 
+	{ 
+		SetSubstreamMidiGain(1, (int32)InValue); 
+		break; 
+	}
+	case MidiConstants::EControllerID::SubStreamVol3: 
+	{
+		SetSubstreamMidiGain(2, (int32)InValue); 
+		break;
+	}
+	case MidiConstants::EControllerID::SubStreamVol4: 
+	{
+		SetSubstreamMidiGain(3, (int32)InValue);
+		break; 
+	}
+	case MidiConstants::EControllerID::SubStreamVol5: 
+	{ 
+		SetSubstreamMidiGain(4, (int32)InValue); 
+		break;
+	}
+	case MidiConstants::EControllerID::SubStreamVol6:
+	{ 
+		SetSubstreamMidiGain(5, (int32)InValue); 
+		break; 
+	}
+	case MidiConstants::EControllerID::SubStreamVol7: 
+	{ 
+		SetSubstreamMidiGain(6, (int32)InValue); 
+		break;
+	}
+	case MidiConstants::EControllerID::SubStreamVol8:
+	{
+		SetSubstreamMidiGain(7, (int32)InValue); 
+		break; 
+	}
+	default: 
+		break;
+	}
+}
+
+void FFusionSampler::Set7BitControllerImpl(MidiConstants::EControllerID InController, int8 InValue, int8 InMidiChannel)
+{
+	FScopeLock Lock(&GetBusLock());
+	// we ignore InMidiChannel because this is a single channel instrument!
+
+	float Min;
+	float Max;
+	float Range;
+
+	float ValueP = 0.0f;
+
+	switch (InController)
+	{
+	case MidiConstants::EControllerID::BankSelection: 
+	{ 
+		ValueP = (float)InValue; 
+		break; 
+	}
+	case MidiConstants::EControllerID::Volume: 
+	{ 
+		ValueP = HarmonixDsp::MidiLinearToDB(InValue); 
+		break; 
+	}
+	case MidiConstants::EControllerID::Expression: 
+	{ 
+		ValueP = HarmonixDsp::DBToLinear(HarmonixDsp::MidiLinearToDB(InValue)); 
+		break; 
+	}
+	case MidiConstants::EControllerID::PanRight: 
+	{ 
+		ValueP = ((float)InValue / 64.0f) - 1.0f; 
+		break; 
+	}
+	case MidiConstants::EControllerID::Release: 
+	{ 
+		Min = 0.005f; Max = 2.000f; Range = Max - Min;
+		ValueP = (float)InValue / 127.0f;
+		ValueP = ValueP * Range + Min; 
+		break; 
+	}
+	case MidiConstants::EControllerID::Attack: 
+	{ 
+		Min = 0.005f; Max = 2.000f; Range = Max - Min;
+		ValueP = (float)InValue / 127.0f;
+		ValueP = ValueP * Range + Min; 
+		break; 
+	}
+	case MidiConstants::EControllerID::PortamentoSwitch: 
+	{ 
+		ValueP = (InValue > 0) ? 1.0f : 0.0f; 
+		break; 
+	}
+	case MidiConstants::EControllerID::PortamentoTime:
+	case MidiConstants::EControllerID::FilterFrequency:
+	case MidiConstants::EControllerID::FilterQ:
+	case MidiConstants::EControllerID::CoarsePitchBend:
+	case MidiConstants::EControllerID::SampleStartTime: 
+	{ 
+		ValueP = (float)InValue / 127.0f; 
+		break; 
+	}
+	case MidiConstants::EControllerID::LFO0Frequency:
+	case MidiConstants::EControllerID::LFO1Frequency: 
+	{ 
+		ValueP = gLfoFreqInterp.EvalClamped((float)InValue / 127.0f); 
+		break; 
+	}
+	case MidiConstants::EControllerID::LFO0Depth:
+	case MidiConstants::EControllerID::LFO1Depth: 
+	{ 
+		ValueP = (float)InValue / 127.0f; 
+		break; 
+	}
+	case MidiConstants::EControllerID::BitCrushWetMix:
+	{ 
+		ValueP = (float)InValue * 100.0f / 127.0f; 
+		break; 
+	}
+	case MidiConstants::EControllerID::BitCrushLevel: 
+	{ 
+		ValueP = (float)InValue * 15.0f / 127.0f; 
+		break; 
+	}
+	case MidiConstants::EControllerID::BitCrushSampleHold: 
+	{ 
+		ValueP = (1.5f + (((float)InValue / 127.0f) * 15));
+		break; 
+	}
+	case MidiConstants::EControllerID::DelayTime:
+	case MidiConstants::EControllerID::DelayDryGain:
+	case MidiConstants::EControllerID::DelayWetGain:
+	case MidiConstants::EControllerID::DelayEQEnabled:
+	case MidiConstants::EControllerID::DelayEQType:
+	case MidiConstants::EControllerID::DelayEQFreq:
+	case MidiConstants::EControllerID::DelayEQQ:
+	case MidiConstants::EControllerID::DelayLFOEnabled:
+	case MidiConstants::EControllerID::DelayLFOBeatSync:
+	case MidiConstants::EControllerID::DelayLFORate:
+	case MidiConstants::EControllerID::DelayLFODepth:
+	case MidiConstants::EControllerID::DelayStereoType:
+	case MidiConstants::EControllerID::DelayPanLeft:
+	case MidiConstants::EControllerID::DelayPanRight:
+	case MidiConstants::EControllerID::DelayFeedback: 
+	{ 
+		ValueP = (float)InValue / 127.0f; 
+		break; 
+	}
+	case MidiConstants::EControllerID::TimeStretchEnvelopeOrder: 
+	{ 
+		ValueP = (float)InValue; 
+		break; 
+	}
+	case MidiConstants::EControllerID::SubStreamVol1:
+	case MidiConstants::EControllerID::SubStreamVol2:
+	case MidiConstants::EControllerID::SubStreamVol3:
+	case MidiConstants::EControllerID::SubStreamVol4:
+	case MidiConstants::EControllerID::SubStreamVol5:
+	case MidiConstants::EControllerID::SubStreamVol6:
+	case MidiConstants::EControllerID::SubStreamVol7:
+	case MidiConstants::EControllerID::SubStreamVol8: 
+	{ 
+		ValueP = (float)InValue; 
+		break; 
+	}
+	default: break;
+	}
+	SetController(InController, ValueP);
+}
+
+void FFusionSampler::Set14BitControllerImpl(MidiConstants::EControllerID InController, int16 InValue, int8 InMidiChannel)
+{
+	FScopeLock Lock(&GetBusLock());
+	// we ignore InMidiChannel because this is a single channel instrument!
+
+	float Min;
+	float Max;
+	float Range;
+
+	float ValueP;
+
+	switch (InController)
+	{
+	case MidiConstants::EControllerID::BankSelection: 
+	{
+		ValueP = (float)InValue; 
+		break; 
+	}
+	case MidiConstants::EControllerID::Volume: 
+	{ 
+		ValueP = HarmonixDsp::Midi14BitLinearToDB(InValue); 
+		break; 
+	}
+	case MidiConstants::EControllerID::Expression: 
+	{ 
+		ValueP = HarmonixDsp::DBToLinear(HarmonixDsp::Midi14BitLinearToDB(InValue));
+		break; 
+	}
+	case MidiConstants::EControllerID::PanRight: 
+	{
+		ValueP = ((float)InValue / 8191.0f) - 1.0f;
+		break;
+	}
+	case MidiConstants::EControllerID::Release: 
+	{ 
+		Min = 0.005f; Max = 2.000f; 
+		Range = Max - Min;
+		ValueP = (float)InValue / 16383.0f;
+		ValueP = ValueP * Range + Min; 
+		break; 
+	}
+	case MidiConstants::EControllerID::Attack: 
+	{ 
+		Min = 0.005f; Max = 2.000f; 
+		Range = Max - Min;
+		ValueP = (float)InValue / 16383.0f;
+		ValueP = ValueP * Range + Min; 
+		break; 
+	}
+	case MidiConstants::EControllerID::PortamentoSwitch: 
+	{ 
+		ValueP = (InValue > 0) ? 1.0f : 0.0f; 
+		break; 
+	}
+	case MidiConstants::EControllerID::PortamentoTime:
+	case MidiConstants::EControllerID::FilterFrequency:
+	case MidiConstants::EControllerID::FilterQ:
+	case MidiConstants::EControllerID::CoarsePitchBend:
+	case MidiConstants::EControllerID::SampleStartTime: 
+	{ 
+		ValueP = (float)InValue / 16383.0f; 
+		break; 
+	}
+	case MidiConstants::EControllerID::LFO0Frequency:
+	case MidiConstants::EControllerID::LFO1Frequency: 
+	{ 
+		ValueP = gLfoFreqInterp.EvalClamped((float)InValue / 16383.0f); 
+		break; 
+	}
+	case MidiConstants::EControllerID::LFO0Depth:
+	case MidiConstants::EControllerID::LFO1Depth: 
+	{ 
+		ValueP = (float)InValue / 16383.0f;
+		break;
+	}
+	case MidiConstants::EControllerID::BitCrushWetMix: 
+	{
+		ValueP = (float)InValue * 100.0f / 16383.0f;
+		break; 
+	}
+	case MidiConstants::EControllerID::BitCrushLevel:
+	{ 
+		ValueP = (float)InValue * 15.0f / 16383.0f; 
+		break;
+	}
+	case MidiConstants::EControllerID::BitCrushSampleHold: 
+	{ 
+		ValueP = (1.5f + (((float)InValue / 16383.0f) * 15));
+		break; 
+	}
+	case MidiConstants::EControllerID::DelayTime:
+	case MidiConstants::EControllerID::DelayDryGain:
+	case MidiConstants::EControllerID::DelayWetGain:
+	case MidiConstants::EControllerID::DelayEQEnabled:
+	case MidiConstants::EControllerID::DelayEQType:
+	case MidiConstants::EControllerID::DelayEQFreq:
+	case MidiConstants::EControllerID::DelayEQQ:
+	case MidiConstants::EControllerID::DelayLFOEnabled:
+	case MidiConstants::EControllerID::DelayLFOBeatSync:
+	case MidiConstants::EControllerID::DelayLFORate:
+	case MidiConstants::EControllerID::DelayLFODepth:
+	case MidiConstants::EControllerID::DelayStereoType:
+	case MidiConstants::EControllerID::DelayPanLeft:
+	case MidiConstants::EControllerID::DelayPanRight:
+	case MidiConstants::EControllerID::DelayFeedback:
+	{ 
+		ValueP = (float)InValue / 16383.0f;
+		break; 
+	}
+	case MidiConstants::EControllerID::TimeStretchEnvelopeOrder: 
+	{ 
+		ValueP = (float)(InValue >> 7); 
+		break; 
+	}
+	case MidiConstants::EControllerID::SubStreamVol1:
+	case MidiConstants::EControllerID::SubStreamVol2:
+	case MidiConstants::EControllerID::SubStreamVol3:
+	case MidiConstants::EControllerID::SubStreamVol4:
+	case MidiConstants::EControllerID::SubStreamVol5:
+	case MidiConstants::EControllerID::SubStreamVol6:
+	case MidiConstants::EControllerID::SubStreamVol7:
+	case MidiConstants::EControllerID::SubStreamVol8: 
+	{ 
+		ValueP = (float)(InValue >> 7); 
+		break; 
+	}
+	default:
+		ValueP = 0.0f; 
+		break;
+	}
+	SetController(InController, ValueP);
+}
+
+void FFusionSampler::UpdateVoiceLfos()
+{
+	FScopeLock Lock(&GetBusLock());
+
+	for (auto Node = ActiveVoices.GetHead(); Node; Node = Node->GetNextNode())
+	{
+		// skip voices not assigned to this channel
+		FFusionVoice* Voice = Node->GetValue();
+		check(Voice);
+		Voice->SetupLfo(0, LfoSettings.Pan);
+		Voice->SetupLfo(1, LfoSettings.Pitch);
+	}
+}
+
+void FFusionSampler::UpdateVoicesForEnvelopeOrderChange()
+{
+	FScopeLock Lock(&GetBusLock());
+	for (auto Node = ActiveVoices.GetHead(); Node; Node = Node->GetNextNode())
+	{
+		FFusionVoice* Voice = Node->GetValue();
+		check(Voice);
+
+		if (TSharedPtr<IStretcherAndPitchShifter, ESPMode::ThreadSafe> Shifter = Voice->GetPitchShifter())
+		{
+			TMap<FName, FTypedParameter> Options;
+			Options.Add("EnvelopeOrder", TimeStretchEnvelopeOverride);
+			Shifter->ApplyOptions(Options);
+		}
+	}
+}
+
+void FFusionSampler::SetDelayEnabled(bool InEnabled)
+{
+	FScopeLock Lock(&GetBusLock());
+	DelayEnabled = InEnabled;
+
+	if (InEnabled)
+	{
+		Delay.Prepare((float)GetSamplesPerSecond(), kScratchBufferChannels, 2500.0f);
+	}
+	else
+	{
+		Delay.Unprepare();
+	}
+}
+
+void FFusionSampler::SetBitCrusherEnabled(bool InEnabled)
+{
+	if (BitCrusherEnabled != InEnabled && InEnabled)
+	{
+		BitCrusher.Reset();
+	}
+	BitCrusherEnabled = InEnabled;
+}
+
+// TODO: Needs to be reimplemented without FHarmonixGeneratorHandle
+//bool FFusionSampler::AddChild(const FHarmonixGeneratorHandle& child, EInstrumentRenderMode InRenderMode)
+//{
+//	FInstrumentGeneratorLock ihl(child);
+//	if (!ihl)
+//		return false;
+//
+//	_AddChild(ihl.GetRaw(), InRenderMode);
+//
+//	FusionChildHandles.Add(child);
+//
+//	return true;
+//}
+//
+//bool FFusionSampler::RemoveChild(const FHarmonixGeneratorHandle& child)
+//{
+//	FInstrumentGeneratorLock ihl(child);
+//	if (!ihl)
+//		return false;
+//
+//	FScopeLock Lock(&GetBusLock());
+//
+//	bool found = _RemoveChild(ihl.GetRaw());
+//
+//	int32 ChildIdx = FusionChildHandles.Find(child);
+//	if (ChildIdx != INDEX_NONE)
+//	{
+//		FusionChildHandles.RemoveAt(ChildIdx);
+//		found = true;
+//	}
+//	return found;
+//}
+
+void FFusionSampler::AddChild(FVirtualInstrument* InChildInstrument, EInstrumentRenderMode InRenderMode)
+{
+	FScopeLock Lock(&GetBusLock());
+
+	if (InRenderMode == EInstrumentRenderMode::PreFxChild)
+	{
+		FusionPreChildren.Add(InChildInstrument);
+	}
+	else
+	{
+		FusionPostChildren.Add(InChildInstrument);
+	}
+}
+
+bool FFusionSampler::RemoveChild(FVirtualInstrument* InChildInstrument)
+{
+	FScopeLock Lock(&GetBusLock());
+
+	bool OutRemoved = false;
+
+	OutRemoved |= FusionPreChildren.Remove(InChildInstrument) > 0; 
+	OutRemoved |= FusionPostChildren.Remove(InChildInstrument) > 0;
+
+	return OutRemoved;
+}
+
+void FFusionSampler::DumpAllChildren()
+{
+	FScopeLock Lock(&BusLock);
+	FusionPreChildren.Reset();
+	FusionPostChildren.Reset();
+}
+
+FString FFusionSampler::GetPatchPath() const
+{
+	unimplemented();
+	return FString(); // FusionPatch->GetPathName();
+}
+
+#ifdef UE_BUILD_DEVELOPMENT
+TSet<int32> FFusionSampler::GetActiveKeyzones()
+{
+	FScopeLock Lock(&GetBusLock());
+
+	check(FusionPatchData);
+
+	const TArray<FKeyzoneSettings>& KeyzonesRef = FusionPatchData->GetKeyzones();
+
+	TSet<int32> OutIndices;
+	if (!VoicePool)
+	{
+		return OutIndices;
+	}
+
+	for (FFusionVoice* Voice : ActiveVoices)
+	{
+		check(Voice);
+
+		for (int32 Idx = 0; Idx < KeyzonesRef.Num(); ++Idx)
+		{
+			if (Voice->UsesKeyzone(&KeyzonesRef[Idx]))
+			{
+				OutIndices.Add(Idx);
+			}
+		}
+	}
+
+	return OutIndices;
+}
+#endif
+
+void FFusionSampler::SetSubstreamMidiGain(int32 Idx, uint8 InMidiGain)
+{
+	if (Idx >= kMaxSubstreams)
+	{
+		UE_LOG(LogFusionSampler, Error, TEXT("Substream index out of range!"));
+		return;
+	}
+
+	float VolumeDb = HarmonixDsp::LinearToDB(InMidiGain);
+	SubstreamGain[Idx] = HarmonixDsp::DBToLinear(VolumeDb);
+}
\ No newline at end of file
diff --git a/Engine/Plugins/Runtime/Harmonix/Source/HarmonixDsp/Private/HarmonixDsp/FusionSampler/Settings/KeyzoneSettings.cpp b/Engine/Plugins/Runtime/Harmonix/Source/HarmonixDsp/Private/HarmonixDsp/FusionSampler/Settings/KeyzoneSettings.cpp
new file mode 100644
--- /dev/null
+++ b/Engine/Plugins/Runtime/Harmonix/Source/HarmonixDsp/Private/HarmonixDsp/FusionSampler/Settings/KeyzoneSettings.cpp
@@ -0,0 +1,96 @@
+// Copyright Epic Games, Inc. All Rights Reserved.
+#include "HarmonixDsp/FusionSampler/Settings/KeyzoneSettings.h"
+#include "HarmonixDsp/FusionSampler/SingletonFusionVoicePool.h"
+#include "HarmonixDsp/FusionSampler/FusionPatch.h"
+#include "HarmonixDsp/AudioUtility.h"
+#include "HarmonixDsp/AudioData.h"
+#include "HarmonixDsp/AudioData/StreamingAudioData.h"
+
+#include "Sound/SoundWave.h"
+
+#include "HAL/PlatformTime.h"
+
+DEFINE_LOG_CATEGORY(LogKeyzoneSettings);
+
+void FKeyzoneSettings::InitProxyData(const Audio::FProxyDataInitParams& InitParams)
+{
+	QUICK_SCOPE_CYCLE_COUNTER(STAT_FKeyzoneSettings_InitProxyData);
+
+	UE_LOG(LogKeyzoneSettings, Verbose, TEXT("Initializing new KeyzoneSettings proxy"));
+	
+	// we should only be initializing the proxy data if the sound wave is valid
+	if (!ensure(SoundWave))
+	{
+		return;
+	}
+
+	AudioSample = CreateStreamingAudioData(InitParams);
+
+	check(AudioSample);
+
+	// nulling out sound wave asset for this Keyzone
+	// now that we are turning into ProxyData
+	SoundWave = nullptr;
+
+	if (UseSingletonVoicePool && !SingletonFusionVoicePool)
+	{
+		SingletonFusionVoicePool = MakeShared<FSingletonFusionVoicePool>(FSingletonFusionVoicePool::kMaxSingletonAliases, *this);
+	}
+	else if (SingletonFusionVoicePool)
+	{
+		SingletonFusionVoicePool.Reset();
+	}
+}
+
+TSharedPtr<FStreamingAudioData, ESPMode::ThreadSafe> FKeyzoneSettings::CreateStreamingAudioData(const Audio::FProxyDataInitParams& InitParams) const
+{
+	TSharedPtr<FStreamingAudioData, ESPMode::ThreadSafe> StreamingAudioSample = StaticCastSharedPtr<FStreamingAudioData>(AudioSample);
+	if (StreamingAudioSample.IsValid())
+	{
+		return StreamingAudioSample;
+	}
+
+	if (!ensure(SoundWave))
+	{
+		UE_LOG(LogKeyzoneSettings, Error, TEXT("SoundWave asset is null in keyzone settings when making proxy data!!"));
+		return nullptr;
+	}
+
+	TSharedPtr<Audio::IProxyData, ESPMode::ThreadSafe> AudioProxy = SoundWave->CreateProxyData(InitParams);
+
+	FSoundWaveProxyPtr WaveProxy = StaticCastSharedPtr<FSoundWaveProxy>(AudioProxy);
+	if (!ensure(WaveProxy))
+	{
+		UE_LOG(LogKeyzoneSettings, Error, TEXT("Failed to make sound wave proxy data for SoundWave asset!!"));
+		return nullptr;
+	}
+
+	UE_LOG(LogKeyzoneSettings, Verbose, TEXT("Creating new shared AudioData"));
+
+	FStreamingAudioData::FSettings StreamingAudioSettings;
+	StreamingAudioSample = FStreamingAudioData::GetShared<FStreamingAudioData>(WaveProxy.ToSharedRef(), StreamingAudioSettings);
+	return StreamingAudioSample;
+}
+
+bool FKeyzoneSettings::ContainsNoteAndVelocity(uint8 InNote, uint8 InVelocity) const
+{
+	return (InNote >= MinNote && InNote <= MaxNote) 
+		&& (InVelocity >= MinVelocity && InVelocity <= MaxVelocity);
+}
+
+void FKeyzoneSettings::SetVolumeDb(float Db)
+{ 
+	Db = HarmonixDsp::ClampDB(Db);
+	Gain = HarmonixDsp::DBToLinear(Db);
+}
+
+float FKeyzoneSettings::GetVolumeDb() const
+{
+	return HarmonixDsp::LinearToDB(FMath::Clamp(Gain, 0.0f, 1.0f));
+}
+
+void FKeyzoneSettings::SetFineTuneCents(float InCents)
+{
+	FineTuneCents = InCents;
+	FineTuneAdjustment = FMath::Pow(2.0f, FineTuneCents / 1200.0f);
+}
\ No newline at end of file
diff --git a/Engine/Plugins/Runtime/Harmonix/Source/HarmonixDsp/Private/HarmonixDsp/Instruments/VirtualInstrument.cpp b/Engine/Plugins/Runtime/Harmonix/Source/HarmonixDsp/Private/HarmonixDsp/Instruments/VirtualInstrument.cpp
new file mode 100644
--- /dev/null
+++ b/Engine/Plugins/Runtime/Harmonix/Source/HarmonixDsp/Private/HarmonixDsp/Instruments/VirtualInstrument.cpp
@@ -0,0 +1,170 @@
+// Copyright Epic Games, Inc. All Rights Reserved.
+#include "HarmonixDsp/Instruments/VirtualInstrument.h"
+
+#include "HarmonixMidi/MidiConstants.h"
+#include "HarmonixMidi/MidiMsg.h"
+
+bool FVirtualInstrument::bShouldPrintMidiActivity = false;
+
+void FVirtualInstrument::ResetInstrumentState()
+{
+	FMemory::Memset(LastCcVals, -1, 16 * 128);
+	ResetInstrumentStateImpl();
+}
+
+void FVirtualInstrument::ResetMidiState()
+{
+	FMemory::Memset(LastCcVals, -1, 16 * 128);
+	ResetMidiStateImpl();
+}
+
+void FVirtualInstrument::Set7BitController(MidiConstants::EControllerID InController, int8 InByteValue, int8 InMidiChannel)
+{
+	LastCcVals[InMidiChannel][(uint8)InController] = InByteValue;
+	Set7BitControllerImpl(InController, InByteValue, InMidiChannel);
+}
+
+void FVirtualInstrument::Set14BitController(MidiConstants::EControllerID InController, int16 InByteValue, int8 InMidiChannel)
+{
+	int32 LsbIdx;
+	int32 MsbIdx;
+	if (GetMsbLsbIndexes(InController, MsbIdx, LsbIdx))
+	{
+		LastCcVals[InMidiChannel][MsbIdx] = (InByteValue >> 7) & 0x7F;
+		LastCcVals[InMidiChannel][LsbIdx] = InByteValue & 0x7F;
+		Set14BitControllerImpl((MidiConstants::EControllerID)MsbIdx, InByteValue, InMidiChannel);
+	}
+}
+
+void FVirtualInstrument::HandleMidiMessage(FMidiVoiceId InVoiceId, int8 InStatus, int8 InData1, int8 InData2, int32 InEventTick, int32 InCurrentTick, float InMsOffset)
+{
+	int8 InChannel = InStatus & 0xF;
+	switch (InStatus & 0xF0)
+	{
+	case MidiConstants::kNoteOff:
+		NoteOff(InVoiceId, InData1, InChannel);
+		break;
+	case MidiConstants::kNoteOn:
+		NoteOn(InVoiceId, InData1, InData2, InChannel, InEventTick, InCurrentTick, InMsOffset);
+		break;
+	case MidiConstants::kPolyPres:
+		PolyPressure(InVoiceId, InData1, InData2, InChannel);
+		break;
+	case MidiConstants::kChanPres:
+		ChannelPressure(InData1, InData2, InChannel);
+		break;
+	case MidiConstants::kControl:
+		SetHighOrLowControllerByte((MidiConstants::EControllerID)InData1, InData2, InChannel);
+		break;
+	case MidiConstants::kPitch:
+		SetPitchBend(FMidiMsg::GetPitchBendFromData(InData1, InData2), InChannel);
+		break;
+	}
+}
+
+bool FVirtualInstrument::IsHighResController(MidiConstants::EControllerID InControllerId, bool& bIsHighResLowByte)
+{
+	uint8 Id = (uint8)InControllerId;
+	if (Id < 32)
+	{
+		bIsHighResLowByte = false;
+		return true;
+	}
+	if (Id < 64)
+	{
+		bIsHighResLowByte = true;
+		return true;
+	}
+	if (Id == 98 || Id == 100)
+	{
+		bIsHighResLowByte = true;
+		return true;
+	}
+	if (Id == 99 || Id == 101)
+	{
+		bIsHighResLowByte = false;
+		return true;
+	}
+	return false;
+}
+
+bool FVirtualInstrument::GetMsbLsbIndexes(MidiConstants::EControllerID InControllerId, int& InMsb, int& InLsb)
+{
+	uint8 Id = (uint8)InControllerId;
+	if (Id < 32)
+	{
+		InMsb = Id; InLsb = Id + 32;
+		return true;
+	}
+	if (Id < 64)
+	{
+		InMsb = Id - 32; InLsb = Id;
+		return true;
+	}
+	if (Id == 98 || Id == 100)
+	{
+		InMsb = Id + 1; InLsb = Id;
+		return true;
+	}
+	if (Id == 99 || Id == 101)
+	{
+		InMsb = Id; InLsb = Id - 1;
+		return true;
+	}
+	return false;
+}
+
+// convert a value in [0,127] to the range expected by SetController
+float FVirtualInstrument::ConvertCCValue(MidiConstants::EControllerID InController, uint8 InValue)
+{
+	check(InValue <= 127);
+	switch (InController)
+	{
+	case MidiConstants::EControllerID::PanRight:
+	case MidiConstants::EControllerID::CoarsePitchBend:
+		// map to [-1, 1], making sure 64 maps exactly to 0
+		return FMath::Max(-1.0f, ((float)InValue - 64) / 63);
+	case MidiConstants::EControllerID::PortamentoSwitch:
+		// 63 and below are "off", 64 and above are "on"
+		return InValue >= 64 ? 1.0f : 0.0f;
+	default:
+		// map to [0, 1]
+		return (float)InValue / 127;
+	}
+}
+
+void FVirtualInstrument::SetController(MidiConstants::EControllerID InController, float InValue, int8 InMidiChannel /*= 0*/)
+{
+	int32 MsbIndex;
+	int32 LsbIndex;
+	if (GetMsbLsbIndexes(InController, MsbIndex, LsbIndex))
+	{
+		Set14BitController(InController, (int16)(InValue * 16383.0f), InMidiChannel);
+	}
+	else
+	{
+		Set7BitController(InController, (int8)(InValue * 127.0f), InMidiChannel);
+	}
+}
+
+void FVirtualInstrument::SetHighOrLowControllerByte(MidiConstants::EControllerID InController, int8 InValue, int8 InMidiChannel)
+{
+	LastCcVals[InMidiChannel][(uint8)InController] = InValue;
+	int32 HighByteIdx;
+	int32 LowByteIdx;
+	if (GetMsbLsbIndexes(InController, HighByteIdx, LowByteIdx))
+	{
+		if (LastCcVals[InMidiChannel][LowByteIdx] == -1)
+		{
+			Set7BitController(InController, InValue, InMidiChannel);
+		}
+		else if (LastCcVals[InMidiChannel][HighByteIdx] != -1)
+		{
+			Set14BitController((MidiConstants::EControllerID)HighByteIdx, (LastCcVals[InMidiChannel][HighByteIdx] << 7 | LastCcVals[InMidiChannel][LowByteIdx]), InMidiChannel);
+		}
+	}
+	else
+	{
+		Set7BitController(InController, InValue, InMidiChannel);
+	}
+}
