[Harmonix] Move Harmonix plugin into Engine/Plugins/Runtime to make it available to licensees in UE 5.4

diff --git a/Engine/Plugins/Runtime/Harmonix/Source/HarmonixDsp/Public/HarmonixDsp/AudioAnalysis/DynamicRangeMeterAnalyzer.h b/Engine/Plugins/Runtime/Harmonix/Source/HarmonixDsp/Public/HarmonixDsp/AudioAnalysis/DynamicRangeMeterAnalyzer.h
new file mode 100644
--- /dev/null
+++ b/Engine/Plugins/Runtime/Harmonix/Source/HarmonixDsp/Public/HarmonixDsp/AudioAnalysis/DynamicRangeMeterAnalyzer.h
@@ -0,0 +1,78 @@
+// Copyright Epic Games, Inc. All Rights Reserved.
+
+#pragma once
+
+#include "HarmonixDsp/AudioBuffer.h"
+#include "HarmonixDsp/Effects/BiquadFilter.h"
+
+namespace Harmonix::Dsp::AudioAnalysis
+{
+	class HARMONIXDSP_API FDynamicRangeMeterAnalyzer
+	{
+	public:
+		struct FSettings
+		{
+			float SampleRate = 48000.0f;
+			float HighRisingAlpha = 1.0f;
+			float HighFallingAlpha = 0.4f;
+			float LowRisingAlpha = 0.001f;
+			float LowFallingAlpha = 1.0f;
+			int32 WindowSize = 480;
+			bool ResetRequested = false;
+			bool FilterResult = false;
+			float FilterCutoff = 30.0f;
+		};
+	
+		struct FResults
+		{
+			float HighEnvelope = 0.0f;
+			float LowEnvelope = 0.0f;
+			float LevelDecibels = -96.0f;
+			float MonoPeakDecibels = -96.0f;
+			float MonoPeakHighEnvelopeDecibels = -96.0f;
+			bool RequestReset = false;
+		};
+
+		FDynamicRangeMeterAnalyzer();
+
+		/**
+		 * @brief Set the settings for this analyzer
+		 * @param InSettings - The settings to use
+		 */
+		void SetSettings(const FSettings& InSettings);
+
+		/**
+		 * @brief Reset the analyzer to its initial state
+		 */
+		void Reset();
+
+		/**
+		 * @brief Process a buffer and get the analyzer results
+		 * @param InBuffer - The audio buffer to analyze
+		 * @param InOutResults - The results from the analysis
+		 */
+		void Process(const TAudioBuffer<float>& InBuffer, FResults& InOutResults);
+		
+	private:
+		FCriticalSection SettingsGuard;
+		FSettings Settings;
+
+		struct FState
+		{
+			float EnvLast = 0.0f;
+			float HighLast = 0.0f;
+			float LowLast = 1.0f;
+			float RmsSum = 0.0f;
+			bool HasHadFirstSettle = false;
+			TArray<TDynamicStridePtr<float>> ChannelPointers;
+			Effects::TMultipassBiquadFilter<double, 3> Filter;
+			Effects::FBiquadFilterCoefs FilterCoefs;
+			bool FilterEnabled = false;
+			float FilterCutoff = 1000.0f;
+			TArray<float> RmsWindow;
+			int32 RmsNextWrite = 0;
+		};
+
+		FState State;
+	};
+}
diff --git a/Engine/Plugins/Runtime/Harmonix/Source/HarmonixDsp/Public/HarmonixDsp/AudioAnalysis/VuMeterAnalyzer.h b/Engine/Plugins/Runtime/Harmonix/Source/HarmonixDsp/Public/HarmonixDsp/AudioAnalysis/VuMeterAnalyzer.h
new file mode 100644
--- /dev/null
+++ b/Engine/Plugins/Runtime/Harmonix/Source/HarmonixDsp/Public/HarmonixDsp/AudioAnalysis/VuMeterAnalyzer.h
@@ -0,0 +1,101 @@
+// Copyright Epic Games, Inc. All Rights Reserved.
+
+#pragma once
+
+#include "HarmonixDsp/AudioBuffer.h"
+#include "HarmonixDsp/AudioAnalysis/OutputSettings.h"
+
+#include "VuMeterAnalyzer.generated.h"
+
+USTRUCT(BlueprintType)
+struct FHarmonixVuMeterAnalyzerSettings
+{
+	GENERATED_BODY()
+
+	UPROPERTY(EditAnywhere, Category="Settings")
+	float AvgWindowMs = 1.0f;
+
+	UPROPERTY(EditAnywhere, Category="Settings")
+	float PeakHoldMs = 1000.0f;
+
+	UPROPERTY(EditAnywhere, Category="Settings")
+	FHarmonixAudioAnalyzerOutputSettings OutputSettings{};
+};
+
+USTRUCT(BlueprintType)
+struct FHarmonixVuMeterAnalyzerChannelValues
+{
+	GENERATED_BODY()
+
+	UPROPERTY(BlueprintReadOnly, Category="Results")
+	float LevelMeanSquared = 0.0f;
+
+	UPROPERTY(BlueprintReadOnly, Category="Results")
+	float PeakSquared = 0.0f;
+};
+
+USTRUCT(BlueprintType)
+struct FHarmonixVuMeterAnalyzerResults
+{
+	GENERATED_BODY()
+
+	UPROPERTY(BlueprintReadOnly, Category="Results")
+	float MonoPeakDecibels = -96.0f;
+
+	UPROPERTY(BlueprintReadOnly, Category="Results")
+	TArray<FHarmonixVuMeterAnalyzerChannelValues> ChannelValues;
+
+	UPROPERTY(BlueprintReadOnly, Category="Results")
+	FHarmonixVuMeterAnalyzerChannelValues MonoValues{};
+};
+
+namespace Harmonix::Dsp::AudioAnalysis
+{
+	class HARMONIXDSP_API FVuMeterAnalyzer
+	{
+	public:
+		explicit FVuMeterAnalyzer(float InSampleRate);
+
+		/**
+		 * @brief Set the settings for this analyzer
+		 * @param InSettings - The settings to use
+		 */
+		void SetSettings(const FHarmonixVuMeterAnalyzerSettings& InSettings);
+
+		/**
+		 * @brief Reset the analyzer to its initial state
+		 */
+		void Reset();
+
+		/**
+		 * @brief Process a buffer and get the analyzer results
+		 * @param InBuffer - The audio buffer to analyze
+		 * @param InOutResults - The results from the analysis
+		 */
+		void Process(const TAudioBuffer<float>& InBuffer, FHarmonixVuMeterAnalyzerResults& InOutResults);
+		
+	private:
+		struct FChannelState
+    	{
+    		FHarmonixVuMeterAnalyzerChannelValues LatestValues{};
+    		float RiseTargetSquared = 0.0f;
+    		float PeakHoldRemainingMs = 0.0f;
+    	};
+
+		static void UpdateChannel(
+			FChannelState& ChannelState,
+			FHarmonixVuMeterAnalyzerChannelValues& ResultsValues,
+			float PeakSquared,
+			float MeanSquared,
+			float ElapsedMs,
+			const FHarmonixVuMeterAnalyzerSettings& SettingsToUse);
+
+		const float SampleRate;
+		
+		FCriticalSection SettingsGuard;
+		FHarmonixVuMeterAnalyzerSettings Settings;
+		
+		TArray<FChannelState> ChannelStates;
+		FChannelState MonoState;
+	};
+}
diff --git a/Engine/Plugins/Runtime/Harmonix/Source/HarmonixDsp/Public/HarmonixDsp/AudioAnalysis/WaveformAnalyzer.h b/Engine/Plugins/Runtime/Harmonix/Source/HarmonixDsp/Public/HarmonixDsp/AudioAnalysis/WaveformAnalyzer.h
new file mode 100644
--- /dev/null
+++ b/Engine/Plugins/Runtime/Harmonix/Source/HarmonixDsp/Public/HarmonixDsp/AudioAnalysis/WaveformAnalyzer.h
@@ -0,0 +1,77 @@
+// Copyright Epic Games, Inc. All Rights Reserved.
+
+#pragma once
+
+#include "HarmonixDsp/AudioBuffer.h"
+
+#include "WaveformAnalyzer.generated.h"
+
+USTRUCT(BlueprintType)
+struct FHarmonixWaveformAnalyzerSettings
+{
+	GENERATED_BODY()
+
+	UPROPERTY(EditAnywhere, Category="Settings")
+	int32 NumBinsHeld = 1024;
+
+	UPROPERTY(EditAnywhere, Category="Settings")
+	int32 NumBinsPerSecond = 512;
+
+	UPROPERTY(EditAnywhere, Category="Settings")
+	int32 SmoothingDistance = 0;
+
+	UPROPERTY(EditAnywhere, Category="Settings")
+	float SmoothingFactor = 0;
+};
+
+USTRUCT(BlueprintType)
+struct FHarmonixWaveformAnalyzerResults
+{
+	GENERATED_BODY()
+	
+	UPROPERTY(BlueprintReadOnly, Category="Results")
+	TArray<float> WaveformRaw;
+
+	UPROPERTY(BlueprintReadOnly, Category="Results")
+	TArray<float> WaveformSmoothed;
+};
+
+namespace Harmonix::Dsp::AudioAnalysis
+{
+	class HARMONIXDSP_API FWaveformAnalyzer
+	{
+	public:
+		explicit FWaveformAnalyzer(float InSampleRate);
+
+		/**
+		 * @brief Set the settings for this analyzer
+		 * @param InSettings - The settings to use
+		 */
+		void SetSettings(const FHarmonixWaveformAnalyzerSettings& InSettings);
+
+		/**
+		 * @brief Reset the analyzer to its initial state
+		 */
+		void Reset();
+
+		/**
+		 * @brief Process a buffer and get the analyzer results
+		 * @param InBuffer - The audio buffer to analyze
+		 * @param InOutResults - The results from the analysis
+		 */
+		void Process(const TAudioBuffer<float>& InBuffer, FHarmonixWaveformAnalyzerResults& InOutResults);
+		
+	private:
+		const float SampleRate;
+		
+		FCriticalSection SettingsGuard;
+		FHarmonixWaveformAnalyzerSettings Settings;
+
+		struct FState
+		{
+			int32 NumLeftoverFrames = 0;
+			float LeftoverBinSumSquared = 0;
+		};
+		FState State;
+	};
+}
diff --git a/Engine/Plugins/Runtime/Harmonix/Source/HarmonixDsp/Public/HarmonixDsp/AudioBufferConfig.h b/Engine/Plugins/Runtime/Harmonix/Source/HarmonixDsp/Public/HarmonixDsp/AudioBufferConfig.h
new file mode 100644
--- /dev/null
+++ b/Engine/Plugins/Runtime/Harmonix/Source/HarmonixDsp/Public/HarmonixDsp/AudioBufferConfig.h
@@ -0,0 +1,56 @@
+// Copyright Epic Games, Inc. All Rights Reserved.
+#pragma once
+
+// Harmonix
+#include "HarmonixDsp/AudioBufferConstants.h"
+
+
+// UE
+#include "Containers/UnrealString.h"
+#include "Logging/LogMacros.h"
+#include "HAL/Platform.h"
+
+DECLARE_LOG_CATEGORY_EXTERN(LogAudioBufferConfig, Log, All);
+
+struct HARMONIXDSP_API FAudioBufferConfig
+{
+public:
+	static const uint32 kMaxAudioBufferChannels = 8;
+
+	FAudioBufferConfig();
+	FAudioBufferConfig(int32 InNumChannels, int32 InNumFrames, float InSampleRate = 0.0f, bool InInterleaved = false);
+	FAudioBufferConfig(EAudioBufferChannelLayout InChannelLayout, int32 InNumFrames, float InSampleRate = 0.0f, bool InInterleaved = false);
+	FAudioBufferConfig(EAudioBufferChannelLayout InChannelLayout, int32 InNumChannels, int32 InNumFrames, float InSampleRate = 0.0f, bool InInterleaved = false);
+
+	int32 GetNumTotalSamples() const;
+	float GetSampleRate() const;
+	void  SetSampleRate(float InSamplesPerSecond);
+	
+	FORCEINLINE int32 GetNumChannels() const { return NumChannels; }
+	void SetNumChannels(int32 InNumChannels);
+
+	FORCEINLINE int32 GetNumFrames() const { return NumFrames; }
+	void SetNumFrames(int32 InNumFrames);
+
+	FAudioBufferConfig GetInterleavedConfig() const;
+
+	FAudioBufferConfig GetDeinterleavedConfig() const;
+
+	FORCEINLINE bool GetIsInterleaved() const { return Interleaved; }
+
+	FORCEINLINE EAudioBufferChannelLayout GetChannelLayout() const { return ChannelLayout; }
+	void SetChannelLayout(EAudioBufferChannelLayout InChannelLayout);
+
+	FORCEINLINE uint32 GetChannelMask() const { return ChannelMask; }
+	void SetChannelMask(uint32 InChannelMask);
+
+	bool operator==(const FAudioBufferConfig&) const;
+
+private:
+	float SampleRate;
+	int32 NumChannels;
+	int32 NumFrames;
+	uint32 ChannelMask;
+	bool Interleaved;
+	EAudioBufferChannelLayout ChannelLayout;
+};
\ No newline at end of file
diff --git a/Engine/Plugins/Runtime/Harmonix/Source/HarmonixDsp/Public/HarmonixDsp/AudioData.h b/Engine/Plugins/Runtime/Harmonix/Source/HarmonixDsp/Public/HarmonixDsp/AudioData.h
new file mode 100644
--- /dev/null
+++ b/Engine/Plugins/Runtime/Harmonix/Source/HarmonixDsp/Public/HarmonixDsp/AudioData.h
@@ -0,0 +1,141 @@
+// Copyright Epic Games, Inc. All Rights Reserved.
+#pragma once
+
+#include "HarmonixDsp/AudioBufferConstants.h"
+
+#include "Containers/Map.h"
+#include "Templates/SharedPointer.h"
+#include "HAL/Platform.h"
+#include "HAL/CriticalSection.h"
+#include "Misc/ScopeLock.h"
+
+#include "Logging/LogMacros.h"
+
+DECLARE_LOG_CATEGORY_EXTERN(LogHarmonixAudioData, Log, All);
+
+namespace Audio
+{
+	class IProxyData;
+};
+
+enum class EAudioEncodedFormat : uint8
+{
+	PCM,
+	Float32,
+	Num,
+	Invalid
+};
+
+namespace HarmonixDsp
+{
+	class IAudioData;
+
+	using FSharedAudioDataRef = TSharedRef<IAudioData, ESPMode::ThreadSafe>;
+	using FSharedAudioDataPtr = TSharedPtr<IAudioData, ESPMode::ThreadSafe>;
+	
+	using FSharedProxyDataRef = TSharedRef<Audio::IProxyData, ESPMode::ThreadSafe>;
+	using FSharedProxyDataPtr = TSharedPtr<Audio::IProxyData, ESPMode::ThreadSafe>;
+
+	using FWeakAudioDataPtr = TWeakPtr<IAudioData, ESPMode::ThreadSafe>;
+
+	using FSharedProxyDataMap = TMap<FSharedProxyDataRef, FWeakAudioDataPtr>;
+
+	/**
+	 *
+	 * When creating IAudioData, you want to call
+	 *
+	 *    IAudioData::GetShared<TDataType, TProxyType>(MyAudioData, ProxyData);
+	 *
+	 * This will return a ptr to a Existing Audio Data if it has aleady been loaded instead of creating a new one.
+	 *
+	 * This works by mapping the proxy data to the audio data
+	 *
+	 *    TSharedRef<IAudioProxy> --> TWeakPtr<IAudioData>
+	 *
+	 * This assumes there's a 1-1 relationship between AudioProxy and AudioData (which I think is a safe assumption)
+	 *
+	 * The mapping is to a WeakPtr of IAudioData. That way, when all shared references to the IAudioData go away,
+	 * the destructor of ~IAudioData gets called and will automatically remove itself from the SharedAudioData mapping
+	 *
+	 * For example:
+	 *
+	 *     FMyAudioData::GetShared<FMyAudioData, FProxyData>(ProxyData.ToSharedRef(), Args...);
+	 *
+	 * any additional args that are used for the construction of a new audio data instance won't be applied to existing audio data instances.
+	 */
+	class HARMONIXDSP_API IAudioData : public TSharedFromThis<IAudioData>
+	{
+	public:
+
+		virtual const FName& GetName() const = 0;
+		virtual EAudioEncodedFormat GetFormat() const { return EAudioEncodedFormat::Invalid; }
+		virtual bool IsStreaming() const = 0;
+
+		virtual bool GetHasTailSection() const = 0;
+		virtual bool GetHasLoopSection() const = 0;
+		virtual uint32 GetLoopStartFrame() const = 0;
+		virtual uint32 GetLoopEndFrame() const = 0;
+
+		virtual float GetSampleRate() const = 0;
+		virtual uint32 GetNumFrames() const = 0;
+		virtual int32 GetNumChannels() const = 0;
+		virtual float GetDurationSeconds() const = 0;
+
+		// what speaker channels are represented in the data?
+		virtual uint32 GetChannelMask() const = 0;
+
+		// what speaker layout was this authored for?
+		virtual EAudioBufferChannelLayout GetChannelLayout() const = 0;
+
+		virtual bool Failed() const = 0;
+
+		virtual ~IAudioData();
+
+		static ESpeakerMask::Type GetChannelMaskForNumChannels(int32 InNumChannels);
+
+	public:
+
+		// You _must_ implement the constructor of your audio data 
+		// so that its first argument is a TSharedRef to the ProxyData it wraps
+		template<typename TDataType, typename TProxyType, typename... Args>
+		static TSharedPtr<TDataType, ESPMode::ThreadSafe> GetShared(TSharedRef<TProxyType, ESPMode::ThreadSafe> ProxyDataRef, Args... InArgs)
+		{
+			FScopeLock Lock(&MapCritSec);
+
+			if (FWeakAudioDataPtr* WeakData = FindShared(ProxyDataRef))
+			{
+				// if the weak ptr is null, then the previous reference to this data is being destroyed
+				// but hasn't been removed from the map yet. We'll have to add a new reference.
+				if (FSharedAudioDataPtr SharedData = WeakData->Pin())
+				{
+					TSharedPtr<TDataType, ESPMode::ThreadSafe> OutData = StaticCastSharedPtr<TDataType>(SharedData);
+					// There a reference to AudioData, and it _is_ valid. 
+					// Check that it's the type we asked for!
+					check(OutData);
+					return OutData;
+				}
+			}
+
+			TSharedPtr<TDataType, ESPMode::ThreadSafe> NewData = MakeShared<TDataType>(ProxyDataRef, Forward<Args>(InArgs)...);
+			NewData->MyProxyData = ProxyDataRef.ToSharedPtr();
+			AddShared(ProxyDataRef, NewData.ToWeakPtr());
+			return NewData;
+		}
+
+	private:
+
+		// These methods are non-locking, so you need to lock in the outer scope
+		static FWeakAudioDataPtr* FindShared(FSharedProxyDataRef ProxyDataRef);
+		static void AddShared(FSharedProxyDataRef ProxyDataRef, FWeakAudioDataPtr AudioData);
+
+		// Have a back ptr to the mapped Proxy Data so we can use it to index into the TMap
+		FSharedProxyDataPtr MyProxyData;
+
+		
+		static FCriticalSection MapCritSec;
+	};
+
+
+	
+}
+
diff --git a/Engine/Plugins/Runtime/Harmonix/Source/HarmonixDsp/Public/HarmonixDsp/AudioDataRenderer.h b/Engine/Plugins/Runtime/Harmonix/Source/HarmonixDsp/Public/HarmonixDsp/AudioDataRenderer.h
new file mode 100644
--- /dev/null
+++ b/Engine/Plugins/Runtime/Harmonix/Source/HarmonixDsp/Public/HarmonixDsp/AudioDataRenderer.h
@@ -0,0 +1,66 @@
+// Copyright Epic Games, Inc. All Rights Reserved.
+#pragma once
+
+#include "Templates/SharedPointer.h"
+
+#include "HarmonixDsp/Streaming/TrackChannelInfo.h"
+
+namespace HarmonixDsp
+{
+	class IAudioData;
+}
+
+class IStretcherAndPitchShifter;
+struct FTrackChannelInfo;
+template<typename Type> 
+class TAudioBuffer;
+class FGainMatrix;
+class FFusionSampler;
+
+class IAudioDataRenderer : public TSharedFromThis<IAudioDataRenderer>
+{
+public:
+
+	/** Settings for a AudioDataRenderer. */
+	struct FSettings
+	{
+		TSharedPtr<IStretcherAndPitchShifter, ESPMode::ThreadSafe> Shifter = nullptr;
+		const TArray<FTrackChannelInfo>* TrackChannelInfo = nullptr;
+		const FFusionSampler* Sampler = nullptr;
+	};
+
+	virtual ~IAudioDataRenderer() {};
+
+	virtual void SetAudioData(TSharedRef<HarmonixDsp::IAudioData, ESPMode::ThreadSafe> AudioData, const FSettings& InSettings) = 0;
+
+	virtual void Reset() = 0;
+
+	virtual const TSharedPtr<HarmonixDsp::IAudioData, ESPMode::ThreadSafe> GetAudioData() const = 0;
+
+	virtual void MigrateToSampler(const FFusionSampler* InSampler) = 0;
+
+	virtual void SetFrame(uint32 InFrameNum) = 0;
+
+	virtual double Render(TAudioBuffer<float>& OutBuffer, double InPos, int32 InMaxFrame, double InResampleInc, double InPitchShift, double InSpeed,
+		bool MaintainPitchWhenSpeedChanges, bool InShouldHonorLoopPoints, const FGainMatrix& InGain) = 0;
+
+	virtual double RenderUnshifted(TAudioBuffer<float>& OutBuffer, double InPos, int32 InMaxFrame, double InInc,
+		bool InShouldHonorLoopPoints, const FGainMatrix& InGain) = 0;
+
+protected:
+
+	struct FLerpData
+	{
+		uint32 PosA;
+		uint32 PosB;
+
+		// pos a and b Relative to start position, and never loop
+		uint32 PosARelative;
+		uint32 PosBRelative;
+
+		float WeightA;
+		float WeightB;
+	};
+
+	double CalculateLerpData(FLerpData* LerpArray, int32 InNumPoints, uint32 InNumOutFrames, double InPos, int32 InMaxFrame, bool InHonorLoopPoints, double Inc) const;
+};
\ No newline at end of file
diff --git a/Engine/Plugins/Runtime/Harmonix/Source/HarmonixDsp/Public/HarmonixDsp/Containers/EnumArray.h b/Engine/Plugins/Runtime/Harmonix/Source/HarmonixDsp/Public/HarmonixDsp/Containers/EnumArray.h
new file mode 100644
--- /dev/null
+++ b/Engine/Plugins/Runtime/Harmonix/Source/HarmonixDsp/Public/HarmonixDsp/Containers/EnumArray.h
@@ -0,0 +1,202 @@
+// Copyright Epic Games, Inc. All Rights Reserved.
+
+#pragma once
+
+#include "Misc/AssertionMacros.h"
+
+// assumes EnumType is an enum 
+// where EnumType::Num is a member
+// and corresponds to the number of
+// enumerated values.
+template<typename EnumType, typename ValueType> 
+struct TEnumArray
+{
+public:
+	using FValueType = ValueType;
+	using EEnumType = EnumType;
+	virtual ~TEnumArray() {}
+
+	// TODO: doesn't compile
+	//template <typename ItArrayType, typename ItValueType>
+	//struct TBaseIterator
+	//{
+	//public:
+
+	//	FORCEINLINE TBaseIterator(ItArrayType& InIterArray, int Index)
+	//		: IterArray(InIterArray)
+	//		, IterIndex(Index)
+	//	{}
+	//	
+	//	FORCEINLINE ItValueType& operator*() const { return IterArray[IterIndex]; }
+	//	FORCEINLINE ItValueType* operator->() const { return &IterArray[IterIndex]; }
+
+	//	FORCEINLINE bool operator==(const TBaseIterator& Other) const { return IterIndex == Other.IterIndex; }
+	//	FORCEINLINE bool operator!=(const TBaseIterator& Other) const { return IterIndex != Other.IterIndex; }
+
+	//	FORCEINLINE int Index() const { return IterIndex; }
+	//	FORCEINLINE EnumType Key() const { return (EnumType)IterIndex; }
+
+	//	FORCEINLINE TBaseIterator& operator++() { ++IterIndex; return *this; }
+
+	//private:
+
+	//	ItArrayType& IterArray;
+	//	int IterIndex = 0;
+	//};
+
+	//typedef TBaseIterator<TEnumArray, ValueType> TIterator;
+	//typedef TBaseIterator<const TEnumArray, const ValueType> TConstIterator;
+
+	//FORCEINLINE TIterator begin() { return TIterator(*this, 0); }
+	//FORCEINLINE TIterator end() { return TIterator(*this, Num()); }
+	//FORCEINLINE TConstIterator begin() const { return TConstIterator(*this, 0); }
+	//FORCEINLINE TConstIterator end() const { return TConstIterator(*this, Num()); }
+
+	ValueType& operator[](EnumType e)
+	{
+		int Index = (int)e;
+		check(IsValidIndex(Index));
+		return Array[Index];
+	}
+
+	const ValueType& operator[](EnumType e) const
+	{
+		int Index = (int)e;
+		check(IsValidIndex(Index));
+		return Array[Index];
+	}
+
+	ValueType& operator[](int Index)
+	{
+		check(IsValidIndex(Index));
+		return Array[Index];
+	}
+
+	const ValueType& operator[](int Index) const
+	{
+		check(IsValidIndex(Index));
+		return Array[Index];
+	}
+
+	static constexpr int Num()
+	{
+		return (int)EnumType::Num;
+	}
+
+	bool IsValidIndex(int Index) const
+	{
+		return Index >= 0 && Index < Num();
+	}
+
+	TEnumArray& operator=(const TEnumArray& Other)
+	{
+		for (int idx = 0; idx < Num(); ++idx)
+		{
+			Array[idx] = Other.Array[idx];
+		}
+		return *this;
+	}
+protected:
+
+	ValueType& Get(EnumType e)
+	{
+		int Index = (int)e;
+		check(IsValidIndex(Index));
+		return Array[Index];
+	}
+
+private:
+
+	ValueType Array[Num()];
+};
+
+
+/**********
+* 
+* usage 
+* 
+* 
+* EStructIndex
+* {
+*	Zero,
+*	One,
+*	Num
+* }
+* 
+* FStruct
+* {
+*	int Data = 0;
+* }
+* 
+* // allocates an array of structs with size EStructIndex::Num
+* TEnumArray<EStructIndex, FStruct> EnumStruct;
+* 
+* // access statically allocated struct by reference
+* FStruct& ZeroRef = EnumStruct[EStructIndex::Zero]
+* 
+* 
+***********/
+
+// assumes EnumType::Num is a valid value and is the correct MAX number in the Enum
+#define IMPLEMENT_ENUM_ARRAY(ArrayName, EnumType, StructType)											\
+	using FValueType = StructType;																		\
+	using FArrayType = StructType ## Array;																\
+	using EEnumType = EnumType;																			\
+																										\
+	const StructType& operator[](int Index) const														\
+	{																									\
+		check(IsValidIndex(Index));																		\
+		return ArrayName[Index];																		\
+	}																									\
+																										\
+	StructType& operator[](int Index)																	\
+	{																									\
+		check(IsValidIndex(Index));																		\
+		return ArrayName[Index];																		\
+	}																									\
+																										\
+	const StructType& operator[](EnumType Index) const													\
+	{																									\
+		check(IsValidIndex(static_cast<int>(Index)))													\
+		return ArrayName[static_cast<int>(Index)];														\
+	}																									\
+																										\
+	StructType& operator[](EnumType Index)																\
+	{																									\
+		check(IsValidIndex(static_cast<int>(Index)));													\
+		return ArrayName[static_cast<int>(Index)];														\
+	}																									\
+																										\
+	static constexpr int Num()																			\
+	{																									\
+		return static_cast<int>(EnumType::Num);															\
+	}																									\
+																										\
+	bool IsValidIndex(int Index) const																	\
+	{																									\
+		return Index >= 0 && Index < Num();																\
+	}																									\
+																										\
+	StructType ## Array() = default;																	\
+																										\
+	StructType ## Array(const FArrayType& Other) = default;												\
+																										\
+	FArrayType& operator=(const FArrayType& Other)														\
+	{																									\
+		for (int idx = 0; idx < Num(); ++idx)															\
+		{																								\
+			ArrayName[idx] = Other.GetArray()[idx];														\
+		}																								\
+		return *this;																					\
+	}																									\
+																										\
+private:																								\
+	const StructType* GetArray() const { return ArrayName; }											\
+		  StructType* GetArray()       { return ArrayName; }											\
+																										\
+public:					
+
+// Makes a property with the name "Name", for a TEnumArray or a Struct that uses IMPLEMENT_ENUM_ARRAY
+// The property will then be a reference to the array element indexed by the enum value of the same name
+// assumes "this" implements the ref by index operator (operator[])
+#define ENUM_PROPERTY(Name) FValueType& Name = this->operator[](EEnumType::Name)
\ No newline at end of file
diff --git a/Engine/Plugins/Runtime/Harmonix/Source/HarmonixDsp/Public/HarmonixDsp/Effects/BitCrusher.h b/Engine/Plugins/Runtime/Harmonix/Source/HarmonixDsp/Public/HarmonixDsp/Effects/BitCrusher.h
new file mode 100644
--- /dev/null
+++ b/Engine/Plugins/Runtime/Harmonix/Source/HarmonixDsp/Public/HarmonixDsp/Effects/BitCrusher.h
@@ -0,0 +1,50 @@
+// Copyright Epic Games, Inc. All Rights Reserved.
+#pragma once
+
+#include "HarmonixDsp/AudioBuffer.h"
+#include "HarmonixDsp/Ramper.h"
+
+#include "HAL/Platform.h"
+
+namespace Harmonix::Dsp::Effects
+{
+	class HARMONIXDSP_API FBitCrusher
+	{
+	public:
+
+		FBitCrusher();
+		virtual ~FBitCrusher();
+
+		void SetCrush(uint16 CrushLevel);
+		void SetInputGain(float Gain, bool Snap = false);
+		void SetOutputGain(float Gain, bool Snap = false);
+		void SetWetGain(float Wet, bool Snap = false);
+		void SetSampleHoldFactor(uint16 SampleHoldFactor);
+
+		uint16 GetCrush() const { return CrushCoef; }
+		uint16 GetSampleHoldFactor() const { return SampleHoldFactor; }
+		float GetWetGain() const { return WetGainRamper.GetTarget(); }
+
+		void Setup(int32 MaxChannels, float SampleRate);
+		void Reset();
+		int32 GetMaxChannels() { return MaxChannels; }
+		void Process(TAudioBuffer<float>& Buffer) { Process(Buffer, Buffer); }
+		void Process(TAudioBuffer<float>& InBuffer, TAudioBuffer<float>& OutBuffer);
+
+	private:
+
+		int32 MaxChannels;
+		uint32 SampleRate;
+
+		static const int32 kMaxNumChannels = 16;
+		float SampleHoldData[kMaxNumChannels];
+
+		TLinearRamper<float> WetGainRamper;
+		TLinearRamper<float> InputGainRamper;
+		TLinearRamper<float> OutputGainRamper;
+
+		uint16 SampleHoldFactor;
+		uint16 CrushCoef;
+		uint16 SampleHoldStage;
+	};
+}
\ No newline at end of file
diff --git a/Engine/Plugins/Runtime/Harmonix/Source/HarmonixDsp/Public/HarmonixDsp/Effects/DistortionV1.h b/Engine/Plugins/Runtime/Harmonix/Source/HarmonixDsp/Public/HarmonixDsp/Effects/DistortionV1.h
new file mode 100644
--- /dev/null
+++ b/Engine/Plugins/Runtime/Harmonix/Source/HarmonixDsp/Public/HarmonixDsp/Effects/DistortionV1.h
@@ -0,0 +1,50 @@
+// Copyright Epic Games, Inc. All Rights Reserved.
+#pragma once
+
+#include "HarmonixDsp/AudioBuffer.h"
+#include "HarmonixDsp/Effects/BiquadFilter.h"
+#include "HarmonixDsp/Effects/FirFilter.h"
+#include "HarmonixDsp/Effects/Settings/DistortionSettings.h"
+
+namespace Harmonix::Dsp::Effects
+{
+
+	class HARMONIXDSP_API FDistortionV1
+	{
+	public:
+		FDistortionV1(int32 sampleRate = 48000);
+
+		void Process(TAudioBuffer<float>& InOutBuffer);
+
+		void SetInputGainDb(float GainDb);
+		void SetOutputGainDb(float GainDb);
+		void SetType(EDistortionTypeV1 Type);
+		void SetupFilter(int32 Index, const FBiquadFilterSettings& settings);
+		void SetupFilter(int32 Index, const FDistortionFilterSettings& Settings);
+		void SetOversample(bool Oversample, int32 RenderBufferSizeFrames);
+		void SetSampleRate(int32 SampleRate);
+
+		static const int32 kNumFilters = 3;
+		static const int32 kMaxChannels = 8;
+
+	private:
+		void ProcessChannel(int32 Channel, float* InOutSamples, int32 NumFrames);
+
+		EDistortionTypeV1    Type;
+		float                InputGain;
+		float                OutputGain;
+		float                DCAdjust;
+		bool                 FilterEnabled[kNumFilters];
+		bool                 FilterPreClip[FDistortionV1::kNumFilters];
+		FBiquadFilterCoefs   FilterCoefs[kNumFilters];
+		TBiquadFilter<double>        Filter[kNumFilters][kMaxChannels];
+		FFirFilter32         OversampleFilterUp[kMaxChannels];
+		FFirFilter32         OversampleFilterDown[kMaxChannels];
+		bool                 DoOversampling;
+		TAudioBuffer<float>  UpsampleBuffer;
+		int32                  SampleRate;
+		static const int32     kNumFilterTaps = 32;
+		static const float   kOversamplingFilterTaps[kNumFilterTaps];
+	};
+
+};
diff --git a/Engine/Plugins/Runtime/Harmonix/Source/HarmonixDsp/Public/HarmonixDsp/Effects/DistortionV2.h b/Engine/Plugins/Runtime/Harmonix/Source/HarmonixDsp/Public/HarmonixDsp/Effects/DistortionV2.h
new file mode 100644
--- /dev/null
+++ b/Engine/Plugins/Runtime/Harmonix/Source/HarmonixDsp/Public/HarmonixDsp/Effects/DistortionV2.h
@@ -0,0 +1,99 @@
+// Copyright Epic Games, Inc. All Rights Reserved.
+#pragma once
+
+#include "HarmonixDsp/AudioBuffer.h"
+#include "HarmonixDsp/AudioUtility.h"
+#include "HarmonixDsp/Ramper.h"
+#include "HarmonixDsp/Effects/Settings/DistortionSettings.h"
+#include "HarmonixDsp/Effects/BiquadFilter.h"
+#include "HarmonixDsp/Effects/FirFilter.h"
+#include "HAL/CriticalSection.h"
+#include "HAL/Platform.h"
+
+namespace Harmonix::Dsp::Effects
+{
+
+	class HARMONIXDSP_API FDistortionV2
+	{
+	public:
+		FDistortionV2(uint32 SampleRate = 48000, uint32 MaxRenderBufferSize = 128);
+		void Setup(uint32 SampleRate, uint32 MaxRenderBufferSize);
+		virtual ~FDistortionV2() {};
+
+		void Process(TAudioBuffer<float>& InOutBuffer) { Process(InOutBuffer, InOutBuffer); }
+		void Process(TAudioBuffer<float>& InBuffer, TAudioBuffer<float>& OutBuffer);
+
+		void  SetInputGainDb(float InGainDb, bool Snap = false);
+		float GetInputGainDb() const;
+
+		void SetInputGain(float InGain, bool Snap = false);
+		float GetInputGain() const { return InputGain; }
+
+		void  SetOutputGainDb(float InGainDb, bool Snap = false);
+		float GetOutputGainDb() const;
+
+		void SetOutputGain(float InGain, bool Snap = false);
+		float GetOutputGain() const { return OutputGain; }
+
+		void  SetDryGainDb(float InGainDb, bool Snap = false);
+		float GetDryGainDb() const;
+
+		void SetDryGain(float InGain, bool Snap = false);
+		float GetDryGain() const { return DryGain; }
+
+		void  SetWetGainDb(float InGainDb, bool Snap = false);
+		float GetWetGainDb() const;
+
+		void SetWetGain(float InGain, bool Snap = false);
+		float GetWetGain() const { return WetGain; }
+
+		void  SetMix(float Mix, bool Snap = false);
+		void  SetDCOffset(float Offset, bool Snap = false);
+
+		float GetDCOffset() const { return DCAdjust; }
+		void  SetType(uint8 InType);
+		void  SetType(EDistortionTypeV2 Type);
+		EDistortionTypeV2 GetType() const { return Type; }
+		void SetupFilter(int32 Index, const FDistortionFilterSettings& InSettings);
+		int32  GetFilterPasses(int32 Index) { return FilterPasses[Index]; }
+		void SetOversample(bool Oversample, int32 RenderBufferSizeFrames);
+		bool GetOversample() const { return DoOversampling; }
+		void SetSampleRate(uint32 SampleRate);
+
+		void Setup(const FDistortionSettingsV2& Settings, uint32 SampleRate, uint32 RenderBufferSizeFrames, bool Snap);
+
+		void Reset();
+
+		static const int32 kMaxChannels = 8;
+		static const int32 kMaxFilterPasses = 3;
+		static const int32 kRampHops = 16;
+
+	private:
+		EDistortionTypeV2	  Type;
+		TLinearRamper<float>  InputGain;
+		TLinearRamper<float>  OutputGain;
+		TLinearRamper<float>  DCAdjust;
+		TLinearRamper<float>  DryGain;
+		TLinearRamper<float>  WetGain;
+		bool		          FilterPreClip[FDistortionSettingsV2::kNumFilters];
+		FBiquadFilterSettings FilterSettings[FDistortionSettingsV2::kNumFilters];
+		TLinearRamper<float>  FilterGain[FDistortionSettingsV2::kNumFilters];
+
+		TLinearRamper<FBiquadFilterCoefs> FilterCoefs[FDistortionSettingsV2::kNumFilters];
+
+		uint32					FilterPasses[kMaxFilterPasses];
+
+		TMultipassBiquadFilter<double, kMaxFilterPasses> Filter[FDistortionSettingsV2::kNumFilters][kMaxChannels];
+
+		FFirFilter32			OversampleFilterUp[kMaxChannels];
+		FFirFilter32			OversampleFilterDown[kMaxChannels];
+		bool					DoOversampling;
+		TAudioBuffer<float>		UpsampleBuffer;
+		uint32					SampleRate;
+
+		static const uint32   kNumFilterTaps = 32;
+		static const float    kOversamplingFilterTaps[kNumFilterTaps];
+
+		mutable FCriticalSection SettingsLock;
+	};
+}
\ No newline at end of file
diff --git a/Engine/Plugins/Runtime/Harmonix/Source/HarmonixDsp/Public/HarmonixDsp/Effects/Vocoder.h b/Engine/Plugins/Runtime/Harmonix/Source/HarmonixDsp/Public/HarmonixDsp/Effects/Vocoder.h
new file mode 100644
--- /dev/null
+++ b/Engine/Plugins/Runtime/Harmonix/Source/HarmonixDsp/Public/HarmonixDsp/Effects/Vocoder.h
@@ -0,0 +1,111 @@
+// Copyright Epic Games, Inc. All Rights Reserved.
+#pragma once
+
+#include "HarmonixDsp/Effects/Settings/VocoderSettings.h"
+#include "HarmonixDsp/Effects/BiquadFilter.h"
+#include "HarmonixDsp/AudioBuffer.h"
+#include "HAL/CriticalSection.h"
+
+namespace Harmonix::Dsp::Effects
+{
+
+	class FFourBiquads;
+
+	class FVocoder
+	{
+	public:
+
+		FVocoder();
+		~FVocoder();
+
+		static void Init();
+
+		void Setup(int32 ChannelCount, float SampleRate, int32 FrameCount, EVocoderBandConfig VocoderBandConfig);
+		void Setup(int32 InChannelCount, float InSampleRate, int32 InFrameCount)
+		{
+			Setup(InChannelCount, InSampleRate, InFrameCount, TargetSettings.BandConfig);
+		}
+
+		void Process(const TAudioBuffer<float>& InCarrier, const TAudioBuffer<float>& InModulator, TAudioBuffer<float>& OutBuffer);
+		void Process(TAudioBuffer<float>& InOutCarrier, const TAudioBuffer<float>& InModulator) { Process(InOutCarrier, InModulator, InOutCarrier); }
+
+		// Mutators
+		void ImportSettings(const FVocoderSettings& InSettings) { TargetSettings = InSettings; }
+		void SetChannelCount(int32 ChannelCount) { TargetSettings.ChannelCount = ChannelCount; }
+		void SetBandConfig(EVocoderBandConfig BandConfig) { TargetSettings.BandConfig = BandConfig; }
+		void SetSoloing(bool Soling) { TargetSettings.Soloing = Soling; }
+		void SetCarrierGain(float Gain) { TargetSettings.CarrierGain = Gain; }
+		void SetModulatorGain(float Gain) { TargetSettings.ModulatorGain = Gain; }
+		void SetCarrierThin(float CarrierThin) { TargetSettings.CarrierThin = CarrierThin; }
+		void SetModulatorThin(float ModulatorThin) { TargetSettings.ModulatorThin = ModulatorThin; }
+		void SetAttack(float Attack) { TargetSettings.Attack = Attack; }
+		void SetRelease(float Release) { TargetSettings.Release = Release; }
+		void SetHighEmphasis(float HighEmphasis) { TargetSettings.HighEmphasis = HighEmphasis; }
+		void SetWet(float Wet) { TargetSettings.Wet = Wet; }
+		void SetOutputGain(float Gain) { TargetSettings.OutputGain = Gain; }
+		void SetPerBandSettings(TArray<FVocoderBand> Bands) { TargetSettings.Bands = Bands; }
+
+		// Accessors
+		int32 GetChannelCount() const { return CurrentSettings.ChannelCount; }
+		EVocoderBandConfig GetBandConfigIndex() const { return CurrentSettings.BandConfig; }
+		int32 GetBandCount() const { return CurrentSettings.GetBandConfig().BandCount; }
+		float GetCarrierThin() const { return CurrentSettings.CarrierThin; }
+		float GetModulatorThin() const { return CurrentSettings.ModulatorThin; }
+		float GetAttack() const { return CurrentSettings.Attack; }
+		float GetRelease() const { return CurrentSettings.Release; }
+		float GetHighEmphasis() const { return CurrentSettings.HighEmphasis; }
+		float GetCarrierGain() const { return CurrentSettings.CarrierGain; }
+		float GetModulatorGain() const { return CurrentSettings.ModulatorGain; }
+		float GetWet() const { return CurrentSettings.Wet; }
+		float GetOutputGain() const { return CurrentSettings.OutputGain; }
+
+
+	private:
+		// Internal methods
+		void ProcessTimeDomain(int32 FrameCount);
+		void ResolveSettings(bool bForceAll = false);
+		void InitScratchBuffers();
+		void CreateFilters();
+		void CreateBandFrequencies();
+		void CreateCarrierFilterCoefficients();
+		void CreateModulatorFilterCoefficients();
+#if !UE_BUILD_DEBUG || SIMD_IN_DEBUG
+		void SetCarrierSets();
+		void SetModulatorSets();
+#endif
+		void InitEmphasisFilterCoefficients();
+		void DeleteFilters();
+		void DeleteBandFrequencies();
+		void DeleteCarrierFilterCoefficients();
+		void DeleteModulatorFilterCoefficients();
+
+		TAudioBuffer<float> WorkingCarrier;
+		TAudioBuffer<float> WorkingModulator;
+		TAudioBuffer<float> BandWork;
+		TAudioBuffer<float> WorkingOutput;
+		float* WorkingBandGains = nullptr;
+
+		float				CarrierQ = 0.0f;
+		float				ModulatorQ = 0.0f;
+		float* BandFrequencies = nullptr;
+#if !UE_BUILD_DEBUG || SIMD_IN_DEBUG
+		FFourBiquads** CarrierFilters = nullptr;
+		FFourBiquads** ModulatorFilters = nullptr;
+#else
+		FBiquadFilter** CarrierFilters = nullptr;
+		FBiquadFilter** ModulatorFilters = nullptr;
+#endif
+		FBiquadFilterCoefs* CarrierFilterCoefs = nullptr;
+		FBiquadFilterCoefs* ModulatorFilterCoefs = nullptr;
+		FBiquadFilter* HighIncreaseFilters = nullptr;
+		FBiquadFilterCoefs  HighIncreaseFilterCoefs;
+		FBiquadFilter* LowReductionFilters = nullptr;
+		FBiquadFilterCoefs  LowReductionFilterCoefs;
+		float** EnvelopeSamples = nullptr;
+
+		FCriticalSection	SettingsLock;
+		FVocoderSettings	CurrentSettings;
+		FVocoderSettings	TargetSettings;
+	};
+
+};
\ No newline at end of file
diff --git a/Engine/Plugins/Runtime/Harmonix/Source/HarmonixDsp/Public/HarmonixDsp/FusionSampler/FusionPatch.h b/Engine/Plugins/Runtime/Harmonix/Source/HarmonixDsp/Public/HarmonixDsp/FusionSampler/FusionPatch.h
new file mode 100644
--- /dev/null
+++ b/Engine/Plugins/Runtime/Harmonix/Source/HarmonixDsp/Public/HarmonixDsp/FusionSampler/FusionPatch.h
@@ -0,0 +1,171 @@
+// Copyright Epic Games, Inc. All Rights Reserved.
+
+#pragma once
+
+#include "HarmonixDsp/FusionSampler/Settings/FusionPatchSettings.h"
+#include "HarmonixDsp/FusionSampler/Settings/KeyzoneSettings.h"
+
+#include "Harmonix/AudioRenderableProxy.h"
+#include "IAudioProxyInitializer.h"
+
+#include "FusionPatch.generated.h"
+
+class UAudioSample;
+class FFusionSampler;
+
+// User facing data of FusionPatch
+USTRUCT(BlueprintType)
+struct HARMONIXDSP_API FFusionPatchData
+{
+	GENERATED_BODY()
+
+public:
+	friend class FFusionPatchJsonImporter;
+	friend class UFusionPatch;
+
+	IMPL_AUDIORENDERABLE_PROXYABLE(FFusionPatchData)
+
+	FFusionPatchData() {};
+
+	const TArray<FKeyzoneSettings>& GetKeyzones() const;
+	TArray<FKeyzoneSettings>& GetKeyzones();
+
+	int32 GetNumPresets() const;
+	int32 GetCurrentPresetIndex() const;
+	void SetCurrentPresetIndex(int32 Index);
+
+	bool IsCurrentPresetIndexValid() const { return Presets.IsValidIndex(CurrentPresetIndex); }
+	bool IsPresetIndexValid(int32 Index) const { return Presets.IsValidIndex(Index); }
+
+	bool HasPreset(FName PresetName) const;
+
+	// returns the index for the preset of the given name
+	// returns INDEX_NONE if there is no preset with the given name
+	int32 FindPresetIndex(FName PresetName) const;
+
+	const FFusionPatchSettings& GetPreset(int32 Index) const;
+	FFusionPatchSettings& GetPreset(int32 Index);
+	
+	const FFusionPatchSettings& GetCurrentPreset() const;
+	FFusionPatchSettings& GetCurrentPreset();
+
+	void InitProxyData(const Audio::FProxyDataInitParams& InitParams);
+	void AddPreset(const FFusionPatchSettings& Preset);
+	void UpdatePreset(int32 Index, const FFusionPatchSettings& Preset);
+
+	void DisconnectSampler(const FFusionSampler* sampler);
+
+private:
+
+	UPROPERTY(EditAnywhere, Category="Presets")
+	TArray<FKeyzoneSettings> Keyzones;
+
+	UPROPERTY(EditAnywhere, Category="Presets")
+	TArray<FFusionPatchSettings> Presets;
+	
+	UPROPERTY(EditAnywhere, Category="Presets")
+	int32 CurrentPresetIndex = -1;
+};
+
+// This next macro does a few things. 
+//   wrap "FFusionPatchData" in a proxy named "FFusionPatchDataProxy" 
+//   this proxy class can then be used as the "guts" for the metasound later
+USING_AUDIORENDERABLE_PROXY(FFusionPatchData, FFusionPatchDataProxy)
+
+UENUM(BlueprintType)
+enum class EFusionPatchAudioLoadResult : uint8
+{
+	Success,
+	Fail,
+	Cancelled
+};
+
+/**
+* Called when a load request for a sound has completed.
+*/
+DECLARE_DYNAMIC_DELEGATE_TwoParams(FOnFusionPatchLoadComplete, const class UFusionPatch*, LoadedFusionPatch, EFusionPatchAudioLoadResult, LoadResult);
+
+UCLASS(BlueprintType, meta = (DisplayName = "Fusion Patch"))
+class HARMONIXDSP_API UFusionPatch : public UObject, public IAudioProxyDataFactory
+{
+	GENERATED_BODY()
+
+public:
+	// IAudioProxyDataFactory
+	virtual TSharedPtr<Audio::IProxyData> CreateProxyData(const Audio::FProxyDataInitParams& InitParams);
+
+public:
+	
+	static const int32 kMaxLayersPerNote = 128;
+	static const int32 kVoicePriorityNoSteal = 0;
+
+	UFusionPatch();
+
+	int32 GetCurrentPresetIndex() const { return FusionPatchData.GetCurrentPresetIndex(); }
+	void SetCurrentPresetIndex(int32 Index);
+	int32 GetNumPresets() const { return FusionPatchData.GetNumPresets(); }
+
+	bool IsCurrentPresetIndexValid() const { return FusionPatchData.IsCurrentPresetIndexValid(); }
+	bool IsPresetIndexValid(int32 Index) const { return FusionPatchData.IsPresetIndexValid(Index); }
+
+	bool HasPreset(FName PresetName) const { return FusionPatchData.HasPreset(PresetName); }
+
+	// returns the index for the preset of the given name
+	// returns INDEX_NONE if there is no preset with the given name
+	int32 FindPresetIndex(FName PresetName) const { return FusionPatchData.FindPresetIndex(PresetName); }
+
+	const FFusionPatchSettings& GetPreset(int32 Index) const { return FusionPatchData.GetPreset(Index); }
+	const FFusionPatchSettings& GetCurrentPreset() const { return FusionPatchData.GetCurrentPreset(); }
+	const TArray<FKeyzoneSettings>& GetKeyzones() const { return FusionPatchData.GetKeyzones(); }
+
+	void AddPreset(FFusionPatchSettings& Preset);
+	void UpdatePreset(int32 Index, FFusionPatchSettings& Preset);
+
+#if WITH_EDITORONLY_DATA
+
+	UPROPERTY(VisibleAnywhere, Instanced, Category = ImportSettings)
+	TObjectPtr<class UAssetImportData> AssetImportData;
+
+	
+	virtual void PostEditChangeProperty(FPropertyChangedEvent& PropertyChangedEvent) override;
+	virtual void PostEditChangeChainProperty(FPropertyChangedChainEvent& PropertyChangedChainEvent) override;
+	virtual void PostInitProperties() override;
+	virtual void GetAssetRegistryTags(FAssetRegistryTagsContext Context) const override;
+#endif
+
+#if WITH_EDITOR
+
+	virtual EDataValidationResult IsDataValid(class FDataValidationContext& Context) const override;
+
+#endif
+
+	virtual void Serialize(FArchive& Ar) override;
+	virtual void PostLoad() override;
+
+protected:
+
+	UPROPERTY(EditAnywhere, BlueprintReadOnly, Category="Fusion Patch")
+	FFusionPatchData FusionPatchData;
+private:
+
+	friend class FFusionPatchJsonImporter;
+
+	// Notice here that we cache a pointer the Proxy's "Queue" so we can...
+	// 1 - Supply it to all instances of Metasound nodes rendering this data. How?
+	//     CreateNewProxyData instantiates a NEW unique ptr to an FFusionPatchDataProxy
+	//     every time it is called. All of those unique proxy instances refer to the same
+	//     queue... this one that we have cached.
+	// 2 - Modify that data in response to changes to this class's UPROPERTIES
+	//     so that we can hear data changes reflected in the rendered audio.
+	TSharedPtr<FFusionPatchDataProxy::QueueType, ESPMode::ThreadSafe> RenderableFusionPatchData;
+
+	void UpdateRenderableForNonTrivialChange();
+
+	enum class ELoadingState
+	{
+		None,
+		Loading
+	};
+
+	ELoadingState LoadingState = ELoadingState::None;
+};
\ No newline at end of file
diff --git a/Engine/Plugins/Runtime/Harmonix/Source/HarmonixDsp/Public/HarmonixDsp/FusionSampler/FusionSampler.h b/Engine/Plugins/Runtime/Harmonix/Source/HarmonixDsp/Public/HarmonixDsp/FusionSampler/FusionSampler.h
new file mode 100644
--- /dev/null
+++ b/Engine/Plugins/Runtime/Harmonix/Source/HarmonixDsp/Public/HarmonixDsp/FusionSampler/FusionSampler.h
@@ -0,0 +1,432 @@
+// Copyright Epic Games, Inc. All Rights Reserved.
+#pragma once
+
+#include "HarmonixDsp/Instruments/VirtualInstrument.h"
+
+#include "Containers/List.h"
+
+#include "HarmonixDsp/Effects/BitCrusher.h"
+#include "HarmonixDsp/Effects/BiquadFilter.h"
+#include "HarmonixDsp/Effects/Delay.h"
+#include "HarmonixDsp/Effects/DistortionV1.h"
+#include "HarmonixDsp/Effects/Vocoder.h"
+
+#include "HarmonixDsp/Modulators/Settings/AdsrSettings.h"
+#include "HarmonixDsp/Modulators/Settings/ModulatorSettings.h"
+#include "HarmonixDsp/Modulators/Modulator.h"
+#include "HarmonixDsp/Modulators/Lfo.h"
+
+#include "HarmonixDsp/FusionSampler/FusionPatch.h"
+#include "HarmonixDsp/FusionSampler/Settings/FusionPatchSettings.h"
+
+#include "HarmonixMidi/MidiConstants.h"
+
+#include "HAL/CriticalSection.h"
+
+#include "Templates/SharedPointer.h"
+
+DECLARE_LOG_CATEGORY_EXTERN(LogFusionSampler, Log, All);
+
+class FFusionVoice;
+struct FFusionPatchData;
+class FFusionVoicePool;
+using FSharedFusionVoicePoolPtr = TSharedPtr<FFusionVoicePool, ESPMode::ThreadSafe>;
+
+enum class EInstrumentRenderMode : uint8
+{
+	Parent = 0,
+	PreFxChild = 1,
+	PostFxChild = 2
+};
+
+// Make Fusion Sampler derive directly from virtual instrument
+class HARMONIXDSP_API FFusionSampler : public FVirtualInstrument
+{
+public:
+	FFusionSampler();
+	virtual ~FFusionSampler();
+
+	void ResetPatchRelatedState();
+
+	// configure this patch
+	virtual void Prepare(float InSampleRateHz, EAudioBufferChannelLayout InChannelLayout, uint32 InMaxSamples, bool bInAllocateBuffer = true) override;
+	void SetGainTable(const FGainTable* GainTable);
+	const FGainTable* GetGainTable() { return GainTable; }
+
+	//*************************************************************************
+	// These are overrides from VirtualInstrument that deal with MIDI message
+	// handling...
+	virtual void  NoteOn(FMidiVoiceId InVoiceId, int8 InMidiNoteNumber, int8 InVelocity, int8 InMidiChannel = 0, int32 InEventTick = 0, int32 InCurrentTick = 0, float InOffsetMs = 0.0f) override;
+	virtual void  NoteOnWithFrameOffset(FMidiVoiceId InVoiceId, int8 InMidiNoteNumber, int8 InVelocity, int8 InMidiChannel = 0, int32 InNumFrames = 0) override;
+	virtual bool  NoteIsOn(int8 InMidiNoteNumber, int8 InMidiChannel = 0) override;
+	virtual void  NoteOff(FMidiVoiceId InVoiceId, int8 InMidiNoteNumber, int8 InMidiChannel = 0) override;
+	virtual void  NoteOffWithFrameOffset(FMidiVoiceId InVoiceId, int8 InMidiNoteNumber, int8 InMidiChannel = 0, int32 InNumFrames = 0) override;
+	virtual void  SetExtraPitchBend(float Semitones, int8 InMidiChannel = 0) override;
+	virtual float GetPitchBend(int8 InMidiChannel = 0) const override;
+	virtual void  SetPitchBend(float Value, int8 InMidiChannel = 0) override;
+	virtual void  GetController(MidiConstants::EControllerID InController, int8& OutMsb, int8& OutLsb, int8 InMidiChannel = 0) const override;
+
+	// Mix Volume (dB) and associated linear gain (0,1]
+	virtual void  SetMidiChannelVolume(float InVolume, float InSeconds = 0.0f, int8 InMidiChannel = 0) override;
+	virtual float GetMidiChannelVolume(int8 InMidiChannel = 0) const override;
+	virtual void  SetMidiChannelGain(float InGain, float InSeconds = 0.0f, int8 InMidiChannel = 0) override;
+	virtual float GetMidiChannelGain(int8 InMidiChannel = 0) const override;
+	virtual void  SetMidiChannelMute(bool InMute, int8 InMidiChannel = 0)  override;
+	virtual bool  GetMidiChannelMute(int8 InMidiChannel = 0) const override { return MidiChannelMuted; }
+
+	virtual int32 GetMaxNumVoices() const override;
+	virtual int32 GetNumVoicesInUse() const override { return (int32)(ActiveVoices.Num()); }
+	virtual int32 GetNumVoicesInUse(FFusionVoice** vArray) const;
+	virtual void  KillAllVoices() override;
+	virtual void  AllNotesOff() override;
+	virtual void  AllNotesOffWithFrameOffset(int32 InNumFrames = 0) override;
+
+	// Sets modulation and sample playback speed when maintainPitch is false. Otherwise it's 1.0.
+	virtual void  SetSpeed(float Speed, bool MaintainPitch = false) override;                
+	
+	// Gets the current speed
+	virtual float GetSpeed(bool* MaintainPitchOut = nullptr) override                           
+	{
+		if (MaintainPitchOut != nullptr)
+		{
+			*MaintainPitchOut = MaintainPitchWhenSpeedChanges;  
+		}
+		return Speed;
+	}
+
+	virtual void SetTempo(float BPM) override;
+	virtual void SetBeat(float Beat) override;
+	float GetBeat() const { return CurrentBeat; }
+
+	float GetSubtsreamGain(int32 Index)
+	{
+		if (Index >= kMaxSubstreams)
+		{
+			return 1.0f;
+		}
+		
+		return SubstreamGain[Index];
+	}
+
+	// "child" will be processed manually by us, with its output
+	//  mixed into our signal chain.
+	// TODO: Commented out for now. Needs to be reimplemented without HarmonixGeneratorHandle
+	// virtual bool AddChild(const FHarmonixGeneratorHandle& child, EInstrumentRenderMode renderMode);
+	// virtual bool RemoveChild(const FHarmonixGeneratorHandle& child);
+
+	virtual bool IsChildRenderer() const { return false; }
+
+	virtual void SetSampleRate(float InSampleRateHz) override;
+
+	float GetSubstreamGain(int32 Index) const
+	{
+		if (Index >= kMaxSubstreams)
+		{
+			return 1.0f;
+		}
+		return SubstreamGain[Index];
+	}
+
+#ifdef UE_BUILD_DEVELOPMENT
+	TSet<int32> GetActiveKeyzones();
+#endif
+
+	virtual void SetRawTransposition(int32 SemiTones) override { Transposition = SemiTones; }
+	virtual int32 GetRawTransposition() const override { return Transposition; }
+	virtual void SetRawPitchMultiplier(float RawPitch) override { RawPitchMultiplier = RawPitch; }
+	virtual float GetRawPitchMultiplier() const override { return RawPitchMultiplier; }
+
+	void SetTicksPerQuarterNote(int32 InTicksPerQuarterNote) { TicksPerQuarterNote = InTicksPerQuarterNote; }
+	int32 GetTicksPerQuarterNote() const { return TicksPerQuarterNote; }
+
+protected:
+
+	//*************************************************************************
+	// These are overrides from VirtualInstrument
+	virtual void  ResetInstrumentStateImpl() override;
+	virtual void  ResetMidiStateImpl() override;
+	virtual void  Set7BitControllerImpl(MidiConstants::EControllerID InController, int8 Value, int8 InMidiChannel = 0) override;
+	virtual void  Set14BitControllerImpl(MidiConstants::EControllerID InController, int16 Value, int8 InMidiChannel = 0) override;
+	
+	// NOTE: This next function's InValue is NOT some standard midi value. By the time it is called it is assumed 
+	// the units and range of InValue are correct for the given controller id. So, for example, if InController 
+	// is MidiConstants::LFO0Frequency than InValue is treated as Hz!
+	void SetController(MidiConstants::EControllerID InController, float InValue);
+
+	void SetVoicePool(FSharedFusionVoicePoolPtr InPool, bool NoCallbacks);
+	FSharedFusionVoicePoolPtr GetVoicePool() const { return VoicePool; }
+	void VoicePoolWillDestruct(const FFusionVoicePool* InPool);
+
+	virtual void Process(uint32 InSliceIndex, uint32 InSubSliceIndex, TAudioBuffer<float>& OutBuffer) override;
+	bool CanProcessFromWorkerThread() const override { return true; }
+
+	void SetPatch(FFusionPatchData* PatchData);
+
+	FString GetPatchPath() const;
+
+	void AddChild(FVirtualInstrument* child, EInstrumentRenderMode renderMode);
+	bool RemoveChild(FVirtualInstrument* child);
+	void DumpAllChildren();
+
+	// apply all preset-specific settings from the patch.
+	// call whenever the preset or the patch change.
+	void ApplyPatchPresetSettings();
+
+	void  SetFineTuneCents(float cents);
+	float GetFineTuneCents() const;
+
+private:
+	friend class FFusionVoice;
+	friend class FFusionVoicePool;
+	friend class FFusionVstEffect;
+
+
+
+	void            SetPortamentoMode(EPortamentoMode InMode);
+	EPortamentoMode GetPortamentoMode() const;
+	void            SetIsPortamentoEnabled(bool InEnabled);
+	bool            GetIsPortamentoEnabled() const { return IsPortamentoEnabled; }
+	void            SetPortamentoTime(float InTimeSec);
+	// returns time in seconds
+	float           GetPortamentoTime() const;
+	float           GetCurrentPortamentoPitch() const;
+
+	void  SetStartPointMs(float InStartPointOffsetMs);
+	float GetStartPointMs() const;
+
+	// This is a soft limit, so old voices will be released instead of killed
+	void SetMaxNumVoices(int32); 
+	bool RelinquishVoice(FFusionVoice* InVoice);
+
+	virtual bool ProcessCallWillProduceSilence() const override { return ActiveVoices.Num() == 0; }
+
+	void  SetPan(const FPannerDetails& pan);
+	const FPannerDetails& GetPan() const;
+
+	//float RandomStartPointMs = 0.0f;
+	//float RandomPitchCents = 0.0f;
+
+	float GetMinPitchBendCents() const;
+	void  SetMinPitchBendCents(float Value);
+	float GetMaxPitchBendCents() const;
+	void  SetMaxPitchBendCents(float Value);
+
+	// returns the multiplicative factor used to change the rate of sample frame increment
+	// to hit the current pitch bend settings
+	float GetPitchBendFactor() const;
+
+	float GetMidiChannelMuteGain() const { return MidiChannelMuteGainRamper.GetCurrent(); }
+
+	float GetTrimVolume() const;
+	void  SetTrimVolume(float Value);
+	float GetTrimGain() const;
+
+	void SetDelayEnabled(bool enabled);
+	void SetBitCrusherEnabled(bool enabled);
+
+	float GetRampedExpression() const { return ExpressionGainRamper.GetCurrent(); }
+	float GetRampedPitchBend() const { return PitchBendRamper.GetCurrent(); }
+
+	void  SetMidiExpressionGain(float InGain, float InSeconds = 0.0f, int8 InMidiChannel = 0);
+
+	// Updates time-based state.
+	void PrepareToProcess(uint32 numSamples);
+
+	void UpdatePitchBendFactor();
+	void UpdateVoiceLfos();
+
+	// Pending note actions...
+	void ProcessNoteActions(int32 InNumFrames);
+	void ResetNoteActions(bool ClearNotes);
+	void ResetNoteStatus();
+
+	// internal note on/off functions
+	// only called on audio thread
+	int32 GatherMatchingKeyzones(uint8 TransposedNote, uint8 InVelocity, bool IsNoteOn, const TArray<FKeyzoneSettings>& InKeyzones, uint16* InMatchingZones);
+	int32 StartNote(FMidiVoiceId InVoiceId, uint8 InMidiNoteNumber, uint8 InVelocity, bool IsNoteOn, int32 InEventTick, int32 InTriggerTick, float InOffsetMs);
+	// Note: StopNote returns a TArray of midi notes that where stopped so the called can check to see if there are any key-off keyzoned that need to be triggered
+	TArray<uint8> StopNote(FMidiVoiceId InVoiceId, uint8 LastVelocity);
+
+	bool TryKeyOnZone(FMidiVoiceId InVoiceId, uint8 InTriggeredNote, uint8 InTransposedNote, uint8 InVelocity, int32 InEventTick, int32 InTriggerTick, float InOffsetMs, const FKeyzoneSettings* keyzone);
+
+private:
+
+	float RampCallRateHz = 375.0f;
+
+	float Speed = 1.0f;
+	bool  MaintainPitchWhenSpeedChanges = false;
+	float CurrentBeat = 0.0f;
+
+	//---------------------------------------------
+	// Gain Settings
+
+	// channel's mix setting... persists independently of the patch
+	float MidiChannelVolume = 0.0f;
+	TLinearRamper<float> MidiChannelGainRamper;
+
+	// channel's mute setting... persists independently of the patch
+	bool MidiChannelMuted = false;
+
+	// ramper to smooth muting
+	TLinearRamper<float> MidiChannelMuteGainRamper;
+
+	// the patch's trim setting... read-only outside of patch editor.
+	float TrimVolume = 0.0f;
+	float TrimGain = 1.0f;
+
+	//the CC 11 expression value [0,1]
+	float ExpressionGain = 1.0f;
+	
+	// expression [0,1] mapped to linear gain (a la the MIDI recommendations)
+	TLinearRamper<float> ExpressionGainRamper;
+
+	FPannerDetails PanSettings;
+
+	//----------------------------------------------
+	// PORTAMENTO
+	bool IsPortamentoEnabled = false;
+
+	EPortamentoMode PortamentoMode = EPortamentoMode::None;
+	float PortamentoTimeMs = 300.0f;
+
+	// True if the portmanto pitch should be considered (false before any note ons for example)
+	bool IsPortamentoActive = false;
+	
+	// midi note number (fractional) that represents the actual pitch of sounding notes
+	TLinearRamper<float> PortamentoPitchRamper;
+
+	//---------------------------------------------
+	// PITCH BEND
+
+	// on range [-1, 1]
+	TLinearRamper<float> PitchBendRamper;
+	
+	// extra pitch bend in semitones
+	float ExtraPitchBend = 0.0f;
+	float PitchBendFactor = 1.0f;
+
+	float FineTuneCents = 0.0f;
+
+	float StartPointMs = 0.0f;
+
+	FBiquadFilterSettings FilterSettings;
+
+
+	//---------------------------------------------
+	// Modulators
+	FAdsrSettingsArray AdsrSettings;
+
+	FLfoSettingsArray LfoSettings;
+
+	Harmonix::Dsp::Modulators::FLfoArray Lfos;
+
+	Harmonix::Dsp::Modulators::FModulatorArray Randomizers;
+	Harmonix::Dsp::Modulators::FModulatorArray VelocityModulators;
+
+	Harmonix::Dsp::Modulators::FModulator Randomizer[(uint8)EModulatorTarget::Num];
+	Harmonix::Dsp::Modulators::FModulator VelocityModulator[(uint8)EModulatorTarget::Num];
+
+	// soft limit, may be over max if some voices are in release stage
+	uint32 MaxNumVoices = 100;
+
+	TDoubleLinkedList<FFusionVoice*> ActiveVoices;
+
+	float MinPitchBendCents = -200.0f;
+	float MaxPitchBendCents = 200.0f;
+
+	FSharedFusionVoicePoolPtr VoicePool;
+
+	Harmonix::Dsp::Effects::FDelay Delay;
+	bool DelayEnabled = false;
+	double DelaySilenceCountdown = 0.0;
+
+	Harmonix::Dsp::Effects::FBitCrusher BitCrusher;
+	bool BitCrusherEnabled = false;
+
+	Harmonix::Dsp::Effects::FDistortionV1 Distortion;
+	bool DistortionEnabled = false;
+
+	Harmonix::Dsp::Effects::FVocoder Vocoder;
+	bool VocoderEnabled = false;
+	int32 VocoderModulatorIndex = 0;
+
+	// if we get note-ons or note-offs on the main
+	// thread, or if we get multiple requests for
+	// the same note number, then this data helps
+	// us track and filter these requests
+	FCriticalSection sNoteActionCritSec;
+	FCriticalSection sNoteStatusCritSec;
+	static const int8 kNoteIgnore = -1;
+	static const int8 kNoteOff = 0;
+	static const int32 kMaxLayersPerNote = 128;
+
+	struct FPendingNoteAction
+	{
+		int8  MidiNote = 0;
+		int8  Velocity = 0;
+		int32 EventTick = 0;
+		int32 TriggerTick = 0;
+		float OffsetMs = 0.0f;
+		int32 FrameOffset = 0;
+		FMidiVoiceId VoiceId;
+	};
+
+	struct FMIDINoteStatus
+	{
+		// is the key pressed down?
+		bool KeyedOn = false;
+
+		// is there any sound coming out of this note? (release could mean key off but voices active)
+		int32 NumActiveVoices = 0;
+	};
+
+	TArray<FPendingNoteAction> PendingNoteActions;
+	FMIDINoteStatus NoteStatus[MidiConstants::kMaxNumNotes];
+
+	struct FFusionPatchData* FusionPatchData = nullptr;
+
+	int8 LastStartLayerSelect[MidiConstants::kMaxNumNotes];
+	int8 LastStopLayerSelect[MidiConstants::kMaxNumNotes];
+	int8 LastVelocity[MidiConstants::kMaxNumNotes];
+
+	// "scratch pad" used during process to get each voice's output audio
+	static const int32 kScratchBufferFrames = 2048;
+	static const int32 kScratchBufferChannels = FAudioBufferConfig::kMaxAudioBufferChannels;
+	static const int32 kScratchBufferSamples = (kScratchBufferFrames * kScratchBufferChannels);
+	static const int32 kScratchBufferBytes = (kScratchBufferSamples * sizeof(float));
+
+	// +4 because we can't count on 'alignas' because we allocate an array of these classes. That type of allocation ignores the alignas attribute :-(
+	float VoiceWorkBuffer[kScratchBufferSamples + 4]; 
+	float* VoiceWorkBufferChannels[kScratchBufferChannels];
+
+protected:
+
+	// TODO: Replace these with FVirtualInstrument
+	TArray<FVirtualInstrument*> FusionPreChildren;
+	TArray<FVirtualInstrument*> FusionPostChildren;
+
+protected:
+
+	// we'll keep track of the 'current tempo' as it is
+	// needed by any 'beat sync' effects...
+	float CurrentTempoBPM = 120.0f;
+
+	float RawPitchMultiplier = 0.0f;
+	int32 Transposition = 0;
+
+	int32 PresetIdx = 0;
+
+	EKeyzoneSelectMode KeyzoneSelectMode = EKeyzoneSelectMode::Layers;
+
+private:
+	int16 TimeStretchEnvelopeOverride = -1;
+	void UpdateVoicesForEnvelopeOrderChange();
+
+	static const int32 kMaxSubstreams = 8;
+	float SubstreamGain[kMaxSubstreams];
+	void SetSubstreamMidiGain(int32 InIndex, uint8 InMidiGain);
+
+	int32 TicksPerQuarterNote = MidiConstants::kTicksPerQuarterNoteInt;
+
+	const FGainTable* GainTable = nullptr;
+};
diff --git a/Engine/Plugins/Runtime/Harmonix/Source/HarmonixDsp/Public/HarmonixDsp/FusionSampler/FusionVoicePool.h b/Engine/Plugins/Runtime/Harmonix/Source/HarmonixDsp/Public/HarmonixDsp/FusionSampler/FusionVoicePool.h
new file mode 100644
--- /dev/null
+++ b/Engine/Plugins/Runtime/Harmonix/Source/HarmonixDsp/Public/HarmonixDsp/FusionSampler/FusionVoicePool.h
@@ -0,0 +1,205 @@
+// Copyright Epic Games, Inc. All Rights Reserved.
+#pragma once
+
+#include "HAL/Platform.h"
+#include "Templates/SharedPointer.h"
+#include "HarmonixDsp/StretcherAndPitchShifterFactory.h"
+#include "HarmonixMidi/MidiVoiceId.h"
+
+struct FKeyzoneSettings;
+class FFusionSampler;
+class FFusionVoice;
+class IStretcherAndPitchShifter;
+class FFusionVoicePool;
+using FSharedFusionVoicePoolPtr = TSharedPtr<FFusionVoicePool, ESPMode::ThreadSafe>;
+
+DECLARE_LOG_CATEGORY_EXTERN(LogFusionVoicePool, Log, All);
+
+struct HARMONIXDSP_API FFusionVoiceConfig
+{
+	float SampleRate = 48000.0f;
+
+	uint32 NumTotalVoices = 0;
+	uint32 SoftVoiceLimit = 24;
+	bool   DecompressSamplesOnLoad = true;
+	bool   DecompressAsync = true;
+	int32  MaxActiveDecoders = -1;
+	int32  AsynchronousDecodePrebufferFrames = 48000;
+	bool   PreserveCompressedData = false;
+	float  FormantDbCorrectionPerHalfStepUp = 0.0f;
+	float  FormantDbCorrectionPerHalfStepDown = 0.0f;
+	float  FormantDbCorrectionMaxUp = 12.0f;
+	float  FormantDbCorrectionMaxDown = -12.0f;
+
+	FFusionVoiceConfig(
+		uint32 InNumTotalVoices,
+		uint32 InSoftVoiceLimit,
+		bool InDecompressSamplesOnLoad)
+		: NumTotalVoices(InNumTotalVoices)
+		, SoftVoiceLimit(InSoftVoiceLimit)
+	{}
+
+	static FFusionVoiceConfig DefaultConfig;
+};
+
+class HARMONIXDSP_API FFusionVoicePool
+{
+public:
+
+	static FSharedFusionVoicePoolPtr GetDefault(float InSampleRate);
+	static FSharedFusionVoicePoolPtr Create(const FFusionVoiceConfig& InConfig);
+
+	FFusionVoicePool(float InSampleRate) 
+		: Voices(nullptr)
+		, NumAllocatedVoices(0)
+		, NumVoicesSetting(kDefaultPoolSize)
+		, SoftVoiceLimit(kDefaultPoolSize)
+		, PeakVoiceUsage(0)
+		, DecompressSamplesOnLoad(false)
+		, SampleRate(InSampleRate)
+	{};
+
+	virtual ~FFusionVoicePool();
+
+	void HardAllocatateVoicesAndShifters();
+	void ReleaseHardAllocation();
+
+	void SetIsMultithreading(bool InIsMultithreaded);
+
+	// let the FusionVoicePool know that you will need voices
+	void AddClient(FFusionSampler* InSampler);
+
+	// let the FusionVoicePool know that you no longer need voices
+	void RemoveClient(FFusionSampler* InSampler);
+
+	// Global Voice Pool
+	static const uint32 kMinPoolSize = 1;
+	static const uint32 kMaxPoolSize = 256;
+	static const uint32 kDefaultPoolSize = 16;
+
+
+	uint32 GetNumVoicesInUse();
+
+	/**
+	 * The number of voices allocated.
+	 * May be different than the hard limit if voices have not been allocated yet
+	 * (or if there are no clients for this voice pool)
+	 * @returns the number of voices currently allocated.
+	 */
+	uint32 GetNumVoicesAvailable() { return NumAllocatedVoices; }
+
+	bool HasVoice(FFusionSampler* InOwner, FMidiVoiceId InVoiceId);
+
+	/**
+	 * The number of voices to allocate for this pool.
+	 * @param limit the maximum number of voices that this pool can handle. (the number of voices to allocate)
+	 */
+	void SetHardVoiceLimit(uint32 InLimit);
+
+	/**
+	 * @returns the maximum number of voices this pool can handle.
+	 */
+	uint32 GetHardVoiceLimit() const { return NumVoicesSetting; }
+
+	/**
+	 * The number of voices to allow before automatically releasing excess voices.
+	 * @param limit the maximum number of voices to allow before fast-releasing
+	 */
+	void SetSoftVoiceLimit(uint32 InLimit);
+
+	/**
+	 * @returns the maximum number of voices this pool will allow before automatically releasing voices.
+	 */
+	uint32 GetSoftVoiceLimit() const { return SoftVoiceLimit; }
+
+	// if we're decompressing samples at load time, we don't need to set up per-voice mogg decoders 
+	//  (which are RAM-intensive)
+	void SetDecompressSamplesOnLoad(bool Decompress);
+
+	void SetFormantVolumeCorrection(float dBperHalfStepUp, float dBperHalfStepDown, float dBMaxUp, float  dBMaxDown);
+
+	uint32 GetPeakVoiceUsage() 
+	{ 
+		GetNumVoicesInUse(); 
+		return PeakVoiceUsage; 
+	}
+	
+	void ResetPeakVoiceUsage() 
+	{ 
+		PeakVoiceUsage = 0; 
+		GetNumVoicesInUse(); 
+	}
+
+	// THIS PROBABLY SHOULD BE PRIVATE. IT IS USED BY NOTE ON.
+	// NOT AN IDEAL ARGUMENT SIGNATURE FOR PUBLIC CONSUMPTION.
+	// pass in a channel and note id to assign to the voice.
+	// also pass in the keyzone that the voice should use.
+	// any active voices with a matching set of IDs will be put into release stage.
+	// a voice might get killed (stopped instantly) if there are not enough free voices.
+	// returns null if there is no patch assigned to the channel.
+	FFusionVoice* GetFreeVoice(
+		FFusionSampler* InSampler, 
+		FMidiVoiceId InVoiceID, 
+		const FKeyzoneSettings* InKeyzone, 
+		TFunction<bool(FFusionVoice*)> Handler,
+		bool AllowAlias = true, 
+		bool IsRendererForAlias = false);
+
+	void ReleaseShifter(TSharedPtr<IStretcherAndPitchShifter, ESPMode::ThreadSafe> InShifter);
+	/**
+	 * Fast releases voices that have exceeded the soft limit for the pool (or the channel).
+	 * If no channel is specified, then voices assigned to any channel are candidates for
+	 * being fast-released. If a specific channel is specified, then only voices being used by that
+	 * channel may be fast released.
+	 * @param channel (optional) restrict the voice candidates for release to the ones uses by a specific channel.
+	 * @returns the number of voices currently being used by the pool (or by the specified channel)
+	 */
+	uint32 FastReleaseExcessVoices(FFusionSampler* InSampler = nullptr);
+
+	// instantly stops all voices with no release stage
+	void KillVoices();
+	//void KillVoices(const MultiInstrument*);
+	void KillVoices(const FFusionSampler* InSampler, bool NoCallbacks);
+	void KillVoices(const FKeyzoneSettings* InKeyzoneSettings);
+
+	FFusionVoice* GetVoice(uint32 VoiceIdx);
+
+	float GetSampleRate() const { return SampleRate; }
+	void SetSampleRate(float InSampleRate);
+
+	void Lock();
+	void Unlock();
+
+private:
+
+	// keep a weak ptr, and pass around shared ptrs
+	// so it gets automatically destroyed when nothing references it anymore
+	using FPoolMap = TMap<int32, TWeakPtr<FFusionVoicePool, ESPMode::ThreadSafe>>;
+	static FPoolMap gDefaultVoicePools;
+
+	FFusionVoice* Voices;
+
+	uint32     NumAllocatedVoices = 0;
+	uint32     NumVoicesSetting = 0;
+
+	uint32     SoftVoiceLimit = 0;
+
+	uint32     PeakVoiceUsage = 0;
+
+	bool       DecompressSamplesOnLoad = false;
+
+	void AllocVoicesAndShifters();
+	void CreateVoices(uint16 InMaxPolyphony);
+	void CreateShifters();
+	void FreeVoicesAndShifters();
+
+	TArray<FFusionSampler*> ClientSamplers;
+
+	FCriticalSection PoolLock;
+
+	float SampleRate = 0.0f;
+
+	bool DynamicAllocAndFree = true;
+
+	bool IsMultithreading = false;
+};
diff --git a/Engine/Plugins/Runtime/Harmonix/Source/HarmonixDsp/Public/HarmonixDsp/GainTable.h b/Engine/Plugins/Runtime/Harmonix/Source/HarmonixDsp/Public/HarmonixDsp/GainTable.h
new file mode 100644
--- /dev/null
+++ b/Engine/Plugins/Runtime/Harmonix/Source/HarmonixDsp/Public/HarmonixDsp/GainTable.h
@@ -0,0 +1,57 @@
+// Copyright Epic Games, Inc. All Rights Reserved.
+#pragma once
+
+#include "HarmonixDsp/AudioBufferConstants.h"
+
+#include "Logging/LogMacros.h"
+#include "Math/VectorRegister.h"
+#include "HAL/Platform.h"
+
+DECLARE_LOG_CATEGORY_EXTERN(LogGainTable, Log, All);
+
+// anonymous namespace to hide this from externals
+namespace Harmonix::Dsp::GainTable
+{
+	static const int32 kMaxSpeakers = 8;
+}
+
+using FChannelGains = MS_ALIGN(16) union 
+{ 
+	VectorRegister4Float simd[Harmonix::Dsp::GainTable::kMaxSpeakers / 4]; 
+	float f[Harmonix::Dsp::GainTable::kMaxSpeakers]; 
+} GCC_ALIGN(16);
+
+class HARMONIXDSP_API FGainTable
+{
+public:
+	static void Init(EAudioBufferChannelLayout ChannelLayout = EAudioBufferChannelLayout::Stereo);
+	static void SetupPrimaryGainTable(EAudioBufferChannelLayout ChannelLayout);
+
+	static FORCEINLINE FGainTable& Get();
+
+	FGainTable();
+
+	void SetChannelLayout(EAudioBufferChannelLayout ChannelLayout);
+
+	const FChannelGains& GetGains(float PolarAngle) const;
+	void PanSample(float InSample, float PolarAngle, FChannelGains& OutGains) const;
+
+	FORCEINLINE float GetDirectChannelAzimuthInCurrentLayout(ESpeakerChannelAssignment ChannelAssignment) const;
+
+	FORCEINLINE void GetGainsForDirectAssignment(ESpeakerChannelAssignment ChannelAssignment, FChannelGains& OutGains) const;
+
+private:
+	EAudioBufferChannelLayout CurrentLayout = EAudioBufferChannelLayout::UnsupportedFormat;
+	int32 SpeakerCount = 0;
+	uint32 SpeakerMask = 0;
+
+	static const uint32 kGainTableSize = 1024;
+	FChannelGains Entries[kGainTableSize];
+
+	void BuildPanEntriesFromPannableSpeakerAzimuths(int32 NumSpeakers, const float* SpeakerAzimuths, const uint8* InRemap);
+	bool CurrentLayoutHasSpeaker(ESpeakerChannelAssignment ChannelAssignment) const;
+
+	static FGainTable* gGainTable;
+};
+
+#include "HarmonixDsp/GainTable.inl"
\ No newline at end of file
diff --git a/Engine/Plugins/Runtime/Harmonix/Source/HarmonixDsp/Public/HarmonixDsp/Modulators/Settings/ModulatorSettings.h b/Engine/Plugins/Runtime/Harmonix/Source/HarmonixDsp/Public/HarmonixDsp/Modulators/Settings/ModulatorSettings.h
new file mode 100644
--- /dev/null
+++ b/Engine/Plugins/Runtime/Harmonix/Source/HarmonixDsp/Public/HarmonixDsp/Modulators/Settings/ModulatorSettings.h
@@ -0,0 +1,63 @@
+// Copyright Epic Games, Inc. All Rights Reserved.
+#pragma once
+
+#include "HarmonixDsp/Containers/EnumArray.h"
+#include "Misc/EnumRange.h"
+
+#include "ModulatorSettings.generated.h"
+
+UENUM(BlueprintType)
+enum class EModulatorTarget : uint8
+{
+	StartPoint	UMETA(Json = "Start Point"),
+	Pitch		UMETA(Json = "Pitch"),
+	Num			UMETA(Hidden),
+	None		UMETA(Hidden, Json = "None")
+};
+
+USTRUCT(BlueprintType)
+struct HARMONIXDSP_API FModulatorSettings
+{
+	GENERATED_BODY()
+
+public:
+
+	FModulatorSettings();
+
+	UPROPERTY(EditDefaultsOnly, Category="Harmonix|DSP")
+	EModulatorTarget Target = EModulatorTarget::None;
+
+	UPROPERTY(EditDefaultsOnly, Category = "Harmonix|DSP", Meta = (UIMin = "0.0", UIMax = "1000000.0", ClampMin = "0.0", ClampMax = "1000000.0"))
+	float Range = 1.0f;
+
+	UPROPERTY(EditDefaultsOnly, Category = "Harmonix|DSP", Meta = (UIMin = "0.0", UIMax = "1.0", ClampMin = "0.0", ClampMax = "1.0"))
+	float Depth = 0.5f;
+
+};
+
+enum class EModulatorIndex : uint8
+{
+	StartPoint,
+	Pitch,
+	Num
+};
+
+ENUM_RANGE_BY_COUNT(EModulatorIndex, EModulatorIndex::Num);
+
+USTRUCT(BlueprintType)
+struct FModulatorSettingsArray
+{
+	GENERATED_BODY()
+
+public:
+
+	IMPLEMENT_ENUM_ARRAY(Array, EModulatorIndex, FModulatorSettings)
+
+	ENUM_PROPERTY(StartPoint);
+	ENUM_PROPERTY(Pitch);
+
+private:
+
+	UPROPERTY(EditDefaultsOnly, Category="Settings")
+	FModulatorSettings Array[(uint8)EModulatorIndex::Num];
+};
\ No newline at end of file
diff --git a/Engine/Plugins/Runtime/Harmonix/Source/HarmonixDsp/Public/HarmonixDsp/Streaming/TrackChannelInfo.h b/Engine/Plugins/Runtime/Harmonix/Source/HarmonixDsp/Public/HarmonixDsp/Streaming/TrackChannelInfo.h
new file mode 100644
--- /dev/null
+++ b/Engine/Plugins/Runtime/Harmonix/Source/HarmonixDsp/Public/HarmonixDsp/Streaming/TrackChannelInfo.h
@@ -0,0 +1,80 @@
+// Copyright Epic Games, Inc. All Rights Reserved.
+#pragma once
+
+#include "HarmonixDsp/PannerDetails.h"
+#include "HarmonixDsp/Streaming/StandardStream.h"
+
+#include "TrackChannelInfo.generated.h"
+
+USTRUCT(BlueprintType)
+struct HARMONIXDSP_API FStreamingChannelParamsArray
+{
+	GENERATED_BODY()
+
+public:
+
+	UPROPERTY(EditDefaultsOnly, Category = "Settings")
+	TArray<FStreamingChannelParams> ChannelParams;
+
+	FStreamingChannelParamsArray() {};
+	FStreamingChannelParamsArray(const TArray<FStreamingChannelParams>& InChannelParams)
+		: ChannelParams(InChannelParams)
+	{
+	}
+
+	FStreamingChannelParams& operator[](int32 Index) { return ChannelParams[Index]; }
+	const FStreamingChannelParams& operator[](int32 Index) const { return ChannelParams[Index]; }
+};
+
+USTRUCT(BlueprintType)
+struct HARMONIXDSP_API FTrackChannelInfo
+{
+	GENERATED_BODY()
+
+public:
+
+	UPROPERTY(EditDefaultsOnly, Category = "Settings", Meta = (UIMin = -1, ClampMin = -1))
+	int32 RealTrackIndex = -1;
+
+	UPROPERTY(EditDefaultsOnly, Category = "Settings")
+	FName Name;
+
+	UPROPERTY(EditDefaultsOnly, Category = "Settings")
+	FName Routing;
+
+	UPROPERTY(EditDefaultsOnly, Category = "Settings")
+	TArray<FStreamingChannelParams> Channels;
+
+	UPROPERTY(EditDefaultsOnly, Category = "Settings")
+	TMap<FName, FStreamingChannelParamsArray> PresetChannels;
+
+	FTrackChannelInfo();
+	FTrackChannelInfo(int32 InRealIndex);
+	FTrackChannelInfo(FTrackChannelInfo&& Other);
+	FTrackChannelInfo(const FTrackChannelInfo& Other);
+
+	FTrackChannelInfo& operator=(const FTrackChannelInfo& Other)
+	{
+		RealTrackIndex = Other.RealTrackIndex;
+		Name = Other.Name;
+		Routing = Other.Routing;
+		Channels = Other.Channels;
+		PresetChannels = Other.PresetChannels;
+		return *this;
+	}
+
+	int32  GetNumChannels() const;
+	int32  GetStreamChannelIndex(int32 ChannelIndex) const;
+	bool GetStreamIndexesPan(int32 StreamIndex, FPannerDetails& OutPan) const;
+	bool GetStreamIndexesGain(int32 StreamIndex, float& OutGain) const;
+	bool GetStreamIndexesPresetGain(FName PresetName, int32 StreamIndex, float& OutGain, bool DefaultIfNotFound = true) const;
+	bool SetStreamIndexesPan(int32 StreamIndex, float Pan);
+	bool SetStreamIndexesPan(int32 StreamIndex, const FPannerDetails& InPan);
+	bool SetStreamIndexesGain(int32 StreamIndex, float InGain);
+	bool SetStreamIndexesPresetGain(FName presetName, int32 StreamIndex, float InGain);
+
+private:
+
+	const FStreamingChannelParams* FindStreamingChannelParams(int32 StreamIndex) const;
+	FStreamingChannelParams* FindStreamingChannelParams(int32 StreamIndex);
+};
\ No newline at end of file
diff --git a/Engine/Plugins/Runtime/Harmonix/Source/HarmonixDsp/Public/HarmonixDsp/StretcherAndPitchShifterConfig.h b/Engine/Plugins/Runtime/Harmonix/Source/HarmonixDsp/Public/HarmonixDsp/StretcherAndPitchShifterConfig.h
new file mode 100644
--- /dev/null
+++ b/Engine/Plugins/Runtime/Harmonix/Source/HarmonixDsp/Public/HarmonixDsp/StretcherAndPitchShifterConfig.h
@@ -0,0 +1,14 @@
+// Copyright Epic Games, Inc. All Rights Reserved.
+
+#pragma once
+
+#include "StretcherAndPitchShifterConfig.generated.h"
+
+// base class for all Stretcher and Pitch Shifter config settings
+// allows configuration to appear in the "Pitch Shifter Settings"
+// along with any other pitch shifter settings
+UCLASS(Abstract, EditInlineNew, CollapseCategories, Config = Engine, defaultconfig)
+class HARMONIXDSP_API UStretcherAndPitchShifterConfig : public UObject
+{
+	GENERATED_BODY()
+};
diff --git a/Engine/Plugins/Runtime/Harmonix/Source/HarmonixDsp/Public/HarmonixDsp/StretcherAndPitchShifterFactoryConfig.h b/Engine/Plugins/Runtime/Harmonix/Source/HarmonixDsp/Public/HarmonixDsp/StretcherAndPitchShifterFactoryConfig.h
new file mode 100644
--- /dev/null
+++ b/Engine/Plugins/Runtime/Harmonix/Source/HarmonixDsp/Public/HarmonixDsp/StretcherAndPitchShifterFactoryConfig.h
@@ -0,0 +1,64 @@
+// Copyright Epic Games, Inc. All Rights Reserved.
+
+#pragma once
+
+#include "HarmonixDsp/PitchShifterName.h"
+#include "Harmonix/HarmonixDeveloperSettings.h"
+
+#include "StretcherAndPitchShifterFactoryConfig.generated.h"
+
+USTRUCT()
+struct FPitchShifterNameRedirect
+{
+	GENERATED_BODY()
+
+	FPitchShifterNameRedirect() 
+	{}
+	
+	FPitchShifterNameRedirect(FName OldName, FName NewName)
+		: OldName(OldName), NewName(NewName) 
+	{}
+
+	UPROPERTY(EditDefaultsOnly, Category = "Factory")
+	FName OldName;
+
+	UPROPERTY(EditDefaultsOnly, Category = "Factory")
+	FName NewName;
+};
+
+UCLASS(config = Engine, defaultconfig, meta = (DisplayName = "Pitch Shifter Settings"))
+class HARMONIXDSP_API UStretcherAndPitchShifterFactoryConfig : public UHarmonixDeveloperSettings
+{
+	GENERATED_BODY()
+
+public:
+	UStretcherAndPitchShifterFactoryConfig();
+
+	UPROPERTY(config, EditDefaultsOnly, Category = "Factory")
+	TArray<FName> FactoryPriority;
+
+	UPROPERTY(config, EditDefaultsOnly, Category = "Factory")
+	FPitchShifterName DefaultFactory;
+
+	void AddFactoryNameRedirect(FName OldName, FName NewName);
+
+	const FPitchShifterNameRedirect* FindFactoryNameRedirect(FName OldName) const;
+
+public:
+
+#if WITH_EDITORONLY_DATA
+
+	virtual void PostEditChangeProperty(FPropertyChangedEvent& PropertyChangedEvent) override;
+
+#endif
+
+	virtual void PostInitProperties() override;
+
+private:
+
+	void PruneDuplicatesAndAddMissingNames(TArray<FName>& InOutNames);
+
+	// visible in editor, but not editable
+	UPROPERTY(config, EditDefaultsOnly, Category = "Factory", Meta=(EditCondition="false"))
+	TArray<FPitchShifterNameRedirect> FactoryNameRedirects;
+};
\ No newline at end of file
diff --git a/Engine/Plugins/Runtime/Harmonix/Source/HarmonixDspEditor/Private/Dta/DtaParser.h b/Engine/Plugins/Runtime/Harmonix/Source/HarmonixDspEditor/Private/Dta/DtaParser.h
new file mode 100644
--- /dev/null
+++ b/Engine/Plugins/Runtime/Harmonix/Source/HarmonixDspEditor/Private/Dta/DtaParser.h
@@ -0,0 +1,20 @@
+// Copyright Epic Games, Inc. All Rights Reserved.
+#pragma once
+
+#include "Containers/UnrealString.h"
+
+#include "Templates/SharedPointer.h"
+
+
+DECLARE_LOG_CATEGORY_EXTERN(LogDtaParser, Log, All);
+
+
+struct FDtaParser
+{
+	static void DtaStringToJsonString(const FString& InDtaString, FString& OutJsonString);
+
+private:
+
+	static void Tokenize(const FString& InDtaString, TArray<FString>& OutTokens);
+	static void ParseToJson(const TArray<FString>& Tokens, FString& OutString);
+};
\ No newline at end of file
diff --git a/Engine/Plugins/Runtime/Harmonix/Source/HarmonixDspEditor/Private/FusionPatchAssetFactory.cpp b/Engine/Plugins/Runtime/Harmonix/Source/HarmonixDspEditor/Private/FusionPatchAssetFactory.cpp
new file mode 100644
--- /dev/null
+++ b/Engine/Plugins/Runtime/Harmonix/Source/HarmonixDspEditor/Private/FusionPatchAssetFactory.cpp
@@ -0,0 +1,261 @@
+// Copyright Epic Games, Inc. All Rights Reserved.
+
+#include "FusionPatchAssetFactory.h"
+#include "FusionPatchImportOptions.h"
+#include "EditorDialogLibrary.h"
+
+#include "HarmonixDsp/FusionSampler/FusionPatch.h"
+#include "Dta/DtaParser.h"
+
+#include "Serialization/JsonReader.h"
+#include "Serialization/JsonSerializer.h"
+
+#include "JsonImporterHelper.h"
+#include "FusionPatchJsonImporter.h"
+
+#include "Misc/MessageDialog.h"
+#include "Misc/Paths.h"
+#include "Misc/PackageName.h"
+#include "EditorFramework/AssetImportData.h"
+#include "UObject/Package.h"
+#include "HAL/FileManager.h"
+
+
+#define LOCTEXT_NAMESPACE "FusionPatchAssetFactory"
+
+DEFINE_LOG_CATEGORY(LogFusionPatchAssetFactory)
+
+UFusionPatchAssetFactory::UFusionPatchAssetFactory()
+{
+	SupportedClass = UFusionPatch::StaticClass();
+	Formats.Add(TEXT("fusion;Fusion Patch"));
+	bText = true;
+	bCreateNew = false;
+	bEditorImport = true;
+	ImportPriority = DefaultImportPriority + 20;
+}
+
+bool UFusionPatchAssetFactory::FactoryCanImport(const FString& Filename)
+{
+	return true;
+}
+
+bool UFusionPatchAssetFactory::CanReimport(UObject* Obj, TArray<FString>& OutFilenames)
+{
+	UFusionPatch* FusionPatch = Cast<UFusionPatch>(Obj);
+	if (FusionPatch && FusionPatch->AssetImportData)
+	{
+		FusionPatch->AssetImportData->ExtractFilenames(OutFilenames);
+		return true;
+	}
+	return false;
+}
+
+void UFusionPatchAssetFactory::SetReimportPaths(UObject* Obj, const TArray<FString>& NewReimportPaths)
+{
+	UFusionPatch* FusionPatch = Cast<UFusionPatch>(Obj);
+	if (FusionPatch && ensure(NewReimportPaths.Num() == 1))
+	{
+		FusionPatch->AssetImportData->UpdateFilenameOnly(NewReimportPaths[0]);
+	}
+}
+
+EReimportResult::Type UFusionPatchAssetFactory::Reimport(UObject* Obj)
+{
+	UFusionPatch* FusionPatch = Cast<UFusionPatch>(Obj);
+	if (!FusionPatch)
+	{
+		return EReimportResult::Failed;
+	}
+
+	if (!FusionPatch->AssetImportData)
+	{
+		return EReimportResult::Failed;
+	}
+
+	const FString Filename = FusionPatch->AssetImportData->GetFirstFilename();
+	if (!Filename.Len() || IFileManager::Get().FileSize(*Filename) == INDEX_NONE)
+	{
+		return EReimportResult::Failed;
+	}
+	
+	bool OutCanceled = false;
+	if (UObject* ImportedObject = ImportObject(FusionPatch->GetClass(), FusionPatch->GetOuter(), *FusionPatch->GetName(), RF_Public | RF_Standalone, Filename, nullptr, OutCanceled))
+	{
+		// The fusion patch should have just been updated, right?
+		ensure(ImportedObject == FusionPatch);
+
+		if (FusionPatch->GetOuter())
+		{
+			FusionPatch->GetOuter()->MarkPackageDirty();
+		}
+		else
+		{
+			FusionPatch->MarkPackageDirty();
+		}
+		return EReimportResult::Succeeded;
+	}
+	
+	if (OutCanceled)
+	{
+		UE_LOG(LogFusionPatchAssetFactory, Warning, TEXT("import canceled"));
+		return EReimportResult::Cancelled;
+	}
+
+	UE_LOG(LogFusionPatchAssetFactory, Warning, TEXT("import failed"));
+	return EReimportResult::Failed;
+}
+
+void UFusionPatchAssetFactory::PostImportCleanUp()
+{
+	ApplyOptionsToAllImport = false;
+	ReplaceExistingSamplesResponse = EAppReturnType::No;
+}
+
+const UFusionPatchImportOptions* UFusionPatchAssetFactory::GetImportOptions(const FString& DefaultSamplesImportPath)
+{
+	UFusionPatchImportOptions* ImportOptions = GetMutableDefault<UFusionPatchImportOptions>();
+	if (!ApplyOptionsToAllImport)
+	{
+		// prompt on the first asset and update user settings
+		if (ImportOptions->SamplesImportDir.Path.IsEmpty())
+		{
+			// otherwise, import audio samples into the destination path provided
+			ImportOptions->SamplesImportDir.Path = DefaultSamplesImportPath;
+		}
+		const FText ImportDialogTitle = NSLOCTEXT("FusionPatchImporter", "FusionPatchImportOptionsTitle", "Fusion Patch Import Options");
+		UEditorDialogLibrary::ShowObjectDetailsView(ImportDialogTitle, ImportOptions);
+		const FString DestPath = ImportOptions->SamplesImportDir.Path;
+
+		ApplyOptionsToAllImport = true;
+	}
+	return ImportOptions;
+}
+
+bool UFusionPatchAssetFactory::GetReplaceExistingSamplesResponse(const FString& InName)
+{
+	if (ReplaceExistingSamplesResponse == EAppReturnType::YesAll)
+	{
+		return true;
+	}
+
+	if (ReplaceExistingSamplesResponse == EAppReturnType::NoAll)
+	{
+		return false;
+	}
+
+	const FText ReplaceExistingTitle = NSLOCTEXT("FusionPatchImporter", "ReplaceExistingSamplesTitle", "Replace Existing Samples");
+	const FText ReplaceExistingMessage = FText::Format(NSLOCTEXT("FusionPatchImporter", "ReplaceExistingSamplesMsg", 
+		"Would you like to reimport and replace existing Sound Wave Assets in the directory with Samples referenced by this Fusion Fatch?" 
+		"\n\nPatch Name: {0}"
+		"\n\nYes. Reimport and replace existing Samples with the new samples."
+		"\n\nNo. New Samples will be imported, but existing Samples will be unchanged. The Fusion Patch will reference any existing Samples in the directory with matching names."), 
+		FText::FromString(InName));
+	ReplaceExistingSamplesResponse = UEditorDialogLibrary::ShowMessage(ReplaceExistingTitle, ReplaceExistingMessage, EAppMsgType::YesNoYesAllNoAll, ReplaceExistingSamplesResponse, EAppMsgCategory::Info);
+	
+
+	switch (ReplaceExistingSamplesResponse)
+	{
+	case EAppReturnType::Yes:
+	case EAppReturnType::YesAll:
+		return true;
+	case EAppReturnType::No:
+	case EAppReturnType::NoAll:
+		return false;
+	default:
+		ensureMsgf(false, TEXT("Unexpected response! default behavior is to NOT replace existing samples"));
+		return false;
+	}
+}
+
+void UFusionPatchAssetFactory::UpdateFusionPatchImportNotificationItem(TSharedPtr<SNotificationItem> InItem, bool bImportSuccessful, FName InName)
+{
+	if (bImportSuccessful)
+	{
+		InItem->SetCompletionState(SNotificationItem::ECompletionState::CS_Success);
+		InItem->SetText(FText::Format(NSLOCTEXT("FusionPatchImporter", "FusionPatchImportProgressNotification", "Successfully imported Fusion Patch asset: {0}"), FText::FromString(InName.ToString())));
+	}
+	else {
+		InItem->SetCompletionState(SNotificationItem::ECompletionState::CS_Fail);
+		InItem->SetText(FText::Format(NSLOCTEXT("FusionPatchImporter", "FusionPatchImportProgressNotification", "Failed to import Fusion Patch: {0}"), FText::FromString(InName.ToString())));
+	}
+	InItem->SetExpireDuration(0.2f);
+	InItem->ExpireAndFadeout();
+}
+
+UObject* UFusionPatchAssetFactory::FactoryCreateText(UClass* InClass, UObject* InParent, FName InName, EObjectFlags Flags, UObject* Context, const TCHAR* Type, const TCHAR*& Buffer, const TCHAR* BufferEnd, FFeedbackContext* Warn)
+{
+	const FString SourceFile = GetCurrentFilename();
+	AdditionalImportedObjects.Empty();
+	FString JsonString;
+	const FString DtaString(BufferEnd - Buffer, Buffer);
+
+	FDtaParser::DtaStringToJsonString(DtaString, JsonString);
+	FString ErrorMessage;
+	TSharedPtr<FJsonObject> JsonObj = FJsonImporter::ParseJsonString(JsonString, ErrorMessage);
+	bool bImportSuccessful = false;
+	if (JsonObj.IsValid())
+	{
+		UFusionPatch* FusionPatch = FindObject<UFusionPatch>(InParent, *InName.ToString());
+		if (!FusionPatch)
+		{
+			FusionPatch = NewObject<UFusionPatch>(InParent, InName, Flags);
+		}
+
+		const FString LongPackagePath = FPackageName::GetLongPackagePath(InParent->GetOutermost()->GetPathName());
+		const FString SourcePath = FPaths::GetPath(SourceFile);
+		
+		//create a notification that displays the import progress at the lower right corner
+		FNotificationInfo ImportNotificationInfo(NSLOCTEXT("FusionPatchImporter","FusionPatchImportProgressNotification","Importing Fusion Asset(s)..."));
+		ImportNotificationInfo.bFireAndForget = false;
+		TSharedPtr<SNotificationItem> ImportNotificationItem;
+		ImportNotificationItem = FSlateNotificationManager::Get().AddNotification(ImportNotificationInfo);
+
+		const UFusionPatchImportOptions* ImportOptions = GetImportOptions(LongPackagePath);
+		bool ReplaceExistingSamples = GetReplaceExistingSamplesResponse(InName.ToString());
+
+
+		// Pass import args to parser so it can import sub files
+		FFusionPatchJsonImporter::FImportArgs ImportArgs(InName, SourcePath, LongPackagePath, ImportOptions->SamplesImportDir.Path, ReplaceExistingSamples);
+		ImportArgs.SampleLoadingBehavior = ImportOptions->SampleLoadingBehavior;
+		ImportArgs.SampleCompressionType = ImportOptions->SampleCompressionType;
+
+		TArray<FString> ImportErrors;
+		if (FFusionPatchJsonImporter::TryParseJson(JsonObj, FusionPatch, ImportArgs, ImportErrors))
+		{
+			UE_LOG(LogFusionPatchAssetFactory, Log, TEXT("Successfully imported FusionPatch asset"));
+
+			if (ensure(FusionPatch->AssetImportData))
+			{
+				FusionPatch->AssetImportData->Update(SourceFile);
+			}
+
+			bImportSuccessful = true;
+			UpdateFusionPatchImportNotificationItem(ImportNotificationItem, bImportSuccessful, InName);
+			return FusionPatch;
+		}
+		else
+		{
+			const FText ImportErrorMessage = FText::Format(
+				LOCTEXT("ImportFailed_FusionPatchJsonImporter",
+					"Failed to import asset:\n'{0}'.\nReasons:\n{1}"),
+				FText::FromString(SourceFile),
+				FText::FromString(FString::Printf(TEXT("%s"), *FString::Join(ImportErrors, TEXT("\n")))));
+			FMessageDialog::Open(EAppMsgType::Ok, ImportErrorMessage);
+			UE_LOG(LogFusionPatchAssetFactory, Error, TEXT("Failed to import fusion patch: %s"), *ImportErrorMessage.ToString());
+			
+			UpdateFusionPatchImportNotificationItem(ImportNotificationItem, bImportSuccessful, InName);
+		}
+
+	}
+	else
+	{
+		const FText ImportErrorMessage = FText::Format(LOCTEXT("ImportFailed_Json", "Failed to import asset:\n'{0}'.\nFailed to read Json:\n{1}"), FText::FromString(SourceFile), FText::FromString(ErrorMessage));
+		FMessageDialog::Open(EAppMsgType::Ok, ImportErrorMessage);
+		UE_LOG(LogFusionPatchAssetFactory, Error, TEXT("Failed to read json: %s"), *ErrorMessage);
+	}
+
+	return nullptr;
+}
+
+#undef LOCTEXT_NAMESPACE
\ No newline at end of file
diff --git a/Engine/Plugins/Runtime/Harmonix/Source/HarmonixDspEditor/Private/FusionPatchAssetFactory.h b/Engine/Plugins/Runtime/Harmonix/Source/HarmonixDspEditor/Private/FusionPatchAssetFactory.h
new file mode 100644
--- /dev/null
+++ b/Engine/Plugins/Runtime/Harmonix/Source/HarmonixDspEditor/Private/FusionPatchAssetFactory.h
@@ -0,0 +1,41 @@
+// Copyright Epic Games, Inc. All Rights Reserved.
+
+#pragma once
+
+#include "Factories/Factory.h"
+#include "EditorReimportHandler.h"
+#include "Widgets/Notifications/SNotificationList.h"
+#include "Framework/Notifications/NotificationManager.h"
+
+#include "FusionPatchAssetFactory.generated.h"
+
+DECLARE_LOG_CATEGORY_EXTERN(LogFusionPatchAssetFactory, Log, All);
+
+UCLASS()
+class HARMONIXDSPEDITOR_API UFusionPatchAssetFactory : public UFactory, public FReimportHandler
+{
+	GENERATED_BODY()
+
+public:
+
+	UFusionPatchAssetFactory();
+
+	virtual bool FactoryCanImport(const FString& Filename) override;
+	virtual bool CanReimport(UObject* Obj, TArray<FString>& OutFilenames) override;
+	virtual void SetReimportPaths(UObject* Obj, const TArray<FString>& NewReimportPaths) override;;
+	virtual EReimportResult::Type Reimport(UObject* Obj) override;
+	virtual void PostImportCleanUp() override;
+
+	virtual UObject* FactoryCreateText(UClass* InClass, UObject* InParent, FName InName, EObjectFlags Flags, UObject* Context, const TCHAR* Type, const TCHAR*& Buffer, const TCHAR* BufferEnd, FFeedbackContext* Warn) override;
+
+private:
+	const class UFusionPatchImportOptions* GetImportOptions(const FString& DefaultSamplesImportPath);
+	bool GetReplaceExistingSamplesResponse(const FString& InName);
+	
+
+	bool ApplyOptionsToAllImport = false;
+
+	EAppReturnType::Type ReplaceExistingSamplesResponse = EAppReturnType::No;
+
+	void UpdateFusionPatchImportNotificationItem(TSharedPtr<SNotificationItem> InItem, bool bImportSuccessful, FName InName);
+};
\ No newline at end of file
diff --git a/Engine/Plugins/Runtime/Harmonix/Source/HarmonixDspEditor/Private/FusionPatchJsonImporter.cpp b/Engine/Plugins/Runtime/Harmonix/Source/HarmonixDspEditor/Private/FusionPatchJsonImporter.cpp
new file mode 100644
--- /dev/null
+++ b/Engine/Plugins/Runtime/Harmonix/Source/HarmonixDspEditor/Private/FusionPatchJsonImporter.cpp
@@ -0,0 +1,242 @@
+// Copyright Epic Games, Inc. All Rights Reserved.
+#include "FusionPatchJsonImporter.h"
+#include "SettingsJsonImporter.h"
+#include "JsonImporterHelper.h"
+
+#include "Serialization/JsonReader.h"
+#include "Serialization/JsonSerializer.h"
+
+#include "HarmonixDsp/FusionSampler/FusionPatch.h"
+#include "HarmonixDsp/FusionSampler/Settings/FusionPatchSettings.h"
+
+#include "Misc/Paths.h"
+#include "IAssetTools.h"
+#include "AssetToolsModule.h"
+
+#include "EditorAssetLibrary.h"
+#include "EditorFramework/AssetImportData.h"
+#include "Sound/SoundWave.h"
+
+#define LOCTEXT_NAMESPACE "FusionPatchImporter"
+
+DEFINE_LOG_CATEGORY(LogFusionPatchJsonImporter)
+
+TArray<UObject*> FFusionPatchJsonImporter::ImportAudioSamples(const TArray<FString>& Files, const FString& DestinationPath, bool ReplaceExisting)
+{
+	UAutomatedAssetImportData* AutomatedAssetImportData = NewObject<UAutomatedAssetImportData>();
+	AutomatedAssetImportData->bReplaceExisting = ReplaceExisting;
+	AutomatedAssetImportData->Filenames = Files;
+	AutomatedAssetImportData->DestinationPath = DestinationPath;
+	FAssetToolsModule& AssetToolsModule = FModuleManager::GetModuleChecked<FAssetToolsModule>("AssetTools");
+	return AssetToolsModule.Get().ImportAssetsAutomated(AutomatedAssetImportData);
+}
+
+bool FFusionPatchJsonImporter::TryParseJson(TSharedPtr<FJsonObject> JsonObj, UFusionPatch* FusionPatch, const FImportArgs& ImportArgs, TArray<FString>& OutErrors)
+{
+	if (!ensure(FusionPatch))
+		return false;
+
+	if (!ensure(JsonObj.IsValid()))
+		return false;
+
+	int Version;
+	if (!JsonObj->TryGetNumberField(TEXT("version"), Version))
+	{
+		UE_LOG(LogHarmonixJsonImporter, Warning, TEXT("Failed to read version number importing Fusion Patch."));
+		return false;
+	}
+
+	static const int kMinVersion = 1;
+	if (Version < kMinVersion)
+	{
+		UE_LOG(LogHarmonixJsonImporter, Warning, TEXT("Bad version number importing Fusion Patch! \"version\" = %d, must be >= %d)"), Version, kMinVersion);
+	};
+	
+	FFusionPatchData& FusionPatchData = FusionPatch->FusionPatchData;
+
+	// import presets and keyzones first
+	for (TSharedPtr<FJsonValue> ArrayValue : IterField(JsonObj, "presets"))
+	{
+		TSharedPtr<FJsonObject> PresetObject = nullptr;
+		if (!TryGetObject(ArrayValue, PresetObject))
+		{
+			UE_LOG(LogFusionPatchJsonImporter, Warning, TEXT("Unable to parse Preset in json"));
+			continue;
+		}
+
+		for (const auto& Pair : PresetObject->Values)
+		{
+			const FString& Key = Pair.Key;
+			TSharedPtr<FJsonValue> Value = Pair.Value;
+
+			TSharedPtr<FJsonObject> PresetJson = nullptr;
+			if (!TryGetObject(Value, PresetJson))
+			{
+				UE_LOG(LogFusionPatchJsonImporter, Warning, TEXT("Unable to parse preset in json"));
+				continue;
+			}
+
+			FFusionPatchSettings Preset;
+			if (!FSettingsJsonImporter::TryParseJson(PresetJson, Preset))
+			{
+				UE_LOG(LogFusionPatchJsonImporter, Warning, TEXT("Unable to parse preset in json"));
+				continue;
+			}
+
+			Preset.Name = FName(Key);
+			FusionPatchData.AddPreset(Preset);
+		}
+	}
+
+	FusionPatchData.CurrentPresetIndex = 0;
+
+	// collect audio sample files for import
+	TArray<FString> AudioSampleFiles;
+	bool AllFilesExist = true;
+	TMap<int, int> KeyzoneIdxToPathIdx;
+
+	int NumKeyzones = 0;
+	for (TSharedPtr<FJsonValue> ArrayValue : IterField(JsonObj, "keymap"))
+	{
+		int KeyzoneIdx = NumKeyzones;
+		TSharedPtr<FJsonObject> KeyzoneObject = nullptr;
+		if (!TryGetObject(ArrayValue, KeyzoneObject))
+			continue;
+
+		TSharedPtr<FJsonObject> KeyzoneValue = nullptr;
+		if (!TryGetObjectField(KeyzoneObject, "keyzone", KeyzoneValue))
+			continue;
+
+		// append elements to the list until the number of keyzones 
+		// is the same as the number of imported keyzones
+		while (!FusionPatchData.Keyzones.IsValidIndex(KeyzoneIdx))
+		{
+			FKeyzoneSettings Settings;
+			FusionPatchData.Keyzones.Add(Settings);
+		}
+
+		while (!KeyzoneIdxToPathIdx.Contains(KeyzoneIdx))
+		{
+			KeyzoneIdxToPathIdx.Add(KeyzoneIdx, -1);
+		}
+
+		FKeyzoneSettings& Settings = FusionPatchData.Keyzones[KeyzoneIdx];
+		NumKeyzones++;
+
+		if (!FSettingsJsonImporter::TryParseJson(KeyzoneValue, Settings))
+		{
+			UE_LOG(LogFusionPatchJsonImporter, Warning, TEXT("Unable to parse keyzone json"));
+			continue;
+		}
+
+		if (Settings.SamplePath.EndsWith(TEXT(".mogg")))
+		{
+			UE_LOG(LogFusionPatchJsonImporter, Warning, TEXT("Found mogg keyzone (%s). Will look for a .wav in the same folder."), *Settings.SamplePath);
+			Settings.SamplePath = Settings.SamplePath.Replace(TEXT(".mogg"), TEXT(".wav"));
+		}
+
+		FString ImportPath = FPaths::Combine(*ImportArgs.SourcePath, *Settings.SamplePath);
+
+		if (!FPaths::FileExists(ImportPath))
+		{
+			OutErrors.Add(FString::Printf(TEXT("Keyzone Idx = %d: Sample Path is invalid: %s"), KeyzoneIdx, *ImportPath));
+			UE_LOG(LogFusionPatchJsonImporter, Error, TEXT("Failed to import FusionPatch. Sample Path is invalid: %s"), *ImportPath);
+			AllFilesExist = false;
+			continue;
+		}
+		int32 PathIndex = AudioSampleFiles.AddUnique(ImportPath);
+		KeyzoneIdxToPathIdx[KeyzoneIdx] = PathIndex;
+	}
+	
+	// prune the list to match the size of the newly imported asset
+	// KeyzoneIdx, at this point, will be equal to the number of _imported_ keyzones
+	while (FusionPatchData.Keyzones.Num() > NumKeyzones)
+	{
+		FusionPatchData.Keyzones.RemoveAt(FusionPatchData.Keyzones.Num() - 1);
+	}
+
+	check(KeyzoneIdxToPathIdx.Num() == FusionPatchData.Keyzones.Num());
+
+	if (!AllFilesExist)
+	{
+		OutErrors.Add(FString::Printf(TEXT("Not all sample paths point to valid files on disk")));
+		UE_LOG(LogFusionPatchJsonImporter, Error, TEXT("Failed to import FusionPatch. Not all sample paths point to valid files on disk"));
+		return false;
+	}
+
+	TArray<UObject*> ImportedAssets = ImportAudioSamples(AudioSampleFiles, ImportArgs.SamplesDestPath, ImportArgs.ReplaceExistingSamples);
+
+	// update the sound wave loading behavior and compression type with the selection set in the import options
+	for (UObject* Asset : ImportedAssets)
+	{
+		if (USoundWave* SoundWave = Cast<USoundWave>(Asset))
+		{
+			// call OverrideLoadingBehavior to actually apply the loading behavior
+			SoundWave->OverrideLoadingBehavior(ImportArgs.SampleLoadingBehavior);
+
+			// Assign the LoadingBehavior to the SoundWave so that it gets serialized with the new setting
+			// (Override doesn't actually assign the new LoadingBehavior to the USoundWave, but rather to the underlying FSoundWaveDataPtr)
+			SoundWave->LoadingBehavior = ImportArgs.SampleLoadingBehavior;
+
+			// up until now, setting the loading behavior has had no effect on the sound wave...
+			// calling SetSoundAssetCompressionType forces the asset to actually update its state (marked as Dirty)
+			SoundWave->SetSoundAssetCompressionType(ImportArgs.SampleCompressionType, true);
+		}
+	}
+
+	for (int KeyzoneIdx = 0; KeyzoneIdx < FusionPatchData.Keyzones.Num(); ++KeyzoneIdx)
+	{
+		int32 PathIndex = KeyzoneIdxToPathIdx.Contains(KeyzoneIdx) ? KeyzoneIdxToPathIdx[KeyzoneIdx] : INDEX_NONE;
+
+		if (PathIndex == INDEX_NONE)
+		{
+			FString SamplePath = FusionPatchData.Keyzones[KeyzoneIdx].SamplePath;
+			FString AssetName = ImportedAssets[KeyzoneIdx]->GetName();
+			OutErrors.Add(FString::Printf(TEXT("Imported asset (Name: %s) failed to map asset (%s) to keyzone: %d"), *AssetName, *SamplePath, KeyzoneIdx));
+			UE_LOG(LogFusionPatchJsonImporter, Error, TEXT("Failed to import FusionPatch. Imported asset (Name: %s) failed to map asset (%s) to keyzone: %d"), *AssetName, *SamplePath, KeyzoneIdx);
+			continue;
+		}
+
+		if (!AudioSampleFiles.IsValidIndex(PathIndex))
+		{
+			FString SamplePath = FusionPatchData.Keyzones[KeyzoneIdx].SamplePath;
+			FString AssetName = AudioSampleFiles[KeyzoneIdx];
+			OutErrors.Add(FString::Printf(TEXT("Imported asset (Name: %s) failed to map file (%s) to keyzone: %d"), *AssetName, *SamplePath, KeyzoneIdx));
+			UE_LOG(LogFusionPatchJsonImporter, Error, TEXT("Failed to import FusionPatch. Imported asset (Name: %s) failed to map asset (%s) to keyzone: %d"), *AssetName, *SamplePath, KeyzoneIdx);
+			continue;
+		}
+
+		const FString& AudioSampleFile = AudioSampleFiles[PathIndex];
+
+		UObject** Asset = ImportedAssets.FindByPredicate([&AudioSampleFile](UObject* Asset)
+			{
+				USoundWave* SoundWave = Cast<USoundWave>(Asset);
+				FString Filename = SoundWave->AssetImportData->GetFirstFilename();
+				UE_LOG(LogFusionPatchJsonImporter, Log, TEXT("Comparing filenames: %s == %s"), *AudioSampleFile, *Filename);
+				return FPaths::IsSamePath(Filename, AudioSampleFile);
+			});
+
+		USoundWave* SoundWave = Asset != nullptr ? Cast<USoundWave>(*Asset) : nullptr;
+
+		if (!SoundWave)
+		{
+			// If the sound wave wasn't imported, check if it exists in the desired path 
+			// (we didn't replace it on import, and/or it was imported previously by another asset)
+			const FString AudioSampleName = FPaths::GetBaseFilename(AudioSampleFile);
+			const FString AssetPath = ImportArgs.SamplesDestPath / AudioSampleName + "." + AudioSampleName;
+			UObject* ExistingObject = UEditorAssetLibrary::LoadAsset(AssetPath);
+			SoundWave = Cast<USoundWave>(ExistingObject);
+		}
+
+		if (SoundWave)
+		{
+			FusionPatchData.Keyzones[KeyzoneIdx].SoundWave = SoundWave;
+		}
+	}
+
+	FusionPatch->UpdateRenderableForNonTrivialChange();
+
+	return true;
+}
+
+#undef LOCTEXT_NAMESPACE
\ No newline at end of file
