[Harmonix] Move Harmonix plugin into Engine/Plugins/Runtime to make it available to licensees in UE 5.4

diff --git a/Engine/Plugins/Runtime/Harmonix/Source/HarmonixDspEditor/Private/SettingsJsonImporter.cpp b/Engine/Plugins/Runtime/Harmonix/Source/HarmonixDspEditor/Private/SettingsJsonImporter.cpp
new file mode 100644
--- /dev/null
+++ b/Engine/Plugins/Runtime/Harmonix/Source/HarmonixDspEditor/Private/SettingsJsonImporter.cpp
@@ -0,0 +1,518 @@
+// Copyright Epic Games, Inc. All Rights Reserved.
+
+#include "SettingsJsonImporter.h"
+
+#include "Serialization/JsonReader.h"
+#include "Serialization/JsonSerializer.h"
+
+#include "HarmonixDsp/AudioUtility.h"
+#include "HarmonixDsp/PitchShifterName.h"
+#include "HarmonixDsp/FusionSampler/FusionPatch.h"
+#include "HarmonixDsp/FusionSampler/Settings/FusionPatchSettings.h"
+#include "HarmonixDsp/StretcherAndPitchShifterFactoryConfig.h"
+
+bool FSettingsJsonImporter::TryParseJson(TSharedPtr<FJsonObject> SettingsJson, FKeyzoneSettings& Settings)
+{
+	FKeyzoneSettings Defaults;
+	TryGetNumberField(SettingsJson, "root_note", Settings.RootNote, Defaults.RootNote);
+	TryGetNumberField(SettingsJson, "min_note", Settings.MinNote, Defaults.MinNote);
+	TryGetNumberField(SettingsJson, "max_note", Settings.MaxNote, Defaults.MaxNote);
+	TryGetBoolField(SettingsJson, "unpitched", Settings.bUnpitched, Defaults.bUnpitched);
+	TryGetBoolField(SettingsJson, "velocity_to_volume", Settings.bVelocityToGain, Defaults.bVelocityToGain);
+
+	TryParseJson(SettingsJson, Settings.TimeStretchConfig);
+
+	const TSharedPtr<FJsonObject>* TimeStretchJson = nullptr;
+	if (SettingsJson->TryGetObjectField(TEXT("timestretch_settings"), TimeStretchJson))
+	{
+		TryParseJson(*TimeStretchJson, Settings.TimeStretchConfig);
+	}
+
+	TryGetNumberField(SettingsJson, "random_weight", Settings.RandomWeight, Defaults.RandomWeight);
+	TryGetNumberField(SettingsJson, "min_velocity", Settings.MinVelocity, Defaults.MinVelocity);
+	TryGetNumberField(SettingsJson, "max_velocity", Settings.MaxVelocity, Defaults.MaxVelocity);
+	float VolumeDb;
+	TryGetNumberField(SettingsJson, "volume", VolumeDb, HarmonixDsp::kDbSilence);
+	Settings.SetVolumeDb(VolumeDb);
+
+	const TSharedPtr<FJsonObject>* PanJson = nullptr;
+	if (SettingsJson->TryGetObjectField(TEXT("pan"), PanJson))
+	{
+		TryParseJson(*PanJson, Settings.Pan);
+	}
+
+	TryGetNumberField(SettingsJson, "fine_tune", Settings.FineTuneCents, Defaults.FineTuneCents);
+	Settings.SetFineTuneCents(Settings.FineTuneCents);
+
+	TryGetBoolField(SettingsJson, "is_note_off_zone", Settings.bIsNoteOffZone, Defaults.bIsNoteOffZone);
+	TryGetNumberField(SettingsJson, "priority", Settings.Priority, Defaults.Priority);
+
+	TryGetStringField(SettingsJson, "sample_path", Settings.SamplePath);
+
+	TryGetNumberField(SettingsJson, "start_offset_frame", Settings.SampleStartOffset, Defaults.SampleStartOffset);
+	TryGetNumberField(SettingsJson, "end_offset_frame", Settings.SampleEndOffset, Defaults.SampleEndOffset);
+
+	TryGetBoolField(SettingsJson, "singleton", Settings.UseSingletonVoicePool);
+
+	return true;
+}
+
+bool FSettingsJsonImporter::TryParseJson(TSharedPtr<FJsonObject> JsonObj, FTimeStretchConfig& TimeStretchConfig)
+{
+	FTimeStretchConfig Defaults;
+	bool MaintainFormant = false;
+	TryGetBoolField(JsonObj, "maintain_time", TimeStretchConfig.bMaintainTime, Defaults.bMaintainTime);
+	TryGetBoolField(JsonObj, "maintain_formant", MaintainFormant);
+
+	const UStretcherAndPitchShifterFactoryConfig* FactoryConfig = GetDefault<UStretcherAndPitchShifterFactoryConfig>();
+	check(FactoryConfig);
+
+	FName AlgorithmName = FactoryConfig->DefaultFactory.Name;
+	TryGetNameField(JsonObj, "algorithm", AlgorithmName, AlgorithmName);
+
+	if (AlgorithmName == "default")
+	{
+		TimeStretchConfig.PitchShifter = FactoryConfig->DefaultFactory;
+	}
+	else if (const FPitchShifterNameRedirect* Redirect = FactoryConfig->FindFactoryNameRedirect(AlgorithmName))
+	{
+		TimeStretchConfig.PitchShifter = Redirect->NewName;
+	}
+	else
+	{
+		TimeStretchConfig.PitchShifter = AlgorithmName;
+	}
+
+	TryGetBoolField(JsonObj, "sync_tempo", TimeStretchConfig.bSyncTempo, Defaults.bSyncTempo);
+	TryGetNumberField(JsonObj, "orig_tempo", TimeStretchConfig.OriginalTempo, Defaults.OriginalTempo);
+
+	int16 EnvelopeOrder = 0;
+	JsonObj->TryGetNumberField(TEXT("envelope_order"), EnvelopeOrder);
+
+	TimeStretchConfig.PitchShifterOptions.Add("EnvelopeOrder", EnvelopeOrder);
+	TimeStretchConfig.PitchShifterOptions.Add("MaintainFormant", MaintainFormant);
+
+	return true;
+}
+
+bool FSettingsJsonImporter::TryParseJson(TSharedPtr<FJsonObject> JsonObj, FPannerDetails& OutPanDetails)
+{
+	FPannerDetails Defaults;
+	TryGetEnumField(JsonObj, "mode", OutPanDetails.Mode, Defaults.Mode);
+	switch (OutPanDetails.Mode)
+	{
+	case EPannerMode::LegacyStereo:
+	case EPannerMode::Stereo:
+		TryGetNumberField(JsonObj, "position", OutPanDetails.Detail.Pan, Defaults.Detail.Pan);
+		OutPanDetails.Detail.EdgeProximity = 0.0f;
+		return true;
+	case EPannerMode::Surround:
+	case EPannerMode::PolarSurround:
+		TryGetNumberField(JsonObj, "position", OutPanDetails.Detail.Pan, Defaults.Detail.Pan);
+		TryGetNumberField(JsonObj, "edge_proximity", OutPanDetails.Detail.EdgeProximity, Defaults.Detail.EdgeProximity);
+		return true;
+	
+	case EPannerMode::DirectAssignment:
+		TryGetEnumField(JsonObj, "position", OutPanDetails.Detail.ChannelAssignment, ESpeakerChannelAssignment::LeftFront);
+		return true;
+	}
+	return false;
+}
+
+bool FSettingsJsonImporter::TryParseJson(TSharedPtr<FJsonObject> JsonObj, FFusionPatchSettings& Preset)
+{
+	FFusionPatchSettings Defaults;
+	TryGetNumberField(JsonObj, "volume", Preset.VolumeDb, HarmonixDsp::kDbSilence);
+	
+	TryParseJson(JsonObj, Preset.PannerDetails);
+
+	TryGetNumberField(JsonObj, "start_point", Preset.StartPointOffsetMs, Defaults.StartPointOffsetMs);
+	TryGetNumberField(JsonObj, "fine_tune", Preset.FineTuneCents, Defaults.FineTuneCents);
+	TryGetNumberField(JsonObj, "voice_limit", Preset.MaxVoices, 8);
+	TryGetEnumField(JsonObj, "layer_select_mode", Preset.KeyzoneSelectMode, EKeyzoneSelectMode::Layers);
+
+	// apply defaults just in case
+	Preset.DownPitchBendCents = -700.0f;
+	Preset.UpPitchBendCents = 700.0f;
+	TSharedPtr<FJsonObject> PitchBendObj;
+	if (TryGetObjectField(JsonObj, "pitch_bend", PitchBendObj))
+	{
+		const TArray<TSharedPtr<FJsonValue>>* RangeValues;
+		if (PitchBendObj->TryGetArrayField(TEXT("range"), RangeValues))
+		{
+			if (ensure(RangeValues->Num() == 2))
+			{
+				TryGetNumber((*RangeValues)[0], Preset.DownPitchBendCents, Defaults.DownPitchBendCents);
+				TryGetNumber((*RangeValues)[1], Preset.UpPitchBendCents, Defaults.UpPitchBendCents);
+			}
+		}
+	}
+
+	{
+	int AdsrIdx = -1;
+	for (TSharedPtr<FJsonValue> ArrayValue : IterField(JsonObj, "adsrs"))
+	{
+		++AdsrIdx;
+		TSharedPtr<FJsonObject> AdsrObj;
+		if (!TryGetObject(ArrayValue, AdsrObj))
+			continue;
+
+		TSharedPtr<FJsonObject> AdsrValue;
+		if (!TryGetObjectField(AdsrObj, "adsr", AdsrValue))
+			continue;
+
+		FAdsrSettings& AdsrSettings = AdsrIdx == 0 ? Preset.Adsrs.Volume : Preset.Adsrs.Assignable;
+		if (!TryParseJson(AdsrValue, AdsrSettings))
+			continue;
+
+		// old defaults
+		if (AdsrIdx == 0 && !AdsrValue->HasField(TEXT("target")))
+		{
+			AdsrSettings.Target = EAdsrTarget::Volume;
+			AdsrSettings.Depth = 1.0f;
+			AdsrSettings.IsEnabled = true;
+		}
+	}
+	}
+
+	{
+	int LfoIdx = -1;
+	for (TSharedPtr<FJsonValue> ArrayValue : IterField(JsonObj, "lfos"))
+	{
+		++LfoIdx;
+
+		if (!ensure(LfoIdx < Preset.Lfos.Num()))
+			break;
+
+		TSharedPtr<FJsonObject> LfoObj;
+		if (!TryGetObject(ArrayValue, LfoObj))
+			continue;
+
+		TSharedPtr<FJsonObject> LfoValue;
+		if (!TryGetObjectField(LfoObj, "lfo", LfoValue))
+			continue;
+
+		FLfoSettings& LfoSettings = Preset.Lfos[LfoIdx];
+
+		if (!TryParseJson(LfoValue, LfoSettings))
+			continue;
+
+		if (!LfoValue->HasField(TEXT("target")))
+		{
+			// old defaults
+			if (LfoIdx == 0)
+				LfoSettings.Target = ELfoTarget::Pan;
+			else if (LfoIdx == 1)
+				LfoSettings.Target = ELfoTarget::Pitch;
+		}
+	}
+	}
+
+	TSharedPtr<FJsonObject> FilterObj;
+	if (TryGetObjectField(JsonObj, "filter", FilterObj))
+	{
+		TryParseJson(FilterObj, Preset.Filter);
+	}
+	else
+	{
+		Preset.Filter.IsEnabled = false;
+	}
+
+	{
+	int RandomizerIdx = -1;
+	for (TSharedPtr<FJsonValue> ArrayValue : IterField(JsonObj, "randomizers"))
+	{
+		++RandomizerIdx;
+		if (!ensure(RandomizerIdx < Preset.Randomizers.Num()))
+			break;
+
+		TSharedPtr<FJsonObject> RandomizerObj;
+		if (!TryGetObject(ArrayValue, RandomizerObj))
+			continue;
+
+		TSharedPtr<FJsonObject> RandomizerValue;
+		if (!TryGetObjectField(RandomizerObj, "modulator", RandomizerValue))
+			continue;
+
+		TryParseJson(RandomizerValue, Preset.Randomizers[RandomizerIdx]);
+	}
+	}
+
+	{
+	int ModIdx = -1;
+	for (TSharedPtr<FJsonValue> ArrayValue : IterField(JsonObj, "velocity_mods"))
+	{
+		++ModIdx;
+		if (!ensure(ModIdx < Preset.VelocityModulators.Num()))
+			break;
+
+		TSharedPtr<FJsonObject> VelocityModObj;
+		if (!TryGetObject(ArrayValue, VelocityModObj))
+			continue;
+
+		TSharedPtr<FJsonObject> VelocityModValue;
+		if (!TryGetObjectField(VelocityModObj, "modulator", VelocityModValue))
+			continue;
+
+		TryParseJson(VelocityModValue, Preset.VelocityModulators[ModIdx]);
+	}
+	}
+
+	{
+	TSharedPtr<FJsonObject> DelayObj;
+	if (TryGetObjectField(JsonObj, "delay_effect", DelayObj))
+	{
+		TryParseJson(DelayObj, Preset.Delay);
+	}
+	else
+	{
+		Preset.Delay.IsEnabled = false;
+	}
+	}
+
+	{
+	TSharedPtr<FJsonObject> BitcrushObj;
+	if (TryGetObjectField(JsonObj, "bitcrush_effect", BitcrushObj))
+	{
+		TryParseJson(BitcrushObj, Preset.BitCrusher);
+	}
+	else
+	{
+		Preset.BitCrusher.IsEnabled = false;
+	}
+	}
+
+	{
+	TSharedPtr<FJsonObject> VocoderObj;
+	if (TryGetObjectField(JsonObj, "vocoder_effect", VocoderObj))
+	{
+		TryParseJson(VocoderObj, Preset.Vocoder);
+	}
+	else
+	{
+		Preset.Vocoder.IsEnabled = false;
+	}
+	}
+
+	{
+	TSharedPtr<FJsonObject> DistortionObj;
+	if (TryGetObjectField(JsonObj, "distortion_effect", DistortionObj))
+	{
+		TryParseJson(DistortionObj, Preset.Distortion);
+	}
+	else
+	{
+		Preset.Distortion.IsEnabled = false;
+	}
+	}
+
+	{
+	TSharedPtr<FJsonObject> PortamentoObj;
+	if (TryGetObjectField(JsonObj, "portamento", PortamentoObj))
+	{
+		TryParseJson(PortamentoObj, Preset.Portamento);
+	}
+	else
+	{
+		Preset.Portamento.IsEnabled = false;
+		Preset.Portamento.Mode = EPortamentoMode::Legato;
+		Preset.Portamento.Seconds = 0.0f;
+	}
+	}
+
+	return true;
+}
+
+bool FSettingsJsonImporter::TryParseJson(TSharedPtr<FJsonObject> JsonObj, FAdsrSettings& AdsrSettings)
+{
+	// simplify the conversion process :)
+	static TFunction<float(float)> MsToSec = [](float Ms) { return Ms * 0.001f; };
+	TryGetEnumField(JsonObj, "target", AdsrSettings.Target, EAdsrTarget::None);
+	TryGetNumberField(JsonObj, "depth", AdsrSettings.Depth, 0.0f);
+	TryGetBoolField(JsonObj, "enabled", AdsrSettings.IsEnabled, false);
+	TryGetNumberField(JsonObj, "attack", AdsrSettings.AttackTime, MsToSec, 1.0f);
+	TryGetNumberField(JsonObj, "attack_curve", AdsrSettings.AttackCurve, 0.0f);
+	TryGetNumberField(JsonObj, "decay", AdsrSettings.DecayTime, MsToSec, 1.0f);
+	TryGetNumberField(JsonObj, "decay_curve", AdsrSettings.DecayCurve, 0.0f);
+	TryGetNumberField(JsonObj, "sustain", AdsrSettings.SustainLevel, 1.0f);
+	TryGetNumberField(JsonObj, "release", AdsrSettings.ReleaseTime, MsToSec, 1.0f);
+	TryGetNumberField(JsonObj, "release_curve", AdsrSettings.ReleaseCurve, 0.0f);
+
+	AdsrSettings.Calculate();
+
+	return true;
+}
+
+bool FSettingsJsonImporter::TryParseJson(TSharedPtr<FJsonObject> JsonObj, FLfoSettings& LfoSettings)
+{
+	TryGetEnumField(JsonObj, "target", LfoSettings.Target, ELfoTarget::None);
+	TryGetNumberField(JsonObj, "frequency", LfoSettings.Freq);
+	TryGetNumberField(JsonObj, "depth", LfoSettings.Depth);
+	TryGetBoolField(JsonObj, "enabled", LfoSettings.IsEnabled);
+	TryGetEnumField(JsonObj, "shape", LfoSettings.Shape, EWaveShape::None);
+	TryGetBoolField(JsonObj, "retrigger", LfoSettings.ShouldRetrigger, false);
+	TryGetNumberField(JsonObj, "initial_phase", LfoSettings.InitialPhase, 0.0f);
+	TryGetBoolField(JsonObj, "beat_sync", LfoSettings.BeatSync, false);
+	TryGetNumberField(JsonObj, "tempo", LfoSettings.TempoBPM, 120.0f);
+
+	return true;
+}
+
+bool FSettingsJsonImporter::TryParseJson(TSharedPtr<FJsonObject> JsonObj, FBiquadFilterSettings& Filter)
+{
+	TryGetBoolField(JsonObj, "enabled", Filter.IsEnabled);
+	TryGetEnumField(JsonObj, "type", Filter.Type, EBiquadFilterType::None);
+	TryGetNumberField(JsonObj, "frequency", Filter.Freq);
+	TryGetNumberField(JsonObj, "q", Filter.Q);
+	TryGetNumberField(JsonObj, "gain_db", Filter.DesignedDBGain, 0.0f);
+
+	return true;
+}
+
+bool FSettingsJsonImporter::TryParseJson(TSharedPtr<FJsonObject> JsonObj, FModulatorSettings& ModulatorSettings)
+{
+	TryGetEnumField(JsonObj, "target", ModulatorSettings.Target, EModulatorTarget::None);
+	TryGetNumberField(JsonObj, "range", ModulatorSettings.Range, 0.0f);
+	TryGetNumberField(JsonObj, "depth", ModulatorSettings.Depth, 0.0f);
+	return true;
+}
+
+bool FSettingsJsonImporter::TryParseJson(TSharedPtr<FJsonObject> JsonObj, FDelaySettings& DelaySettings)
+{
+	TryGetBoolField(JsonObj, "enabled", DelaySettings.IsEnabled, false);
+
+	{
+		bool OldBeatSyncSetting;
+		TryGetBoolField(JsonObj, "beat_sync", OldBeatSyncSetting, false);
+		TryGetEnumField(JsonObj, "time_sync", DelaySettings.TimeSyncOption, OldBeatSyncSetting ? ETimeSyncOption::TempoSync : ETimeSyncOption::None);
+	}
+
+	TryGetNumberField(JsonObj, "time", DelaySettings.TimeSeconds, 0.5f);
+	TryGetNumberField(JsonObj, "dry_gain", DelaySettings.DryGain, 1.0f);
+	TryGetNumberField(JsonObj, "wet_gain", DelaySettings.WetGain, 0.5f);
+	TryGetNumberField(JsonObj, "feedback", DelaySettings.FeedbackGain, 0.25f);
+	TryGetBoolField(JsonObj, "delay_eq_enabled", DelaySettings.EQEnabled, false);
+	TryGetEnumField(JsonObj, "delay_eq_type", DelaySettings.EQType, EDelayFilterType::LowPass);
+	TryGetNumberField(JsonObj, "delay_eq_frequency", DelaySettings.EQFreq, 600.0f);
+	TryGetNumberField(JsonObj, "delay_eq_q", DelaySettings.EQQ, 1.0f);
+	TryGetBoolField(JsonObj, "delay_lfo_enabled", DelaySettings.LfoEnabled, false);
+
+	{
+		bool OldBeatSyncSetting;
+		TryGetBoolField(JsonObj, "delay_lfo_beatsync", OldBeatSyncSetting, false);
+		TryGetEnumField(JsonObj, "delay_lfo_time_sync", DelaySettings.LfoTimeSyncOption, OldBeatSyncSetting ? ETimeSyncOption::TempoSync : ETimeSyncOption::None);
+	}
+
+	TryGetNumberField(JsonObj, "delay_lfo_rate", DelaySettings.LfoRate, 0.0f);
+	TryGetNumberField(JsonObj, "delay_lfo_depth", DelaySettings.LfoDepth, 0.0f);
+	TryGetEnumField(JsonObj, "delay_stereo_type", DelaySettings.StereoType, EDelayStereoType::Default);
+	TryGetNumberField(JsonObj, "delay_stereo_left", DelaySettings.PanLeft, 0.0f);
+	TryGetNumberField(JsonObj, "delay_stereo_right", DelaySettings.PanRight, 0.0f);
+	return true;
+}
+
+bool FSettingsJsonImporter::TryParseJson(TSharedPtr<FJsonObject> JsonObj, FBitCrusherSettings& BitCrusherSettings)
+{
+	static TFunction<float(float)> PctToGain = [](float Pct) { return Pct / 100.0f; };
+
+	TryGetBoolField(JsonObj, "enabled", BitCrusherSettings.IsEnabled, false);
+	TryGetNumberField(JsonObj, "wet_percent", BitCrusherSettings.WetGain, PctToGain, 1.0f);
+	TryGetNumberField(JsonObj, "crush_level", BitCrusherSettings.Crush, (uint16)4);
+	TryGetNumberField(JsonObj, "sample_hold", BitCrusherSettings.SampleHoldFactor, (uint16)5);
+	return true;
+}
+
+bool FSettingsJsonImporter::TryParseJson(TSharedPtr<FJsonObject> JsonObj, FVocoderSettings& VocoderSettings)
+{
+	TryGetBoolField(JsonObj, "enabled", VocoderSettings.IsEnabled, false);
+	TryGetNumberField(JsonObj, "modulator_index", VocoderSettings.ModulatorIndex, -1);
+	if (VocoderSettings.ModulatorIndex == -1)
+	{
+		// index name intentionally masked
+		FString DeprecatedIndex = "\x73\x6c\x61\x76\x65\x5f\x69\x6e\x64\x65\x78";
+		TryGetNumberField(JsonObj, DeprecatedIndex, VocoderSettings.ModulatorIndex, 0);
+	}
+
+	TryGetEnumField(JsonObj, "band_config_index", VocoderSettings.BandConfig, EVocoderBandConfig::k64);
+	TryGetNumberField(JsonObj, "carrier_gain", VocoderSettings.CarrierGain, 1.0f);
+	TryGetNumberField(JsonObj, "carrier_thin", VocoderSettings.CarrierThin, 1.0f);
+	TryGetNumberField(JsonObj, "modulator_gain", VocoderSettings.ModulatorGain, 1.0f);
+	TryGetNumberField(JsonObj, "modultar_thin", VocoderSettings.ModulatorThin, 1.0f);
+	TryGetNumberField(JsonObj, "attack", VocoderSettings.Attack, 0.5f);
+	TryGetNumberField(JsonObj, "release", VocoderSettings.Release, 0.5f);
+	TryGetNumberField(JsonObj, "high_emphasis", VocoderSettings.HighEmphasis, 0.0f);
+	TryGetNumberField(JsonObj, "output_gain_db", VocoderSettings.OutputGain, 1.0f);
+	TryGetBoolField(JsonObj, "solo", VocoderSettings.Soloing, false);
+
+	{
+	const TArray<TSharedPtr<FJsonValue>>* BandArray;
+	if (JsonObj->TryGetArrayField(TEXT("band"), BandArray))
+	{
+		VocoderSettings.Bands.SetNum(BandArray->Num());
+
+		for (TSharedPtr<FJsonValue> BandValue : *BandArray)
+		{
+			TSharedPtr<FJsonObject> BandObj;
+			if (!TryGetObject(BandValue, BandObj))
+				continue;
+
+			VocoderSettings.Bands.SetNum(VocoderSettings.Bands.Num() + 1);
+
+			TryGetNumberField(BandObj, "gain", VocoderSettings.Bands.Last().Gain, 1.0f);
+			TryGetBoolField(BandObj, "solo", VocoderSettings.Bands.Last().Solo, false);
+		}
+	}
+	}
+	return true;
+}
+
+bool FSettingsJsonImporter::TryParseJson(TSharedPtr<FJsonObject> JsonObj, FDistortionSettingsV1& DistortionSettings)
+{
+	TryGetBoolField(JsonObj, "enabled", DistortionSettings.IsEnabled, false);
+	TryGetNumberField(JsonObj, "input_gain", DistortionSettings.InputGainDb, 0.0f);
+	TryGetNumberField(JsonObj, "output_gain", DistortionSettings.OutputGainDb, 0.0f);
+	TryGetNumberField(JsonObj, "dc_adjust", DistortionSettings.DCAdjust, 0.0f);
+	TryGetEnumField(JsonObj, "type", DistortionSettings.Type, EDistortionTypeV1::Clean);
+	TryGetBoolField(JsonObj, "oversample", DistortionSettings.Oversample, false);
+
+	{
+	int FilterIdx = -1;
+	for (TSharedPtr<FJsonValue> ArrayValue : IterField(JsonObj, "filter_set"))
+	{
+		++FilterIdx;
+
+		if (!ensureAlways(FilterIdx < FDistortionSettingsV1::kNumFilters))
+			break;
+
+		TSharedPtr<FJsonObject> FilterObj;
+		if (!TryGetObject(ArrayValue, FilterObj))
+			continue;
+
+		TSharedPtr<FJsonObject> FilterValue;
+		if (!TryGetObjectField(FilterObj, "filter", FilterValue))
+			continue;
+
+		TryParseJson(FilterValue, DistortionSettings.Filters[FilterIdx]);
+	}
+	}
+	return true;
+}
+
+
+bool FSettingsJsonImporter::TryParseJson(TSharedPtr<FJsonObject> JsonObj, FDistortionFilterSettings& DistortionFilterSettings)
+{
+	TryParseJson(JsonObj, DistortionFilterSettings.Filter);
+	TryGetBoolField(JsonObj, "filter_pre", DistortionFilterSettings.FilterPreClip);
+	TryGetNumberField(JsonObj, "num_passes", DistortionFilterSettings.NumPasses);
+
+	return true;
+}
+
+
+bool FSettingsJsonImporter::TryParseJson(TSharedPtr<FJsonObject> JsonObj, FPortamentoSettings& PortamentoSettings)
+{
+	TryGetBoolField(JsonObj, "enabled", PortamentoSettings.IsEnabled, false);
+	TryGetEnumField(JsonObj, "mode", PortamentoSettings.Mode, EPortamentoMode::Legato);
+	TryGetNumberField(JsonObj, "time", PortamentoSettings.Seconds, 0.0f);
+	return true;
+}
diff --git a/Engine/Plugins/Runtime/Harmonix/Source/HarmonixDspEditor/Private/SettingsJsonImporter.h b/Engine/Plugins/Runtime/Harmonix/Source/HarmonixDspEditor/Private/SettingsJsonImporter.h
new file mode 100644
--- /dev/null
+++ b/Engine/Plugins/Runtime/Harmonix/Source/HarmonixDspEditor/Private/SettingsJsonImporter.h
@@ -0,0 +1,40 @@
+// Copyright Epic Games, Inc. All Rights Reserved.
+
+#pragma once
+
+#include "JsonImporterHelper.h"
+
+class FJsonObject;
+
+class FSettingsJsonImporter : public FJsonImporter
+{
+public:
+
+	static bool TryParseJson(TSharedPtr<FJsonObject> JsonObj, struct FKeyzoneSettings& OutSettings);
+
+	static bool TryParseJson(TSharedPtr<FJsonObject> JsonObj, struct FPannerDetails& OutPanDetails);
+
+	static bool TryParseJson(TSharedPtr<FJsonObject> JsonObj, struct FTimeStretchConfig& TimeStretchConfig);
+
+	static bool TryParseJson(TSharedPtr<FJsonObject> JsonObj, struct FFusionPatchSettings& Preset);
+
+	static bool TryParseJson(TSharedPtr<FJsonObject> JsonObj, struct FAdsrSettings& AdsrSettings);
+
+	static bool TryParseJson(TSharedPtr<FJsonObject> JsonObj, struct FLfoSettings& LfoSettings);
+
+	static bool TryParseJson(TSharedPtr<FJsonObject> JsonObj, struct FBiquadFilterSettings& FilterSettings);
+
+	static bool TryParseJson(TSharedPtr<FJsonObject> JsonObj, struct FModulatorSettings& ModulatorSettings);
+
+	static bool TryParseJson(TSharedPtr<FJsonObject> JsonObj, struct FDelaySettings& DelaySettings);
+
+	static bool TryParseJson(TSharedPtr<FJsonObject> JsonObj, struct FBitCrusherSettings& BitCrusherSettings);
+
+	static bool TryParseJson(TSharedPtr<FJsonObject> JsonObj, struct FVocoderSettings& VocoderSettings);
+
+	static bool TryParseJson(TSharedPtr<FJsonObject> JsonObj, struct FDistortionSettingsV1& DistortionSettings);
+
+	static bool TryParseJson(TSharedPtr<FJsonObject> JsonObj, struct FDistortionFilterSettings& DistortionFilterSettings);
+
+	static bool TryParseJson(TSharedPtr<FJsonObject> JsonObj, struct FPortamentoSettings& PortamentoSettings);
+};
\ No newline at end of file
diff --git a/Engine/Plugins/Runtime/Harmonix/Source/HarmonixDspTests/Private/AudioAnalysis/DynamicRangeMeterAnalyzerTests.cpp b/Engine/Plugins/Runtime/Harmonix/Source/HarmonixDspTests/Private/AudioAnalysis/DynamicRangeMeterAnalyzerTests.cpp
new file mode 100644
--- /dev/null
+++ b/Engine/Plugins/Runtime/Harmonix/Source/HarmonixDspTests/Private/AudioAnalysis/DynamicRangeMeterAnalyzerTests.cpp
@@ -0,0 +1,55 @@
+// Copyright Epic Games, Inc. All Rights Reserved.
+
+#include "HarmonixDsp/Generate.h"
+#include "HarmonixDsp/AudioAnalysis/DynamicRangeMeterAnalyzer.h"
+#include "Misc/AutomationTest.h"
+
+#if WITH_DEV_AUTOMATION_TESTS
+
+namespace Harmonix::Dsp::AudioAnalysis::VuMeter::Tests
+{
+	IMPLEMENT_SIMPLE_AUTOMATION_TEST(
+		FHarmonixDynamicRangeMeterBasicTest,
+		"Harmonix.Dsp.AudioAnalysis.DynamicRangeMeter.Basic",
+		EAutomationTestFlags::EditorContext | EAutomationTestFlags::EngineFilter)
+	bool FHarmonixDynamicRangeMeterBasicTest::RunTest(const FString&)
+	{
+		constexpr int32 NumChannels = 2;
+		constexpr int32 NumFramesPerBlock = 256;
+		constexpr float SampleRate = 48000;
+		
+		FDynamicRangeMeterAnalyzer Analyzer;
+		FDynamicRangeMeterAnalyzer::FSettings Settings;
+		Settings.SampleRate = SampleRate;
+		Analyzer.SetSettings(Settings);
+
+		// fill the buffer with a sine wave
+		TAudioBuffer<float> Buffer{ NumChannels, NumFramesPerBlock, EAudioBufferCleanupMode::Delete };
+		{
+			constexpr float Frequency = 440;
+			float Phase = 0;
+			UTEST_TRUE("Test buffer must have at least one channel", Buffer.GetNumValidChannels() > 0);
+			HarmonixDsp::GenerateSine(Buffer.GetValidChannelData(0), NumFramesPerBlock, Frequency, SampleRate, Phase);
+			for (int32 c = 1; c < NumChannels; ++c)
+			{
+				FMemory::Memcpy(Buffer.GetValidChannelData(c), Buffer.GetValidChannelData(0), NumFramesPerBlock * sizeof(float));
+			}
+		}
+
+		FDynamicRangeMeterAnalyzer::FResults Results;
+		
+		Analyzer.Process(Buffer, Results);
+
+		// make sure we got some results
+		UTEST_TRUE("High envelope is greater than zero", Results.HighEnvelope > 0.0f);
+		UTEST_TRUE("Low envelope is greater than zero", Results.LowEnvelope > 0.0f);
+		UTEST_TRUE("Level is greater than -96Db", Results.LevelDecibels > -96.0f);
+		UTEST_TRUE("Mono peak is greater than -96Db", Results.MonoPeakDecibels > -96.0f);
+		UTEST_TRUE("Mono peak high envelope is greater than -96Db", Results.MonoPeakHighEnvelopeDecibels > -96.0f);
+		
+		return true;
+	}
+
+}
+
+#endif
diff --git a/Engine/Plugins/Runtime/Harmonix/Source/HarmonixDspTests/Private/AudioAnalysis/VuMeterAnalyzerTests.cpp b/Engine/Plugins/Runtime/Harmonix/Source/HarmonixDspTests/Private/AudioAnalysis/VuMeterAnalyzerTests.cpp
new file mode 100644
--- /dev/null
+++ b/Engine/Plugins/Runtime/Harmonix/Source/HarmonixDspTests/Private/AudioAnalysis/VuMeterAnalyzerTests.cpp
@@ -0,0 +1,60 @@
+// Copyright Epic Games, Inc. All Rights Reserved.
+
+#include "HarmonixDsp/Generate.h"
+#include "HarmonixDsp/AudioAnalysis/VuMeterAnalyzer.h"
+#include "Misc/AutomationTest.h"
+
+#if WITH_DEV_AUTOMATION_TESTS
+
+namespace Harmonix::Dsp::AudioAnalysis::VuMeter::Tests
+{
+	IMPLEMENT_SIMPLE_AUTOMATION_TEST(
+		FHarmonixVuMeterBasicTest,
+		"Harmonix.Dsp.AudioAnalysis.VuMeter.Basic",
+		EAutomationTestFlags::EditorContext | EAutomationTestFlags::EngineFilter)
+	bool FHarmonixVuMeterBasicTest::RunTest(const FString&)
+	{
+		constexpr int32 NumChannels = 2;
+		constexpr int32 NumFramesPerBlock = 256;
+		constexpr float SampleRate = 48000;
+		
+		FVuMeterAnalyzer VuMeterAnalyzer(SampleRate);
+
+		// fill the buffer with a sine wave
+		TAudioBuffer<float> Buffer{ NumChannels, NumFramesPerBlock, EAudioBufferCleanupMode::Delete };
+		{
+			constexpr float Frequency = 440;
+			float Phase = 0;
+			UTEST_TRUE("Test buffer must have at least one channel", Buffer.GetNumValidChannels() > 0);
+			HarmonixDsp::GenerateSine(Buffer.GetValidChannelData(0), NumFramesPerBlock, Frequency, SampleRate, Phase);
+			for (int32 c = 1; c < NumChannels; ++c)
+			{
+				FMemory::Memcpy(Buffer.GetValidChannelData(c), Buffer.GetValidChannelData(0), NumFramesPerBlock * sizeof(float));
+			}
+		}
+
+		FHarmonixVuMeterAnalyzerResults Results;
+
+		// results should be zero
+		UTEST_EQUAL("Mono level is zero", Results.MonoValues.LevelMeanSquared, 0.0f);
+		UTEST_EQUAL("Mono peak is zero", Results.MonoValues.PeakSquared, 0.0f);
+		UTEST_EQUAL("There are no channel values", Results.ChannelValues.Num(), 0);
+		
+		VuMeterAnalyzer.Process(Buffer, Results);
+
+		// make sure we got some results
+		UTEST_TRUE("Mono level is greater than zero", Results.MonoValues.LevelMeanSquared > 0.0f);
+		UTEST_TRUE("Mono peak is greater than zero", Results.MonoValues.PeakSquared > 0.0f);
+		UTEST_EQUAL("Results have the same number of channels as the input audio", Results.ChannelValues.Num(), NumChannels);
+		for (int32 c = 0; c < NumChannels; ++c)
+		{
+			UTEST_TRUE("Channel level is greater than zero", Results.ChannelValues[c].LevelMeanSquared > 0.0f);
+			UTEST_TRUE("Channel peak is greater than zero", Results.ChannelValues[c].PeakSquared > 0.0f);
+		}
+		
+		return true;
+	}
+
+}
+
+#endif
diff --git a/Engine/Plugins/Runtime/Harmonix/Source/HarmonixDspTests/Private/AudioAnalysis/WaveformAnalyzerTests.cpp b/Engine/Plugins/Runtime/Harmonix/Source/HarmonixDspTests/Private/AudioAnalysis/WaveformAnalyzerTests.cpp
new file mode 100644
--- /dev/null
+++ b/Engine/Plugins/Runtime/Harmonix/Source/HarmonixDspTests/Private/AudioAnalysis/WaveformAnalyzerTests.cpp
@@ -0,0 +1,115 @@
+// Copyright Epic Games, Inc. All Rights Reserved.
+
+#include "HarmonixDsp/AudioAnalysis/WaveformAnalyzer.h"
+#include "Misc/AutomationTest.h"
+
+#if WITH_DEV_AUTOMATION_TESTS
+
+namespace Harmonix::Dsp::AudioAnalysis::WaveformAnalyzer::Tests
+{
+	IMPLEMENT_SIMPLE_AUTOMATION_TEST(
+		FHarmonixWaveformAnalyzerBasicTest,
+		"Harmonix.Dsp.AudioAnalysis.Waveform.Basic",
+		EAutomationTestFlags::EditorContext | EAutomationTestFlags::EngineFilter)
+	bool FHarmonixWaveformAnalyzerBasicTest::RunTest(const FString&)
+	{
+		constexpr int32 NumChannels = 2;
+		constexpr int32 NumFramesPerBlock = 256;
+		constexpr float SampleRate = 48000;
+		
+		FWaveformAnalyzer Analyzer{ SampleRate };
+		constexpr FHarmonixWaveformAnalyzerSettings Settings;
+		Analyzer.SetSettings(Settings);
+		
+		FHarmonixWaveformAnalyzerResults Results;
+
+		// process a few times so we have some data
+		TAudioBuffer<float> Buffer{ NumChannels, NumFramesPerBlock, EAudioBufferCleanupMode::Delete };
+		constexpr float Frequency = 440;
+		float Phase = 0;
+		constexpr int32 NumBlocks = 10;
+		for (int i = 0; i < NumBlocks; ++i)
+		{
+			// Fill the buffer with a sine
+			HarmonixDsp::GenerateSine(Buffer.GetValidChannelData(0), NumFramesPerBlock, Frequency, SampleRate, Phase);
+			for (int32 c = 1; c < NumChannels; ++c)
+			{
+				FMemory::Memcpy(Buffer.GetValidChannelData(c), Buffer.GetValidChannelData(0), NumFramesPerBlock * sizeof(float));
+			}
+			Analyzer.Process(Buffer, Results);
+		}
+		
+		// make sure we got some results
+		UTEST_EQUAL("Got the right number of raw bins", Results.WaveformRaw.Num(), Settings.NumBinsHeld);
+		UTEST_EQUAL("Got the right number of smoothed bins", Results.WaveformSmoothed.Num(), Settings.NumBinsHeld);
+
+		bool ABinWasNonZero = false;
+		
+		for (int32 BinIdx = 0; BinIdx < Settings.NumBinsHeld; ++BinIdx)
+		{
+			UTEST_EQUAL(
+				FString::Printf(TEXT("Samples match at bin index %i"), BinIdx),
+				Results.WaveformRaw[BinIdx],
+				Results.WaveformSmoothed[BinIdx]);
+			if (!FMath::IsNearlyZero(Results.WaveformRaw[BinIdx]))
+			{
+				ABinWasNonZero = true;
+			}
+		}
+
+		UTEST_TRUE("A bin was non-zero", ABinWasNonZero);
+		
+		return true;
+	}
+
+	IMPLEMENT_SIMPLE_AUTOMATION_TEST(
+		FHarmonixWaveformAnalyzerCorrectRawValuesTest,
+		"Harmonix.Dsp.AudioAnalysis.Waveform.CorrectRawValues",
+		EAutomationTestFlags::EditorContext | EAutomationTestFlags::EngineFilter)
+	bool FHarmonixWaveformAnalyzerCorrectRawValuesTest::RunTest(const FString&)
+	{
+		constexpr int32 NumFramesPerBlock = 16;
+		constexpr float SampleRate = 48000;
+		
+		FWaveformAnalyzer Analyzer{ SampleRate };
+		FHarmonixWaveformAnalyzerSettings Settings;
+		Settings.NumBinsPerSecond = SampleRate; // this makes it so we get the waveform out the other side
+		Settings.NumBinsHeld = NumFramesPerBlock; // this makes it so we just hold a block of samples
+		Analyzer.SetSettings(Settings);
+		
+		FHarmonixWaveformAnalyzerResults Results;
+		
+		// Process a block with a sine
+		TAudioBuffer<float> Buffer{ 1, NumFramesPerBlock, EAudioBufferCleanupMode::Delete };
+		constexpr float Frequency = 440;
+		float Phase = 0;
+		HarmonixDsp::GenerateSine(Buffer.GetValidChannelData(0), NumFramesPerBlock, Frequency, SampleRate, Phase);
+		Analyzer.Process(Buffer, Results);
+		
+		// make sure we got some results
+		UTEST_EQUAL("Got the right number of raw bins", Results.WaveformRaw.Num(), Settings.NumBinsHeld);
+		UTEST_EQUAL("Got the right number of smoothed bins", Results.WaveformSmoothed.Num(), Settings.NumBinsHeld);
+
+		TDynamicStridePtr<float> InputPtr = Buffer.GetStridingChannelDataPointer(0);
+		
+		for (int32 BinIdx = 0; BinIdx < Settings.NumBinsHeld; ++BinIdx)
+		{
+			// no smoothing, so we should have the same values in both arrays
+			UTEST_EQUAL(
+				FString::Printf(TEXT("Samples match at bin index %i"), BinIdx),
+				Results.WaveformRaw[BinIdx],
+				Results.WaveformSmoothed[BinIdx]);
+			// make sure the input buffer matches the output buffer
+			// The history is in reverse, so make sure to check opposite sides of the buffer
+			const int32 InputIdx = (Settings.NumBinsHeld - BinIdx) - 1;
+			UTEST_EQUAL(
+				FString::Printf(TEXT("Output == Input at index %i"), BinIdx),
+				Results.WaveformRaw[BinIdx],
+				InputPtr[InputIdx]);
+		}
+		
+		return true;
+	}
+}
+
+#endif
diff --git a/Engine/Plugins/Runtime/Harmonix/Source/HarmonixMetasound/Private/HarmonixMetasound/Analysis/MidiClockVertexAnalyzer.cpp b/Engine/Plugins/Runtime/Harmonix/Source/HarmonixMetasound/Private/HarmonixMetasound/Analysis/MidiClockVertexAnalyzer.cpp
new file mode 100644
--- /dev/null
+++ b/Engine/Plugins/Runtime/Harmonix/Source/HarmonixMetasound/Private/HarmonixMetasound/Analysis/MidiClockVertexAnalyzer.cpp
@@ -0,0 +1,66 @@
+// Copyright Epic Games, Inc. All Rights Reserved.
+
+#include "HarmonixMetasound/Analysis/MidiClockVertexAnalyzer.h"
+
+#include "HarmonixMetasound/DataTypes/MidiClock.h"
+
+#include "HarmonixMidi/BarMap.h"
+
+namespace HarmonixMetasound::Analysis
+{
+	const Frontend::FAnalyzerOutput FMidiClockVertexAnalyzer::FOutputs::Timestamp = { "Timestamp", GetMetasoundDataTypeName<FMusicTimestamp>() };
+	const Frontend::FAnalyzerOutput FMidiClockVertexAnalyzer::FOutputs::Tempo = { "Tempo", GetMetasoundDataTypeName<float>() };
+	const Frontend::FAnalyzerOutput FMidiClockVertexAnalyzer::FOutputs::TimeSignature = { "TimeSignature", GetMetasoundDataTypeName<FTimeSignature>() };
+	const Frontend::FAnalyzerOutput FMidiClockVertexAnalyzer::FOutputs::Speed = { "Speed", GetMetasoundDataTypeName<float>() };
+	
+	const Metasound::Frontend::FAnalyzerOutput& FMidiClockVertexAnalyzer::FOutputs::GetValue()
+	{
+		return Timestamp;
+	}
+	
+	const TArray<Metasound::Frontend::FAnalyzerOutput>& FMidiClockVertexAnalyzer::FFactory::GetAnalyzerOutputs() const
+	{
+		static const TArray<Frontend::FAnalyzerOutput> Outputs
+		{
+			FOutputs::Timestamp,
+			FOutputs::Tempo,
+			FOutputs::TimeSignature,
+			FOutputs::Speed
+		};
+		return Outputs;
+	}
+
+	const FName& FMidiClockVertexAnalyzer::GetAnalyzerName()
+	{
+		static const FName Name = "Harmonix.MidiClock";
+		return Name;
+	}
+
+	const FName& FMidiClockVertexAnalyzer::GetDataType()
+	{
+		return Metasound::GetMetasoundDataTypeName<FMidiClock>();
+	}
+
+	FMidiClockVertexAnalyzer::FMidiClockVertexAnalyzer(const Metasound::Frontend::FCreateAnalyzerParams& InParams)
+		: FVertexAnalyzerBase(InParams.AnalyzerAddress, InParams.VertexDataReference)
+		, Timestamp(FMusicTimestampWriteRef::CreateNew())
+		, Tempo(FFloatWriteRef::CreateNew())
+		, TimeSignature(FTimeSignatureWriteRef::CreateNew())
+		, Speed(FFloatWriteRef::CreateNew())
+	{
+		BindOutputData<FMusicTimestamp>(FOutputs::Timestamp.Name, InParams.OperatorSettings, Timestamp);
+		BindOutputData<float>(FOutputs::Tempo.Name, InParams.OperatorSettings, Tempo);
+		BindOutputData<FTimeSignature>(FOutputs::TimeSignature.Name, InParams.OperatorSettings, TimeSignature);
+		BindOutputData<float>(FOutputs::Speed.Name, InParams.OperatorSettings, Speed);
+	}
+
+	void FMidiClockVertexAnalyzer::Execute()
+	{
+		const FMidiClock& Clock = GetVertexData<FMidiClock>();
+		*Timestamp = Clock.GetCurrentMusicTimestamp();
+		*Tempo = Clock.GetTempoAtEndOfBlock();
+		*TimeSignature = Clock.GetBarMap().GetTimeSignatureAtBar((*Timestamp).Bar);
+		*Speed = Clock.GetSpeedAtEndOfBlock();
+		MarkOutputDirty();
+	}
+}
diff --git a/Engine/Plugins/Runtime/Harmonix/Source/HarmonixMetasound/Private/HarmonixMetasound/Common.cpp b/Engine/Plugins/Runtime/Harmonix/Source/HarmonixMetasound/Private/HarmonixMetasound/Common.cpp
new file mode 100644
--- /dev/null
+++ b/Engine/Plugins/Runtime/Harmonix/Source/HarmonixMetasound/Private/HarmonixMetasound/Common.cpp
@@ -0,0 +1,108 @@
+// Copyright Epic Games, Inc. All Rights Reserved.
+
+#include "HarmonixMetasound/Common.h"
+
+#include "MetasoundNodeRegistrationMacro.h"
+
+#define LOCTEXT_NAMESPACE "HarmonixMetaSound"
+
+namespace HarmonixMetasound
+{
+	const FName HarmonixNodeNamespace = TEXT("HarmonixNodes");
+
+	namespace MetasoundNodeCategories
+	{
+		const FText Harmonix = METASOUND_LOCTEXT("Metasound_HarmonixCategory", "Harmonix");
+		const FText Modulation = METASOUND_LOCTEXT("Metasound_ModulationCategory", "Modulation");
+	}
+	
+	namespace CommonPinNames
+	{
+		namespace Inputs
+		{
+			EXTERN_METASOUND_PARAM_D(Enable,                        "Enable",                             "Enables processing.")
+			EXTERN_METASOUND_PARAM_D(MidiStream,                    "Midi Stream",                        "Midi event input stream.")
+			EXTERN_METASOUND_PARAM_D(MidiChannelNumber,             "Midi Channel",                       "Midi channel to render (0 for all/omni).")
+			EXTERN_METASOUND_PARAM_D(MidiChannelFilterSpecifier,    "Midi Channel Filter",                "Midi channels to process. '*' for all, or a string like '1,3,4-8' to identify specific channels.")
+			EXTERN_METASOUND_PARAM_D(MidiTrackNumber,               "Track Number",                       "Track number (1 based).")
+			EXTERN_METASOUND_PARAM_D(MidiTrackIndexFilterSpecifier, "Midi Track Filter",                  "Midi tracks to process. '*' for all, or a string like '1,3,4-8' to identify specific tracks.")
+			EXTERN_METASOUND_PARAM_D(MinMidiNote,                   "Min Note #",                         "Minimum midi note number (0-127).")
+			EXTERN_METASOUND_PARAM_D(MaxMidiNote,                   "Max Note #",                         "Maximum midi note number (0-127).")
+			EXTERN_METASOUND_PARAM_D(MinMidiVelocity,               "Min Velocity",                       "Minimum Velocity (0-127).")
+			EXTERN_METASOUND_PARAM_D(MaxMidiVelocity,               "Max Velocity",                       "Maximum Velocity (0-127).")
+			EXTERN_METASOUND_PARAM_D(MidiFileAsset,                 "Midi File",                          "Standard Midi file.")
+			EXTERN_METASOUND_PARAM_D(Transport,                     "Transport",                          "Play, pause, continue, stop, etc.")
+			EXTERN_METASOUND_PARAM_D(MidiClock,                     "Midi Clock",                         "Midi timing information.")
+			EXTERN_METASOUND_PARAM_D(Tempo,                         "Tempo",                              "Tempo in BPM. This is MIDI bpm... ALWAYS QUARTER NOTES PER MINUTE regardless of time signature!")
+			EXTERN_METASOUND_PARAM_D(TimeSigNumerator,	            "Time Sig. Numerator",                "Time Signature Numerator")
+			EXTERN_METASOUND_PARAM_D(TimeSigDenominator,            "Time Sig. Denominator",              "TimeSignature Denominator")
+			EXTERN_METASOUND_PARAM_D(Speed,                         "Speed",                              "Playback speed multiplier.")
+			EXTERN_METASOUND_PARAM_D(Loop,				            "Simple Loop",			              "Whether looping is enabled for this node. Loop length will be rounded to the nearest bar.")
+			EXTERN_METASOUND_PARAM_D(LoopLengthBars,                "Loop Length (Bars)",                 "The number of bars to loop for simple looping")
+			EXTERN_METASOUND_PARAM_D(TransportPrepare,              "Prepare",                            "Causes receiving nodes to prepare to play (pre-buffer, cache, etc.)")
+			EXTERN_METASOUND_PARAM_D(TransportPlay,                 "Play",                               "Causes a transport play request.")
+			EXTERN_METASOUND_PARAM_D(TransportPause,                "Pause",                              "Causes a transport pause request.")
+			EXTERN_METASOUND_PARAM_D(TransportContinue,             "Continue",                           "Causes a transport contiunue request.")
+			EXTERN_METASOUND_PARAM_D(TransportStop,                 "Stop",                               "Causes a transport stop request.")
+			EXTERN_METASOUND_PARAM_D(TransportKill,                 "Kill",                               "Causes a transport kill request.")
+			EXTERN_METASOUND_PARAM_D(TriggerSeek,                   "Trigger Seek",                       "Causes a transport seek request. Seek destination provided on a separate input.")
+			EXTERN_METASOUND_PARAM_D(SeekDestination,               "Seek Target",                        "Destination for the next seek when it is triggered.")
+			EXTERN_METASOUND_PARAM_D(PrerollBars,                   "Pre-roll Bars",                      "Number of bars to look back when seeking midi to find note-on messages that should hold over.")
+			EXTERN_METASOUND_PARAM_D(GridSizeUnits,                 "Base Grid Size",                     "Base size of the qrid squares in musical terms.")
+			EXTERN_METASOUND_PARAM_D(GridSizeMult,                  "Grid Size Multiplier",               "Grid size multiplier.")
+			EXTERN_METASOUND_PARAM_D(OffsetUnits,                   "Base Offset Size ",                  "Base size of the qrid offset.")
+			EXTERN_METASOUND_PARAM_D(OffsetMult,                    "Offset Multiplier",                  "Offset multiplier.")
+			EXTERN_METASOUND_PARAM_D(Bar,		   	                "Bar",                                "Musical bar (or measure) index (0 based!)")
+			EXTERN_METASOUND_PARAM_D(FloatBeat,			            "Beat",                               "Musical beat index. 0 based and in units of the time signature denomninator! The value is rounded to the nearest 32nd note.")
+			EXTERN_METASOUND_PARAM_D(SynthPatch,                    "Patch",                              "Synthesizer patch asset.")
+			EXTERN_METASOUND_PARAM_D(Transposition,                 "Transposition",                      "Transposition in semitones.")
+			EXTERN_METASOUND_PARAM_D(ClockSpeedToPitch,             "Midi Clock Speed Affects Pitch",     "When true the clock speed of the incoming midi stream will affect the pitch of the synth.")
+			EXTERN_METASOUND_PARAM_D(ClockSpeedToFrequency,         "Midi Clock Speed Affects Frequency", "When true the clock speed of the incoming midi stream will affect the calculated frequency.")
+			EXTERN_METASOUND_PARAM_D(LFOSyncType,					"LFO Sync Type",					  "Specifies if and how the LFO should sync to a clock");
+			EXTERN_METASOUND_PARAM_D(LFOFrequency,					"LFO Frequency",					  "The period of the LFO. If Sync Type is TempoSync, the unit is cycles per quarter note. Otherwise, the unit is Hz (cycles per second).");
+			EXTERN_METASOUND_PARAM_D(LFOInvert,						"Invert LFO",						  "Toggle to invert the LFO");
+			EXTERN_METASOUND_PARAM_D(LFOShape,						"LFO Shape",						  "The shape of the LFO. 0.0-1.0 = morph between square and triangle. 1.0-2.0 = morph between triangle and sawtooth");
+			EXTERN_METASOUND_PARAM_D(AudioMono,						"Audio Mono",						  "The mono audio input");
+		}
+		namespace Outputs
+		{
+			EXTERN_METASOUND_PARAM_D(Transport,          "Transport",                   "Play, pause, continue, stop, etc.")
+			EXTERN_METASOUND_PARAM_D(MidiClock,          "Midi Clock",                  "Midi timing information.")
+			EXTERN_METASOUND_PARAM_D(MidiFileAsset,      "Midi File",                   "Standard Midi file.")
+			EXTERN_METASOUND_PARAM_D(MidiStream,         "Midi Stream",                 "Midi event output stream.")
+			EXTERN_METASOUND_PARAM_D(NoteOn,             "Note On",                     "Midi note on.")
+			EXTERN_METASOUND_PARAM_D(NoteOff,            "Note Off",                    "Midi note off.")
+			EXTERN_METASOUND_PARAM_D(MidiNoteNumber,     "Midi Note #",                 "Midi note number (0 - 127).")
+			EXTERN_METASOUND_PARAM_D(Frequency,          "Frequency",                   "Midi note number converted to frequency.")
+			EXTERN_METASOUND_PARAM_D(MidiVelocity,       "Velocity",                    "Midi velocity (0 - 127).")
+			EXTERN_METASOUND_PARAM_D(NormalizedVelocity, "Normalized Velocity",         "Velocity (0.0 - 1.0).")
+			EXTERN_METASOUND_PARAM_D(TransportPrepare,   "Prepare",                     "Causes receiving nodes to prepare to play (pre-buffer, cache, etc.)")
+			EXTERN_METASOUND_PARAM_D(TransportPlay,      "Play",                        "Causes a transport play request.")
+			EXTERN_METASOUND_PARAM_D(TransportPause,     "Pause",                       "Causes a transport pause request.")
+			EXTERN_METASOUND_PARAM_D(TransportContinue,  "Continue",                    "Causes a transport continue request.")
+			EXTERN_METASOUND_PARAM_D(TransportStop,      "Stop",                        "Causes a transport stop request.")
+			EXTERN_METASOUND_PARAM_D(TransportKill,      "Kill",                        "Causes a transport kill request.")
+			EXTERN_METASOUND_PARAM_D(MusicTimestamp,     "Music Timestamp",             "A single structure representing a musical time (bar & beat).")
+			EXTERN_METASOUND_PARAM_D(SeekTarget,         "Seek Target",                 "A well formed musical seek target.")
+			EXTERN_METASOUND_PARAM_D(AudioMono,          "Audio Mono",                  "The mono audio output.")
+			EXTERN_METASOUND_PARAM_D(AudioLeft,          "Audio Left",                  "The left channel audio output.")
+			EXTERN_METASOUND_PARAM_D(AudioRight,         "Audio Right",                 "The right channel audio output.")
+			EXTERN_METASOUND_PARAM_D(AudioCenter,        "Audio Center",                "5.1/7.1: The center channel audio output.")
+			EXTERN_METASOUND_PARAM_D(AudioLFE,           "Audio LFE",                   "5.1/7.1: The LFE channel audio output.")
+			EXTERN_METASOUND_PARAM_D(AudioLeftSide,      "Audio Left Side",             "Quad/5.1: Left Surround, 7.1: Left Side.")
+			EXTERN_METASOUND_PARAM_D(AudioRightSide,     "Audio Right Side",            "Quad/5.1: Right Surround, 7.1: Right Side.")
+			EXTERN_METASOUND_PARAM_D(AudioLeftBack,      "Audio Left Back - Surround",  "7.1: Left Back.")
+			EXTERN_METASOUND_PARAM_D(AudioRightBack,     "Audio Right Back - Surround", "7.1: Right Back.")
+			EXTERN_METASOUND_PARAM_D(MusicTimespanBar,   "Bar",                         "Bar (Measure). 1 based!")
+			EXTERN_METASOUND_PARAM_D(MusicTimespanBeat,  "Beat",                        "Beat. 1 based!")
+			EXTERN_METASOUND_PARAM_D(Tempo,              "Tempo",                       "BPM. This is MIDI bpm... ALWAYS QUARTER NOTES PER MINUTE regardless of time signature!")
+			EXTERN_METASOUND_PARAM_D(Speed,              "Speed",                       "Speed Multiplier")
+			EXTERN_METASOUND_PARAM_D(TimeSigNumerator,   "Time Sig. Numerator",         "Time Signature Numerator")
+			EXTERN_METASOUND_PARAM_D(TimeSigDenominator, "Time Sig. Denominator",       "TimeSignature Denominator")
+			EXTERN_METASOUND_PARAM_D(SecsPerQuarter,     "Secs Per Quarter Note",       "Seconds per quarter note.")
+			EXTERN_METASOUND_PARAM_D(SecsPerBeat,        "Secs Per Beat",               "Seconds per beat. Equals seconds per beat if time signature denominator is 4.")
+		}
+	}
+}
+
+#undef LOCTEXT_NAMESPACE // "HarmonixMetaSound"
diff --git a/Engine/Plugins/Runtime/Harmonix/Source/HarmonixMetasound/Private/HarmonixMetasound/Components/MetasoundMusicClockDriver.cpp b/Engine/Plugins/Runtime/Harmonix/Source/HarmonixMetasound/Private/HarmonixMetasound/Components/MetasoundMusicClockDriver.cpp
new file mode 100644
--- /dev/null
+++ b/Engine/Plugins/Runtime/Harmonix/Source/HarmonixMetasound/Private/HarmonixMetasound/Components/MetasoundMusicClockDriver.cpp
@@ -0,0 +1,266 @@
+// Copyright Epic Games, Inc. All Rights Reserved.
+#include "MetasoundMusicClockDriver.h"
+#include "MetasoundGeneratorHandle.h"
+#include "Components/AudioComponent.h"
+#include "HarmonixMetasound/DataTypes/MidiClock.h"
+#include "MetasoundGeneratorHandle.h"
+#include "MetasoundGenerator.h"
+#include "Engine/World.h"
+#include "Harmonix.h"
+
+bool FMetasoundMusicClockDriver::CalculateSongPosWithOffset(float MsOffset, ECalibratedMusicTimebase Timebase, FMidiSongPos& OutResult) const
+{
+	// if we have an owner, ask them directly
+	if (const FMidiPlayCursorMgr* Owner = Cursor.GetOwner())
+	{
+		float RawMs = Owner->GetCurrentHiResMs();
+		switch (Timebase)
+		{
+		case ECalibratedMusicTimebase::AudioRenderTime:
+			OutResult = Owner->CalculateSongPosWithOffsetMs((Clock->CurrentSmoothedAudioRenderSongPos.SecondsIncludingCountIn * 1000.0f) + MsOffset - RawMs);
+			break;
+		case ECalibratedMusicTimebase::ExperiencedTime:
+			OutResult = Owner->CalculateSongPosWithOffsetMs((Clock->CurrentPlayerExperiencedSongPos.SecondsIncludingCountIn * 1000.0f) + MsOffset - RawMs);
+			break;
+		case ECalibratedMusicTimebase::VideoRenderTime:
+		default:
+			OutResult = Owner->CalculateSongPosWithOffsetMs((Clock->CurrentVideoRenderSongPos.SecondsIncludingCountIn * 1000.0f) + MsOffset - RawMs);
+			break;
+		}
+		return true;
+	}
+	return false;
+}
+
+bool FMetasoundMusicClockDriver::RefreshCurrentSongPos()
+{
+	//	Only for use when on the game thread.
+	if (ensureMsgf(
+			IsInGameThread(),
+			TEXT("%hs called from non-game thread.  This is not supported"),
+			__FUNCTION__) == false)
+	{
+		return false;
+	}
+
+	if (AudioComponentToWatch.IsValid())
+	{
+		if (!CurrentGeneratorHandle)
+		{
+			AttemptToConnectToAudioComponentsMetasound();
+		}
+		else if (!Cursor.GetOwner())
+		{
+			TryToRegisterPlayCursor();
+		}
+	}
+
+	if (Cursor.GetOwner())
+	{
+		// cursor is attached and has the current info
+		RefreshCurrentSongPosFromCursor();
+		return true;
+	}
+	else
+	{
+		// Cursor not attached so use wall clock
+		if (!WasEverConnected || Clock->RunPastMusicEnd)
+		{
+			RefreshCurrentSongPosFromWallClock();
+			return true;
+		}
+	}
+	return false;
+}
+
+void FMetasoundMusicClockDriver::OnStart()
+{
+	SongPosOffsetMs = 0.0f;
+	FreeRunStartTimeSecs = Clock ? Clock->GetWorld()->GetTimeSeconds() : 0.0;
+}
+
+void FMetasoundMusicClockDriver::OnContinue()
+{
+	if (!Cursor.GetOwner())
+	{
+		RefreshCurrentSongPosFromWallClock();
+	}
+}
+
+void FMetasoundMusicClockDriver::Disconnect()
+{
+	if (FMidiPlayCursorMgr* OldOwner = Cursor.GetOwner())
+	{
+		OldOwner->UnregisterPlayCursor(&Cursor);
+	}
+	DetachAllCallbacks();
+	AudioComponentToWatch.Reset();
+	CurrentGeneratorHandle.Reset();
+}
+
+const FSongMaps* FMetasoundMusicClockDriver::GetCurrentSongMaps() const
+{
+	if (Cursor.GetOwner())
+	{
+		return &Cursor.GetOwner()->GetSongMaps();
+	}
+	return &Clock->DefaultMaps;
+}
+
+bool FMetasoundMusicClockDriver::ConnectToAudioComponentsMetasound(UAudioComponent* InAudioComponent, FName MetasoundOuputPinName)
+{
+	AudioComponentToWatch = InAudioComponent;
+	MetasoundOutputName = MetasoundOuputPinName;
+	return AttemptToConnectToAudioComponentsMetasound();
+}
+
+bool FMetasoundMusicClockDriver::AttemptToConnectToAudioComponentsMetasound()
+{
+	if (!AudioComponentToWatch.IsValid() || MetasoundOutputName.IsNone())
+	{
+		return false;
+	}
+
+	DetachAllCallbacks();
+
+	CurrentGeneratorHandle.Reset(UMetasoundGeneratorHandle::CreateMetaSoundGeneratorHandle(AudioComponentToWatch.Get()));
+	if (!CurrentGeneratorHandle)
+	{
+		return false;
+	}
+	GeneratorAttachedCallbackHandle = CurrentGeneratorHandle->OnGeneratorHandleAttached.AddLambda([this](){OnGeneratorAttached();});
+	GeneratorDetachedCallbackHandle = CurrentGeneratorHandle->OnGeneratorHandleDetached.AddLambda([this](){OnGeneratorDetached();});
+	UMetasoundGeneratorHandle::FOnSetGraph::FDelegate OnSetGraph;
+	OnSetGraph.BindLambda([this](){OnGraphSet();});
+	GraphChangedCallbackHandle = CurrentGeneratorHandle->AddGraphSetCallback(MoveTemp(OnSetGraph));
+	OnGeneratorAttached();
+	return true;
+}
+
+void FMetasoundMusicClockDriver::DetachAllCallbacks()
+{
+	if (CurrentGeneratorHandle)
+	{
+		CurrentGeneratorHandle->OnGeneratorHandleAttached.Remove(GeneratorAttachedCallbackHandle);
+		GeneratorAttachedCallbackHandle.Reset();
+		CurrentGeneratorHandle->OnGeneratorHandleDetached.Remove(GeneratorDetachedCallbackHandle);
+		GeneratorDetachedCallbackHandle.Reset();
+		CurrentGeneratorHandle->RemoveGraphSetCallback(GraphChangedCallbackHandle);
+		GraphChangedCallbackHandle.Reset();
+	}
+}
+
+void FMetasoundMusicClockDriver::OnGeneratorAttached()
+{
+	TryToRegisterPlayCursor();
+}
+
+void FMetasoundMusicClockDriver::OnGraphSet()
+{
+	TryToRegisterPlayCursor();
+}
+
+void FMetasoundMusicClockDriver::OnGeneratorDetached()
+{
+	check(Clock);
+	check(IsInGameThread());
+	if (FMidiPlayCursorMgr* OldOwner = Cursor.GetOwner())
+	{
+		OldOwner->UnregisterPlayCursor(&Cursor);
+		if (Clock->GetState() != EMusicClockState::Stopped)
+		{
+			Clock->DefaultMaps.Copy(OldOwner->GetSongMaps(), 0, Cursor.GetCurrentTick());
+			SongPosOffsetMs = Clock->CurrentSmoothedAudioRenderSongPos.SecondsIncludingCountIn * 1000.0f;
+			FreeRunStartTimeSecs = Clock->GetWorld()->GetTimeSeconds();
+		}
+		Clock->MusicClockDisconnectedEvent.Broadcast();
+	}
+}
+
+void FMetasoundMusicClockDriver::TryToRegisterPlayCursor()
+{
+	check(Clock);
+	check(IsInGameThread());
+	if (FMidiPlayCursorMgr* OldOwner = Cursor.GetOwner())
+	{
+		OldOwner->UnregisterPlayCursor(&Cursor);
+	}
+	if (CurrentGeneratorHandle && !MetasoundOutputName.IsNone())
+	{
+		TSharedPtr<Metasound::FMetasoundGenerator> LowLevelGenerator = CurrentGeneratorHandle->GetGenerator();
+		if (LowLevelGenerator.IsValid())
+		{
+			const TOptional<Metasound::TDataReadReference<HarmonixMetasound::FMidiClock>> MidiClockRef = LowLevelGenerator->GetOutputReadReference<HarmonixMetasound::FMidiClock>(MetasoundOutputName);
+			const Metasound::TDataReadReference<HarmonixMetasound::FMidiClock>* MidiClock = MidiClockRef.GetPtrOrNull();
+			if (MidiClock)
+			{
+				(*MidiClock)->RegisterLowResPlayCursor(&Cursor);
+				WasEverConnected = true;
+				Clock->MusicClockConnectedEvent.Broadcast();
+			}
+			else
+			{
+				UE_LOG(LogMusicClock, Verbose, TEXT("Didn't find Midi Clock output named \"%s\" in the Metasound!"), *MetasoundOutputName.ToString());
+			}
+		}
+	}
+}
+
+void FMetasoundMusicClockDriver::RefreshCurrentSongPosFromWallClock()
+{
+	check(Clock);
+
+	bool TempoChanged = Clock->CurrentSmoothedAudioRenderSongPos.Tempo != Clock->Tempo;
+
+	double FreeRunTime = (Clock->GetWorld()->GetTimeSeconds() - FreeRunStartTimeSecs) * Clock->CurrentClockAdvanceRate;
+
+	Clock->CurrentSmoothedAudioRenderSongPos.SetByTime(((float)FreeRunTime * 1000.0) + SongPosOffsetMs, Clock->DefaultMaps);
+	Clock->CurrentPlayerExperiencedSongPos.SetByTime(Clock->CurrentSmoothedAudioRenderSongPos.SecondsIncludingCountIn * 1000.0f - FHarmonixModule::GetMeasuredUserExperienceAndReactionToAudioRenderOffsetMs(), Clock->DefaultMaps);
+	Clock->CurrentVideoRenderSongPos.SetByTime(Clock->CurrentSmoothedAudioRenderSongPos.SecondsIncludingCountIn * 1000.0f - FHarmonixModule::GetMeasuredVideoToAudioRenderOffsetMs(), Clock->DefaultMaps);
+	if (Clock->CurrentSmoothedAudioRenderSongPos.SecondsIncludingCountIn > Clock->RawUnsmoothedAudioRenderPos.SecondsIncludingCountIn)
+	{
+		Clock->RawUnsmoothedAudioRenderPos.SetByTime(Clock->CurrentSmoothedAudioRenderSongPos.SecondsIncludingCountIn * 1000.0f, Clock->DefaultMaps);
+	}
+
+	if (TempoChanged)
+	{
+		Clock->Tempo = Clock->CurrentSmoothedAudioRenderSongPos.Tempo;
+		Clock->CurrentBeatDurationSec = (60.0f / Clock->Tempo) / Clock->CurrentClockAdvanceRate;
+		Clock->CurrentBarDurationSec = ((Clock->TimeSignatureNum * Clock->CurrentBeatDurationSec) / (Clock->TimeSignatureDenom / 4.0f)) / Clock->CurrentClockAdvanceRate;
+	}
+}
+
+void FMetasoundMusicClockDriver::RefreshCurrentSongPosFromCursor()
+{
+	check(Clock);
+	Clock->CurrentSmoothedAudioRenderSongPos = Cursor.GetCurrentSongPos();
+
+	float PrevClockAdvanceRate = Clock->CurrentClockAdvanceRate;
+	if (const FMidiPlayCursorMgr* Owner = Cursor.GetOwner())
+	{
+		Clock->CurrentClockAdvanceRate = Owner->GetCurrentAdvanceRate();
+		float RawMs = Owner->GetCurrentHiResMs();
+		Clock->CurrentPlayerExperiencedSongPos = Owner->CalculateSongPosWithOffsetMs(Clock->CurrentSmoothedAudioRenderSongPos.SecondsIncludingCountIn * 1000.0f - FHarmonixModule::GetMeasuredUserExperienceAndReactionToAudioRenderOffsetMs() - RawMs);
+		Clock->CurrentVideoRenderSongPos = Owner->CalculateSongPosWithOffsetMs(Clock->CurrentSmoothedAudioRenderSongPos.SecondsIncludingCountIn * 1000.0f - FHarmonixModule::GetMeasuredVideoToAudioRenderOffsetMs() - RawMs);
+		Clock->RawUnsmoothedAudioRenderPos = Owner->CalculateSongPosWithOffsetMs(0.0f);
+	}
+	else
+	{
+		Clock->CurrentPlayerExperiencedSongPos.SetByTime(Clock->CurrentSmoothedAudioRenderSongPos.SecondsIncludingCountIn * 1000.0f - FHarmonixModule::Get().GetMeasuredUserExperienceAndReactionToAudioRenderOffsetMs(), Clock->DefaultMaps);
+		Clock->CurrentVideoRenderSongPos.SetByTime(Clock->CurrentSmoothedAudioRenderSongPos.SecondsIncludingCountIn * 1000.0f - FHarmonixModule::Get().GetMeasuredVideoToAudioRenderOffsetMs(), Clock->DefaultMaps);
+		if (Clock->CurrentSmoothedAudioRenderSongPos.SecondsIncludingCountIn > Clock->RawUnsmoothedAudioRenderPos.SecondsIncludingCountIn)
+		{
+			Clock->RawUnsmoothedAudioRenderPos.SetByTime(Clock->CurrentSmoothedAudioRenderSongPos.SecondsIncludingCountIn * 1000.0f, Clock->DefaultMaps);
+		}
+	}
+
+	Clock->TimeSignatureNum = Clock->CurrentSmoothedAudioRenderSongPos.TimeSigNumerator;
+	Clock->TimeSignatureDenom = Clock->CurrentSmoothedAudioRenderSongPos.TimeSigDenominator;
+
+	if (Clock->Tempo != Clock->CurrentSmoothedAudioRenderSongPos.Tempo || PrevClockAdvanceRate != Clock->CurrentClockAdvanceRate)
+	{
+		Clock->Tempo = Clock->CurrentSmoothedAudioRenderSongPos.Tempo;
+		Clock->CurrentBeatDurationSec = (60.0f / Clock->Tempo) / Clock->CurrentClockAdvanceRate;
+		Clock->CurrentBarDurationSec = ((Clock->TimeSignatureNum * Clock->CurrentBeatDurationSec) / (Clock->TimeSignatureDenom / 4.0f)) / Clock->CurrentClockAdvanceRate;
+	}
+}
diff --git a/Engine/Plugins/Runtime/Harmonix/Source/HarmonixMetasound/Private/HarmonixMetasound/Components/MetasoundMusicClockDriver.h b/Engine/Plugins/Runtime/Harmonix/Source/HarmonixMetasound/Private/HarmonixMetasound/Components/MetasoundMusicClockDriver.h
new file mode 100644
--- /dev/null
+++ b/Engine/Plugins/Runtime/Harmonix/Source/HarmonixMetasound/Private/HarmonixMetasound/Components/MetasoundMusicClockDriver.h
@@ -0,0 +1,54 @@
+// Copyright Epic Games, Inc. All Rights Reserved.
+#pragma once
+#include "HarmonixMetasound/Components/MusicClockComponent.h"
+#include "MetasoundGeneratorHandle.h"
+#include "UObject/StrongObjectPtr.h"
+
+struct FMetasoundMusicClockDriver : public FMusicClockDriverBase
+{
+public:
+	FMetasoundMusicClockDriver(UMusicClockComponent* InClock)
+		: FMusicClockDriverBase(InClock)
+	{}
+
+	virtual bool CalculateSongPosWithOffset(float MsOffset, ECalibratedMusicTimebase Timebase, FMidiSongPos& OutResult) const override;
+
+	virtual void Disconnect() override;
+	virtual bool RefreshCurrentSongPos() override;
+	virtual void OnStart() override;
+	virtual void OnPause() override {}
+	virtual void OnContinue() override;
+	virtual void OnStop() override {}
+	virtual const FSongMaps* GetCurrentSongMaps() const override;
+
+
+	bool ConnectToAudioComponentsMetasound(UAudioComponent* InAudioComponent, FName MetasoundOuputPinName = "Midi Clock");
+
+protected:
+	void OnGeneratorAttached();
+	void OnGeneratorDetached();
+	void OnGraphSet();
+
+private:
+	FName MetasoundOutputName;
+	// We can keep a week reference to this because our "owner" is a UClass and
+	// also has a reference to it...
+	TWeakObjectPtr<UAudioComponent> AudioComponentToWatch;
+	// We need a strong object ptr to this next thing since we will be the only one holding a reference to it...
+	TStrongObjectPtr<UMetasoundGeneratorHandle> CurrentGeneratorHandle;
+
+	FSmoothedMidiPlayCursor Cursor;
+	double FreeRunStartTimeSecs = 0.0;
+	bool WasEverConnected = false;
+	float SongPosOffsetMs = 0.0f;
+
+	FDelegateHandle GeneratorAttachedCallbackHandle;
+	FDelegateHandle GeneratorDetachedCallbackHandle;
+	FDelegateHandle GraphChangedCallbackHandle;
+
+	bool AttemptToConnectToAudioComponentsMetasound();
+	void DetachAllCallbacks();
+	void TryToRegisterPlayCursor();
+	void RefreshCurrentSongPosFromWallClock();
+	void RefreshCurrentSongPosFromCursor();
+};
diff --git a/Engine/Plugins/Runtime/Harmonix/Source/HarmonixMetasound/Private/HarmonixMetasound/Components/MusicClockComponent.cpp b/Engine/Plugins/Runtime/Harmonix/Source/HarmonixMetasound/Private/HarmonixMetasound/Components/MusicClockComponent.cpp
new file mode 100644
--- /dev/null
+++ b/Engine/Plugins/Runtime/Harmonix/Source/HarmonixMetasound/Private/HarmonixMetasound/Components/MusicClockComponent.cpp
@@ -0,0 +1,601 @@
+// Copyright Epic Games, Inc. All Rights Reserved.
+#include "HarmonixMetasound/Components/MusicClockComponent.h"
+#include "HarmonixMetasound/Components/MetasoundMusicClockDriver.h"
+#include "HarmonixMetasound/Components/WallClockMusicClockDriver.h"
+#include "Harmonix.h"
+#include "HarmonixMidi/MusicTimeSpan.h"
+#include "Components/AudioComponent.h"
+#include "Templates/SharedPointer.h"
+#include "Misc/Optional.h"
+
+#include UE_INLINE_GENERATED_CPP_BY_NAME(MusicClockComponent)
+
+DEFINE_LOG_CATEGORY(LogMusicClock)
+
+UMusicClockComponent::UMusicClockComponent()
+{
+	MakeDefaultSongMap();
+	PrimaryComponentTick.bCanEverTick = true;
+	// We want the music player and clocks to tick before other components, which default to the TG_DuringPhysics group.
+	// Though, this means any other TG_PrePhysics group actors or components that care about accurate song time will
+	// want to manually add a tick prerequisite on this component.
+	PrimaryComponentTick.TickGroup = TG_PrePhysics;
+}
+
+UMusicClockComponent* UMusicClockComponent::CreateMetasoundDrivenMusicClock(UObject* WorldContextObject, UAudioComponent* InAudioComponent, FName MetasoundOutputPinName, bool Start)
+{
+	UMusicClockComponent* NewClock = NewObject<UMusicClockComponent>(WorldContextObject);
+	if (!NewClock->GetOwner())
+	{
+		UE_LOG(LogMusicClock, Warning, TEXT("Can't create a music clock in a non-actor context!"));
+		return nullptr;
+	}
+	NewClock->RegisterComponent();
+	NewClock->DriveMethod = EMusicClockDriveMethod::MetaSound;
+	NewClock->MetasoundOutputName = MetasoundOutputPinName;
+	NewClock->ConnectToMetasoundOnAudioComponent(InAudioComponent);
+	if (Start)
+	{
+		NewClock->Start();
+	}
+	return NewClock;
+}
+
+UMusicClockComponent* UMusicClockComponent::CreateWallClockDrivenMusicClock(UObject* WorldContextObject, UMidiFile* InTempoMap, bool Start)
+{
+	UMusicClockComponent* NewClock = NewObject<UMusicClockComponent>(WorldContextObject);
+	if (!NewClock->GetOwner())
+	{
+		UE_LOG(LogMusicClock, Warning, TEXT("Can't create a music clock in a non-actor context!"));
+		return nullptr;
+	}
+	NewClock->RegisterComponent();
+	NewClock->ConnectToWallClockForMidi(InTempoMap);
+	if (Start)
+	{
+		NewClock->Start();
+	}
+	return NewClock;
+}
+
+bool UMusicClockComponent::ConnectToMetasoundOnAudioComponent(UAudioComponent* InAudioComponent)
+{
+	DriveMethod = EMusicClockDriveMethod::MetaSound;
+	MetasoundsAudioComponent = InAudioComponent;
+	return ConnectToMetasound();
+}
+
+void UMusicClockComponent::ConnectToWallClockForMidi(UMidiFile* InTempoMap)
+{
+	DriveMethod = EMusicClockDriveMethod::WallClock;
+	TempoMap = InTempoMap;
+	ConnectToWallClock();
+}
+
+void UMusicClockComponent::UpdateClockIfNeeded()
+{
+	//	Not for use outside the game thread.
+	if (ensureMsgf(
+		IsInGameThread(),
+		TEXT("%hs called from non-game thread.  This is not supported!"), __FUNCTION__) == false)
+	{
+		return;
+	}
+
+	if (GFrameCounter == LastUpdateFrame)
+	{
+		return;
+	}
+
+	//	Run the actual clock update.
+	if (State == EMusicClockState::Running && ClockDriver)
+	{
+		// Make note of what will soon become the "previous position"...	
+		PrevAudioRenderSongPos = CurrentSmoothedAudioRenderSongPos;
+		PrevPlayerExperiencedSongPos = CurrentPlayerExperiencedSongPos;
+		PrevVideoRenderSongPos = CurrentVideoRenderSongPos;
+
+		if (ClockDriver->RefreshCurrentSongPos())
+		{
+			CalcDelta();
+			LastUpdateFrame = GFrameCounter;
+		}
+	}
+}
+
+void UMusicClockComponent::CreateClockDriver()
+{
+	if (DriveMethod == EMusicClockDriveMethod::WallClock || !::IsValid(MetasoundsAudioComponent))
+	{
+		ConnectToWallClock();
+	}
+	else
+	{
+		ConnectToMetasound();
+	}
+}
+
+bool UMusicClockComponent::ConnectToMetasound()
+{
+	check(DriveMethod == EMusicClockDriveMethod::MetaSound);
+	if (!::IsValid(MetasoundsAudioComponent))
+	{
+		return false;
+	}
+	if (ClockDriver)
+	{
+		ClockDriver->Disconnect();
+		ClockDriver = nullptr;
+	}
+	TUniquePtr<FMetasoundMusicClockDriver> MetasoundClockDriver = MakeUnique<FMetasoundMusicClockDriver>(this);
+	bool Connected = MetasoundClockDriver->ConnectToAudioComponentsMetasound(MetasoundsAudioComponent, MetasoundOutputName);
+	ClockDriver = MoveTemp(MetasoundClockDriver);
+	return Connected;
+}
+
+void UMusicClockComponent::ConnectToWallClock()
+{
+	// we don't 'check' that the driver mode is wall clock here because if the driver mode is metasound and we can't 
+	// connect for some reason we will fall back to this clock driver!
+	if (ClockDriver)
+	{
+		ClockDriver->Disconnect();
+	}
+	ClockDriver = MakeUnique<FWallClockMusicClockDriver>(this, TempoMap);
+}
+
+FMidiSongPos UMusicClockComponent::CalculateSongPosWithOffset(float MsOffset, ECalibratedMusicTimebase Timebase)
+{
+	check(ClockDriver);
+
+	UpdateClockIfNeeded();
+
+	FMidiSongPos Result;
+	if (ClockDriver->CalculateSongPosWithOffset(MsOffset, Timebase, Result))
+	{
+		return Result;
+	}
+
+	// otherwise, use our song maps copy
+	const FSongMaps* Maps = &DefaultMaps;
+	switch (Timebase)
+	{
+	case ECalibratedMusicTimebase::AudioRenderTime:
+		Result.SetByTime((CurrentSmoothedAudioRenderSongPos.SecondsIncludingCountIn * 1000.0f) + MsOffset, *Maps);
+		break;
+	case ECalibratedMusicTimebase::ExperiencedTime:
+		Result.SetByTime((CurrentPlayerExperiencedSongPos.SecondsIncludingCountIn * 1000.0f) + MsOffset, *Maps);
+		break;
+	case ECalibratedMusicTimebase::VideoRenderTime:
+	default:
+		Result.SetByTime((CurrentVideoRenderSongPos.SecondsIncludingCountIn * 1000.0f) + MsOffset, *Maps);
+		break;
+	}
+
+	return Result;
+}
+
+void UMusicClockComponent::BeginPlay()
+{
+	if (!ClockDriver)
+	{
+		CreateClockDriver();
+	}
+	Super::BeginPlay();
+}
+
+void UMusicClockComponent::EndPlay(const EEndPlayReason::Type EndPlayReason)
+{
+	Super::EndPlay(EndPlayReason);
+	if (ClockDriver)
+	{
+		ClockDriver->Disconnect();
+		ClockDriver = nullptr;
+	}
+}
+
+void UMusicClockComponent::TickComponent(float DeltaTime, enum ELevelTick TickType, FActorComponentTickFunction* ThisTickFunction)
+{
+	Super::TickComponent(DeltaTime, TickType, ThisTickFunction);
+
+	UpdateClockIfNeeded();
+	Broadcast();
+}
+
+void UMusicClockComponent::Start()
+{
+	MakeDefaultSongMap();
+	if (!ClockDriver)
+	{
+		CreateClockDriver();
+	}
+	ClockDriver->OnStart();
+	LastBroadcastBeat = -1;
+	LastBroadcastBar  = -1;
+	CurrentSmoothedAudioRenderSongPos.Reset();
+	CurrentVideoRenderSongPos.Reset();
+	CurrentPlayerExperiencedSongPos.Reset();
+	RawUnsmoothedAudioRenderPos.Reset();
+	State = EMusicClockState::Running;
+	PlayStateEvent.Broadcast(State);
+}
+
+void UMusicClockComponent::Pause()
+{
+	if (State != EMusicClockState::Running)
+	{
+		return;
+	}
+
+	if (ClockDriver)
+	{
+		ClockDriver->OnPause();
+	}
+
+	State = EMusicClockState::Paused;
+	PlayStateEvent.Broadcast(State);
+}
+
+void UMusicClockComponent::Continue()
+{
+	if (State != EMusicClockState::Paused)
+	{
+		return;
+	}
+
+	if (ClockDriver)
+	{
+		ClockDriver->OnContinue();
+	}
+
+	State = EMusicClockState::Running;
+	PlayStateEvent.Broadcast(State);
+}
+
+void UMusicClockComponent::Stop()
+{
+	if (ClockDriver)
+	{
+		ClockDriver->OnStop();
+	}
+	State = EMusicClockState::Stopped;
+	CurrentSmoothedAudioRenderSongPos.Reset();
+	CurrentVideoRenderSongPos.Reset();
+	CurrentPlayerExperiencedSongPos.Reset();
+	RawUnsmoothedAudioRenderPos.Reset();
+	PlayStateEvent.Broadcast(State);
+}
+
+float UMusicClockComponent::GetSecondsIncludingCountIn(ECalibratedMusicTimebase Timebase)
+{
+	UpdateClockIfNeeded();
+	switch (Timebase)
+	{
+	case ECalibratedMusicTimebase::AudioRenderTime:	return CurrentSmoothedAudioRenderSongPos.SecondsIncludingCountIn;
+	case ECalibratedMusicTimebase::ExperiencedTime:	return CurrentPlayerExperiencedSongPos.SecondsIncludingCountIn;
+	case ECalibratedMusicTimebase::VideoRenderTime:
+	default:                                        return CurrentVideoRenderSongPos.SecondsIncludingCountIn;
+	}
+}
+
+float UMusicClockComponent::GetSecondsFromBarOne(ECalibratedMusicTimebase Timebase)
+{
+	UpdateClockIfNeeded();
+	switch (Timebase)
+	{
+	case ECalibratedMusicTimebase::AudioRenderTime:	return CurrentSmoothedAudioRenderSongPos.SecondsFromBarOne;
+	case ECalibratedMusicTimebase::ExperiencedTime:	return CurrentPlayerExperiencedSongPos.SecondsFromBarOne;
+	case ECalibratedMusicTimebase::VideoRenderTime:
+	default:                                        return CurrentVideoRenderSongPos.SecondsFromBarOne;
+	}
+}
+
+float UMusicClockComponent::GetBarsIncludingCountIn(ECalibratedMusicTimebase Timebase)
+{
+	UpdateClockIfNeeded();
+	switch (Timebase)
+	{
+	case ECalibratedMusicTimebase::AudioRenderTime:	return CurrentSmoothedAudioRenderSongPos.BarsIncludingCountIn;
+	case ECalibratedMusicTimebase::ExperiencedTime:	return CurrentPlayerExperiencedSongPos.BarsIncludingCountIn;
+	case ECalibratedMusicTimebase::VideoRenderTime:
+	default:                                        return CurrentVideoRenderSongPos.BarsIncludingCountIn;
+	}
+}
+
+float UMusicClockComponent::GetBeatsIncludingCountIn(ECalibratedMusicTimebase Timebase)
+{
+	UpdateClockIfNeeded();
+	switch (Timebase)
+	{
+	case ECalibratedMusicTimebase::AudioRenderTime: return CurrentSmoothedAudioRenderSongPos.BeatsIncludingCountIn;
+	case ECalibratedMusicTimebase::ExperiencedTime:	return CurrentPlayerExperiencedSongPos.BeatsIncludingCountIn;
+	case ECalibratedMusicTimebase::VideoRenderTime:
+	default:                                        return CurrentVideoRenderSongPos.BeatsIncludingCountIn;
+	}
+}
+
+FMusicTimestamp UMusicClockComponent::GetCurrentTimestamp(ECalibratedMusicTimebase Timebase)
+{
+	UpdateClockIfNeeded();
+	switch (Timebase)
+	{
+	case ECalibratedMusicTimebase::AudioRenderTime:	return CurrentSmoothedAudioRenderSongPos.Timestamp;
+	case ECalibratedMusicTimebase::ExperiencedTime:	return CurrentPlayerExperiencedSongPos.Timestamp;
+	case ECalibratedMusicTimebase::VideoRenderTime:
+	default:                                        return CurrentVideoRenderSongPos.Timestamp;
+	}
+}
+
+FString UMusicClockComponent::GetCurrentSectionName(ECalibratedMusicTimebase Timebase /*= ECalibratedMusicTimebase::VideoRenderTime*/)
+{
+	const FMidiSongPos& SongPos = GetSongPos(Timebase);
+	return SongPos.CurrentSongSection.Name;
+}
+
+int32 UMusicClockComponent::GetCurrentSectionIndex(ECalibratedMusicTimebase Timebase /*= ECalibratedMusicTimebase::VideoRenderTime*/)
+{
+	const FMidiSongPos& SongPos = GetSongPos(Timebase);
+	return GetSongMaps().GetSectionMap().TickToSectionIndex(SongPos.CurrentSongSection.StartTick);
+}
+
+float UMusicClockComponent::GetCurrentSectionStartMs(ECalibratedMusicTimebase Timebase /*= ECalibratedMusicTimebase::VideoRenderTime*/)
+{
+	const FMidiSongPos& SongPos = GetSongPos(Timebase);
+	return GetSongMaps().TickToMs(SongPos.CurrentSongSection.StartTick);
+}
+
+float UMusicClockComponent::GetCurrentSectionLengthMs(ECalibratedMusicTimebase Timebase /*= ECalibratedMusicTimebase::VideoRenderTime*/)
+{
+	const FMidiSongPos& SongPos = GetSongPos(Timebase);
+	return GetSongMaps().TickToMs(SongPos.CurrentSongSection.LengthTicks);
+}
+
+float UMusicClockComponent::GetDistanceFromCurrentBeat(ECalibratedMusicTimebase Timebase /*= ECalibratedMusicTimebase::VideoRenderTime*/)
+{
+	return FMath::Fractional(GetSongPos(Timebase).BeatsIncludingCountIn);
+}
+
+float UMusicClockComponent::GetDistanceToNextBeat(ECalibratedMusicTimebase Timebase /*= ECalibratedMusicTimebase::VideoRenderTime*/)
+{
+	return 1 - GetDistanceFromCurrentBeat(Timebase);
+}
+
+float UMusicClockComponent::GetDistanceToClosestBeat(ECalibratedMusicTimebase Timebase /*= ECalibratedMusicTimebase::VideoRenderTime*/)
+{
+	return FMath::Min(GetDistanceFromCurrentBeat(Timebase), GetDistanceToNextBeat(Timebase));
+}
+
+float UMusicClockComponent::GetDistanceFromCurrentBar(ECalibratedMusicTimebase Timebase /*= ECalibratedMusicTimebase::VideoRenderTime*/)
+{
+	return FMath::Fractional(GetSongPos(Timebase).BarsIncludingCountIn);
+}
+
+float UMusicClockComponent::GetDistanceToNextBar(ECalibratedMusicTimebase Timebase /*= ECalibratedMusicTimebase::VideoRenderTime*/)
+{
+	return 1 - GetDistanceFromCurrentBar(Timebase);
+}
+
+float UMusicClockComponent::GetDistanceToClosestBar(ECalibratedMusicTimebase Timebase /*= ECalibratedMusicTimebase::VideoRenderTime*/)
+{
+	return FMath::Min(GetDistanceFromCurrentBar(Timebase), GetDistanceToNextBar(Timebase));
+}
+
+float UMusicClockComponent::GetDeltaBar(ECalibratedMusicTimebase Timebase)
+{
+	UpdateClockIfNeeded();
+	switch (Timebase)
+	{
+	case ECalibratedMusicTimebase::AudioRenderTime: return AudioRenderDeltaBarF;
+	case ECalibratedMusicTimebase::ExperiencedTime:	return PlayerExperienceDeltaBarF;
+	case ECalibratedMusicTimebase::VideoRenderTime:
+	default:                                        return VideoRenderDeltaBarF;
+	}
+}
+
+float UMusicClockComponent::GetDeltaBeat(ECalibratedMusicTimebase Timebase)
+{
+	UpdateClockIfNeeded();
+	switch (Timebase)
+	{
+	case ECalibratedMusicTimebase::AudioRenderTime: return AudioRenderDeltaBeatF;
+	case ECalibratedMusicTimebase::ExperiencedTime:	return PlayerExperienceDeltaBeatF;
+	case ECalibratedMusicTimebase::VideoRenderTime:
+	default:                                        return VideoRenderDeltaBeatF;
+	}
+}
+
+const TArray<FSongSection>& UMusicClockComponent::GetSongSections() const
+{
+	return GetSongMaps().GetSectionMap().GetSections();
+}
+
+float UMusicClockComponent::GetCountInSeconds() const
+{
+	return GetSongMaps().GetCountInSeconds();
+}
+
+float UMusicClockComponent::TickToMs(float Tick) const
+{
+	return GetSongMaps().TickToMs(Tick);
+}
+
+float UMusicClockComponent::BeatToMs(float Beat) const
+{
+	return GetSongMaps().GetMsAtBeat(Beat);
+}
+
+float UMusicClockComponent::GetMsPerBeatAtMs(float Ms) const
+{
+	return GetSongMaps().GetMsPerBeatAtMs(Ms);
+}
+
+float UMusicClockComponent::GetNumBeatsInBarAtMs(float Ms) const
+{
+	return GetSongMaps().GetNumBeatsInPulseBarAtMs(Ms);
+}
+
+float UMusicClockComponent::GetBeatInBarAtMs(float Ms) const
+{
+	return GetSongMaps().GetBeatInPulseBarAtMs(Ms);
+}
+
+float UMusicClockComponent::BarToMs(float Bar) const
+{
+	if (const FTimeSignature* TimeSigAtBar = GetSongMaps().GetTimeSignatureAtBar(static_cast<int32>(Bar)))
+	{
+		return BeatToMs(static_cast<float>(TimeSigAtBar->Numerator) * Bar);
+	}
+
+	return 0.0f;
+}
+
+float UMusicClockComponent::GetMsPerBarAtMs(float Ms) const
+{
+	return GetSongMaps().GetMsPerBarAtMs(Ms);
+}
+
+FString UMusicClockComponent::GetSectionNameAtMs(float Ms) const
+{
+	return GetSongMaps().GetSectionNameAtMs(Ms);
+}
+
+float UMusicClockComponent::GetSectionLengthMsAtMs(float Ms) const
+{
+	return GetSongMaps().GetSectionLengthMsAtMs(Ms);
+}
+
+float UMusicClockComponent::GetSectionStartMsAtMs(float Ms) const
+{
+	return GetSongMaps().GetSectionStartMsAtMs(Ms);
+}
+
+float UMusicClockComponent::GetSectionEndMsAtMs(float Ms) const
+{
+	return GetSongMaps().GetSectionEndMsAtMs(Ms);
+}
+
+float UMusicClockComponent::GetNumSections() const
+{
+	return GetSongMaps().GetSectionMap().GetNumSections();
+}
+
+float UMusicClockComponent::GetSongLengthMs() const
+{
+	return GetSongMaps().GetSongLengthMs();
+}
+
+float UMusicClockComponent::GetSongLengthBeats() const
+{
+	return GetSongMaps().GetSongLengthBeats();
+}
+
+int32 UMusicClockComponent::GetSongLengthBars() const
+{
+	return GetSongMaps().GetSongLengthBars();
+}
+
+float UMusicClockComponent::GetSongRemainingTime(ECalibratedMusicTimebase Timebase /*= ECalibratedMusicTimebase::VideoRenderTime*/)
+{
+	const float SongLengthMs = GetSongMaps().GetSongLengthMs();
+	return SongLengthMs <= 0.f ? 0.f : SongLengthMs - (GetSongPos(Timebase).SecondsIncludingCountIn * 1000.0f);
+}
+
+const FSongMaps& UMusicClockComponent::GetSongMaps() const
+{
+	const FSongMaps* SongMaps = ClockDriver ? ClockDriver->GetCurrentSongMaps() : nullptr;
+	return SongMaps ? *SongMaps : DefaultMaps;
+}
+
+const FMidiSongPos& UMusicClockComponent::GetSongPos(ECalibratedMusicTimebase Timebase)
+{
+	UpdateClockIfNeeded();
+
+	switch (Timebase)
+	{
+	case ECalibratedMusicTimebase::AudioRenderTime:
+		return CurrentSmoothedAudioRenderSongPos;
+	case ECalibratedMusicTimebase::ExperiencedTime:
+		return CurrentPlayerExperiencedSongPos;
+	case ECalibratedMusicTimebase::VideoRenderTime:
+	default:
+		return CurrentVideoRenderSongPos;
+	}
+}
+
+FMidiSongPos UMusicClockComponent::GetCurrentSmoothedAudioRenderSongPos()
+{
+	return GetSongPos(ECalibratedMusicTimebase::AudioRenderTime);
+}
+
+FMidiSongPos UMusicClockComponent::GetCurrentVideoRenderSongPos()
+{
+	return GetSongPos(ECalibratedMusicTimebase::VideoRenderTime);
+}
+
+FMidiSongPos UMusicClockComponent::GetCurrentPlayerExperiencedSongPos()
+{
+	return GetSongPos(ECalibratedMusicTimebase::ExperiencedTime);
+}
+
+float UMusicClockComponent::MeasureSpanProgress(const FMusicalTimeSpan& Span, ECalibratedMusicTimebase Timebase)
+{
+	UpdateClockIfNeeded();
+
+	const FSongMaps* Maps = ClockDriver ? ClockDriver->GetCurrentSongMaps() : &DefaultMaps;
+	if (!Maps)
+	{
+		return 0.0f;
+	}
+
+	return Span.CalcPositionInSpan(CurrentSmoothedAudioRenderSongPos, *Maps);
+}
+
+void UMusicClockComponent::CalcDelta()
+{
+	AudioRenderDeltaBarF = CurrentSmoothedAudioRenderSongPos.BarsIncludingCountIn - PrevAudioRenderSongPos.BarsIncludingCountIn;
+	AudioRenderDeltaBeatF = CurrentSmoothedAudioRenderSongPos.BeatsIncludingCountIn - PrevAudioRenderSongPos.BeatsIncludingCountIn;
+	PlayerExperienceDeltaBarF = CurrentPlayerExperiencedSongPos.BarsIncludingCountIn - PrevPlayerExperiencedSongPos.BarsIncludingCountIn;
+	PlayerExperienceDeltaBeatF = CurrentPlayerExperiencedSongPos.BeatsIncludingCountIn - PrevPlayerExperiencedSongPos.BeatsIncludingCountIn;
+	VideoRenderDeltaBarF = CurrentVideoRenderSongPos.BarsIncludingCountIn - PrevVideoRenderSongPos.BarsIncludingCountIn;
+	VideoRenderDeltaBeatF = CurrentVideoRenderSongPos.BeatsIncludingCountIn - PrevVideoRenderSongPos.BeatsIncludingCountIn;
+}
+
+void UMusicClockComponent::Broadcast()
+{
+	const FMidiSongPos* Basis = &CurrentVideoRenderSongPos;
+	switch (TimebaseForBarAndBeatEvents)
+	{
+	case ECalibratedMusicTimebase::AudioRenderTime:
+		Basis = &CurrentSmoothedAudioRenderSongPos;
+		break;
+	case ECalibratedMusicTimebase::ExperiencedTime:
+		Basis = &CurrentPlayerExperiencedSongPos;
+		break;
+	}
+	int32 CurrBar = FMath::FloorToInt32(Basis->BarsIncludingCountIn);
+	if (LastBroadcastBar != CurrBar)
+	{
+		BarEvent.Broadcast(Basis->Timestamp.Bar);
+		LastBroadcastBar = CurrBar;
+	}
+	int32 CurrBeat = FMath::FloorToInt32(Basis->BeatsIncludingCountIn);
+	if (LastBroadcastBeat != CurrBeat)
+	{
+		BeatEvent.Broadcast(CurrBeat, FMath::FloorToInt32(Basis->Timestamp.Beat));
+		LastBroadcastBeat = CurrBeat;
+	}
+	const FSongSection& SongSection = Basis->CurrentSongSection;
+	if (LastBroadcastSongSection.StartTick != SongSection.StartTick || LastBroadcastSongSection.LengthTicks != SongSection.LengthTicks)
+	{
+		SectionEvent.Broadcast(SongSection.Name, SongSection.StartTick, SongSection.LengthTicks);
+		LastBroadcastSongSection = FSongSection(SongSection.Name, SongSection.StartTick, SongSection.LengthTicks);
+	}
+}
+
+void UMusicClockComponent::MakeDefaultSongMap()
+{
+	DefaultMaps.EmptyAllMaps();
+	DefaultMaps.Init(MidiConstants::kTicksPerQuarterNoteInt);
+	DefaultMaps.GetTempoMap().AddTempoInfoPoint(MidiConstants::BPMToMidiTempo(Tempo), 0);
+	DefaultMaps.GetBarMap().AddTimeSignatureAtBarIncludingCountIn(0, TimeSignatureNum, TimeSignatureDenom);
+}
diff --git a/Engine/Plugins/Runtime/Harmonix/Source/HarmonixMetasound/Private/HarmonixMetasound/Components/MusicTempometerComponent.cpp b/Engine/Plugins/Runtime/Harmonix/Source/HarmonixMetasound/Private/HarmonixMetasound/Components/MusicTempometerComponent.cpp
new file mode 100644
--- /dev/null
+++ b/Engine/Plugins/Runtime/Harmonix/Source/HarmonixMetasound/Private/HarmonixMetasound/Components/MusicTempometerComponent.cpp
@@ -0,0 +1,105 @@
+// Copyright Epic Games, Inc. All Rights Reserved.
+
+#include "HarmonixMetasound/Components/MusicTempometerComponent.h"
+#include "Materials/MaterialParameterCollectionInstance.h"
+#include "Engine/World.h"
+
+#include UE_INLINE_GENERATED_CPP_BY_NAME(MusicTempometerComponent)
+
+UMusicTempometerComponent::UMusicTempometerComponent()
+{
+	PrimaryComponentTick.bCanEverTick = true;
+}
+
+void UMusicTempometerComponent::TickComponent(float DeltaTime, enum ELevelTick TickType, FActorComponentTickFunction* ThisTickFunction)
+{
+	Super::TickComponent(DeltaTime, TickType, ThisTickFunction);
+
+	if (MaterialParameterCollection)
+	{
+		UpdateCachedSongPosIfNeeded();
+	}
+	else
+	{
+		SetComponentTickEnabled(false);
+	}
+}
+
+#if WITH_EDITOR
+void UMusicTempometerComponent::PostEditChangeProperty(FPropertyChangedEvent& PropertyChangedEvent)
+{
+	Super::PostEditChangeProperty(PropertyChangedEvent);
+
+	if (PropertyChangedEvent.GetPropertyName() == GET_MEMBER_NAME_CHECKED(UMusicTempometerComponent, MaterialParameterCollection))
+	{
+		SetComponentTickEnabled(MaterialParameterCollection != nullptr);
+	}
+}
+#endif // WITH_EDITOR
+
+void UMusicTempometerComponent::UpdateCachedSongPos() const
+{
+	LastFrameCounter = GFrameCounter;
+
+	// Cache or clear various FSongPos
+	if (const UMusicClockComponent* Clock = GetClockNoMutex())
+	{
+		SongPos = Clock->CurrentVideoRenderSongPos;
+	}
+	else
+	{
+		SongPos.Reset();
+	}
+
+	// Find a MaterialParameterCollectionInstance to update
+	if (!IsValid(MaterialParameterCollectionInstance))
+	{
+		if (MaterialParameterCollection)
+		{
+			if (UWorld* World = GetWorld())
+			{
+				MaterialParameterCollectionInstance = World->GetParameterCollectionInstance(MaterialParameterCollection);
+			}
+		}
+
+		if (!IsValid(MaterialParameterCollectionInstance))
+		{
+			return;
+		}
+	}
+
+	MaterialParameterCollectionInstance->SetScalarParameterValue(SecondsIncludingCountInParameterName, SongPos.SecondsIncludingCountIn);
+	MaterialParameterCollectionInstance->SetScalarParameterValue(BarsIncludingCountInParameterName, SongPos.BarsIncludingCountIn);
+	MaterialParameterCollectionInstance->SetScalarParameterValue(BeatsIncludingCountInParameterName, SongPos.BeatsIncludingCountIn);
+
+	MaterialParameterCollectionInstance->SetScalarParameterValue(SecondsFromBarOneParameterName, SongPos.SecondsFromBarOne);
+	MaterialParameterCollectionInstance->SetScalarParameterValue(TimestampBarParameterName, (float)SongPos.Timestamp.Bar);
+	MaterialParameterCollectionInstance->SetScalarParameterValue(TimestampBeatInBarParameterName, SongPos.Timestamp.Beat);
+
+	MaterialParameterCollectionInstance->SetScalarParameterValue(BarProgressParameterName, FMath::Fractional(SongPos.BarsIncludingCountIn));
+	MaterialParameterCollectionInstance->SetScalarParameterValue(BeatProgressParameterName, FMath::Fractional(SongPos.BeatsIncludingCountIn));
+
+	MaterialParameterCollectionInstance->SetScalarParameterValue(TimeSignatureNumeratorParameterName, SongPos.TimeSigNumerator);
+	MaterialParameterCollectionInstance->SetScalarParameterValue(TimeSignatureDenominatorParameterName, SongPos.TimeSigDenominator);
+	MaterialParameterCollectionInstance->SetScalarParameterValue(TempoParameterName, SongPos.Tempo);
+}
+
+UMusicClockComponent* UMusicTempometerComponent::FindClock(AActor* Actor) const
+{
+	UMusicClockComponent* FoundClock = nullptr;
+	for (UActorComponent* Component : Actor->GetComponents())
+	{
+		FoundClock = Cast<UMusicClockComponent>(Component);
+		if (FoundClock)
+		{
+			break;
+		}
+	}
+	return FoundClock;
+}
+
+void UMusicTempometerComponent::SetOwnerClock() const
+{
+	MusicClock = FindClock(GetOwner());
+}
+
diff --git a/Engine/Plugins/Runtime/Harmonix/Source/HarmonixMetasound/Private/HarmonixMetasound/Components/WallClockMusicClockDriver.cpp b/Engine/Plugins/Runtime/Harmonix/Source/HarmonixMetasound/Private/HarmonixMetasound/Components/WallClockMusicClockDriver.cpp
new file mode 100644
--- /dev/null
+++ b/Engine/Plugins/Runtime/Harmonix/Source/HarmonixMetasound/Private/HarmonixMetasound/Components/WallClockMusicClockDriver.cpp
@@ -0,0 +1,93 @@
+// Copyright Epic Games, Inc. All Rights Reserved.
+#include "WallClockMusicClockDriver.h"
+#include "Harmonix.h"
+#include "Engine/World.h"
+
+bool FWallClockMusicClockDriver::CalculateSongPosWithOffset(float MsOffset, ECalibratedMusicTimebase Timebase, FMidiSongPos& OutResult) const
+{
+	if (!TempoMap.IsValid())
+	{
+		return false;
+	}
+
+	const FSongMaps* Maps = TempoMap->GetSongMaps();
+	switch (Timebase)
+	{
+	case ECalibratedMusicTimebase::AudioRenderTime:
+		OutResult.SetByTime((Clock->CurrentSmoothedAudioRenderSongPos.SecondsIncludingCountIn * 1000.0f) + MsOffset, *Maps);
+		break;
+	case ECalibratedMusicTimebase::ExperiencedTime:
+		OutResult.SetByTime((Clock->CurrentPlayerExperiencedSongPos.SecondsIncludingCountIn * 1000.0f) + MsOffset, *Maps);
+		break;
+	case ECalibratedMusicTimebase::VideoRenderTime:
+	default:
+		OutResult.SetByTime((Clock->CurrentVideoRenderSongPos.SecondsIncludingCountIn * 1000.0f) + MsOffset, *Maps);
+		break;
+	}
+
+	return true;
+}
+
+void FWallClockMusicClockDriver::Disconnect()
+{
+	TempoMap = nullptr;
+}
+
+bool FWallClockMusicClockDriver::RefreshCurrentSongPos()
+{
+	check(Clock);
+	check(Clock->GetWorld());
+
+	bool TempoChanged = Clock->CurrentSmoothedAudioRenderSongPos.Tempo != Clock->Tempo;
+
+	double RunTime = Clock->GetWorld()->GetTimeSeconds() - StartTimeSecs;
+
+	const FSongMaps* Maps = GetCurrentSongMaps();
+	check(Maps);
+
+	Clock->CurrentSmoothedAudioRenderSongPos.SetByTime((float)(RunTime * 1000.0), *Maps);
+	Clock->CurrentPlayerExperiencedSongPos.SetByTime(Clock->CurrentSmoothedAudioRenderSongPos.SecondsIncludingCountIn * 1000.0f - FHarmonixModule::GetMeasuredUserExperienceAndReactionToAudioRenderOffsetMs(), *Maps);
+	Clock->CurrentVideoRenderSongPos.SetByTime(Clock->CurrentSmoothedAudioRenderSongPos.SecondsIncludingCountIn * 1000.0f - FHarmonixModule::GetMeasuredVideoToAudioRenderOffsetMs(), *Maps);
+
+	if (TempoChanged)
+	{
+		Clock->Tempo = Clock->CurrentSmoothedAudioRenderSongPos.Tempo;
+		Clock->CurrentBeatDurationSec = (60.0f / Clock->Tempo) / Clock->CurrentClockAdvanceRate;
+		Clock->CurrentBarDurationSec = ((Clock->TimeSignatureNum * Clock->CurrentBeatDurationSec) / (Clock->TimeSignatureDenom / 4.0f)) / Clock->CurrentClockAdvanceRate;
+	}
+
+	return true;
+
+}
+
+void FWallClockMusicClockDriver::OnStart()
+{
+	check(Clock);
+	StartTimeSecs = Clock->GetWorld()->GetTimeSeconds();
+	PauseTimeSecs = 0.0;
+}
+
+void FWallClockMusicClockDriver::OnPause()
+{
+	check(Clock);
+	PauseTimeSecs = Clock->GetWorld()->GetTimeSeconds();
+}
+
+void FWallClockMusicClockDriver::OnContinue()
+{
+	check(Clock);
+	double CurrentTime = Clock->GetWorld()->GetTimeSeconds();
+	StartTimeSecs += (CurrentTime - PauseTimeSecs);
+	PauseTimeSecs = 0.0;
+	RefreshCurrentSongPos();
+}
+
+const FSongMaps* FWallClockMusicClockDriver::GetCurrentSongMaps() const
+{
+	if (TempoMap.IsValid())
+	{
+		return TempoMap->GetSongMaps();
+	}
+	check(Clock);
+	return &Clock->DefaultMaps;
+}
diff --git a/Engine/Plugins/Runtime/Harmonix/Source/HarmonixMetasound/Private/HarmonixMetasound/Components/WallClockMusicClockDriver.h b/Engine/Plugins/Runtime/Harmonix/Source/HarmonixMetasound/Private/HarmonixMetasound/Components/WallClockMusicClockDriver.h
new file mode 100644
--- /dev/null
+++ b/Engine/Plugins/Runtime/Harmonix/Source/HarmonixMetasound/Private/HarmonixMetasound/Components/WallClockMusicClockDriver.h
@@ -0,0 +1,30 @@
+// Copyright Epic Games, Inc. All Rights Reserved.
+#pragma once
+#include "HarmonixMetasound/Components/MusicClockComponent.h"
+#include "HarmonixMidi/MidiFile.h"
+#include "UObject/StrongObjectPtr.h"
+
+struct FWallClockMusicClockDriver : public FMusicClockDriverBase
+{
+public:
+	FWallClockMusicClockDriver(UMusicClockComponent* InClock, UMidiFile* InTempoMap)
+		: FMusicClockDriverBase(InClock)
+		, TempoMap(InTempoMap)
+	{}
+
+	virtual bool CalculateSongPosWithOffset(float MsOffset, ECalibratedMusicTimebase Timebase, FMidiSongPos& OutResult) const override;
+
+	virtual void Disconnect() override;
+	virtual bool RefreshCurrentSongPos() override;
+	virtual void OnStart() override;
+	virtual void OnPause() override;
+	virtual void OnContinue() override;
+	virtual void OnStop() override {}
+	virtual const FSongMaps* GetCurrentSongMaps() const override;
+
+private:
+	TWeakObjectPtr<UMidiFile> TempoMap;
+
+	double StartTimeSecs = 0.0;
+	double PauseTimeSecs = 0.0f;
+};
\ No newline at end of file
