[Harmonix] Move Harmonix plugin into Engine/Plugins/Runtime to make it available to licensees in UE 5.4

diff --git a/Engine/Plugins/Runtime/Harmonix/Source/HarmonixMetasound/Private/HarmonixMetasound/DataTypes/MidiClock.cpp b/Engine/Plugins/Runtime/Harmonix/Source/HarmonixMetasound/Private/HarmonixMetasound/DataTypes/MidiClock.cpp
new file mode 100644
--- /dev/null
+++ b/Engine/Plugins/Runtime/Harmonix/Source/HarmonixMetasound/Private/HarmonixMetasound/DataTypes/MidiClock.cpp
@@ -0,0 +1,625 @@
+// Copyright Epic Games, Inc. All Rights Reserved.
+#include "HarmonixMetasound/DataTypes/MidiClock.h"
+
+#include "MetasoundDataTypeRegistrationMacro.h"
+#include "Engine/Engine.h"
+#include "HarmonixMetasound/Subsystems/MidiClockUpdateSubsystem.h"
+
+DEFINE_LOG_CATEGORY_STATIC(LogHarmonixMidiClock, Log, All);
+
+#define LOCTEXT_NAMESPACE "HarmonixMetaSound"
+
+namespace Metasound
+{
+	DEFINE_METASOUND_ENUM_BEGIN(EMidiClockSubdivisionQuantization, FEnumMidiClockSubdivisionQuantizationType, "SubdivisionQuantizationType")
+		DEFINE_METASOUND_ENUM_ENTRY(EMidiClockSubdivisionQuantization::Bar, "BarDesc", "Bar", "BarTT", "Bar"),
+		DEFINE_METASOUND_ENUM_ENTRY(EMidiClockSubdivisionQuantization::Beat, "BeatDesc", "Beat", "BeatTT", "Beat"),
+		DEFINE_METASOUND_ENUM_ENTRY(EMidiClockSubdivisionQuantization::ThirtySecondNote, "ThirtySecondNoteDesc", "1/32", "ThirtySecondNoteTT", "1/32"),
+		DEFINE_METASOUND_ENUM_ENTRY(EMidiClockSubdivisionQuantization::SixteenthNote, "SixteenthNoteDesc", "1/16", "SixteenthNoteTT", "1/16"),
+		DEFINE_METASOUND_ENUM_ENTRY(EMidiClockSubdivisionQuantization::EighthNote, "EighthNoteDesc", "1/8", "EighthNoteTT", "1/8"),
+		DEFINE_METASOUND_ENUM_ENTRY(EMidiClockSubdivisionQuantization::QuarterNote, "QuarterNoteDesc", "1/4", "QuarterNoteTT", "1/4"),
+		DEFINE_METASOUND_ENUM_ENTRY(EMidiClockSubdivisionQuantization::HalfNote, "HalfNoteDesc", "Half", "HalfNoteTT", "Half"),
+		DEFINE_METASOUND_ENUM_ENTRY(EMidiClockSubdivisionQuantization::WholeNote, "WholeNoteDesc", "Whole", "WholeNoteTT", "Whole"),
+		DEFINE_METASOUND_ENUM_ENTRY(EMidiClockSubdivisionQuantization::DottedSixteenthNote, "DottedSixteenthNoteDesc", "(dotted) 1/16", "DottedSixteenthNoteTT", "(dotted) 1/16"),
+		DEFINE_METASOUND_ENUM_ENTRY(EMidiClockSubdivisionQuantization::DottedEighthNote, "DottedEighthNoteDesc", "(dotted) 1/8", "DottedEighthNoteTT", "(dotted) 1/8"),
+		DEFINE_METASOUND_ENUM_ENTRY(EMidiClockSubdivisionQuantization::DottedQuarterNote, "DottedQuarterNoteDesc", "(dotted) 1/4", "DottedQuarterNoteTT", "(dotted) 1/4"),
+		DEFINE_METASOUND_ENUM_ENTRY(EMidiClockSubdivisionQuantization::DottedHalfNote, "DottedHalfNoteDesc", "(dotted) Half", "DottedHalfNoteTT", "(dotted) Half"),
+		DEFINE_METASOUND_ENUM_ENTRY(EMidiClockSubdivisionQuantization::DottedWholeNote, "DottedWholeNoteDesc", "(dotted) Whole", "DottedWholeNoteTT", "(dotted) Whole"),
+		DEFINE_METASOUND_ENUM_ENTRY(EMidiClockSubdivisionQuantization::SixteenthNoteTriplet, "SixteenthNoteTripletDesc", "1/16 (triplet)", "SixteenthNoteTripletTT", "1/16 (triplet)"),
+		DEFINE_METASOUND_ENUM_ENTRY(EMidiClockSubdivisionQuantization::EighthNoteTriplet, "EighthNoteTripletDesc", "1/8 (triplet)", "EighthNoteTripletTT", "1/8 (triplet)"),
+		DEFINE_METASOUND_ENUM_ENTRY(EMidiClockSubdivisionQuantization::QuarterNoteTriplet, "QuarterNoteTripletDesc", "1/4 (triplet)", "QuarterNoteTripletTT", "1/4 (triplet)"),
+		DEFINE_METASOUND_ENUM_ENTRY(EMidiClockSubdivisionQuantization::HalfNoteTriplet, "HalfNoteTripletDesc", "1/2 (triplet)", "HalfNoteTripletTT", "1/2 (triplet)"),
+	DEFINE_METASOUND_ENUM_END()
+}
+
+int32 SubdivisionToMidiTicks(const EMidiClockSubdivisionQuantization Division, const int32 CurrentTick, const FSongMaps& SongMap)
+{
+	int32 BarMapPointIndex = SongMap.GetBarMap().GetPointIndexForTick(CurrentTick);
+	if (BarMapPointIndex < 0)
+	{
+		return 0;
+	}
+	const FTimeSignature* TimeSignature =  SongMap.GetTimeSignatureAtTick(CurrentTick);
+	if (!TimeSignature)
+	{
+		return 0;
+	}
+	switch (Division)
+	{
+	case EMidiClockSubdivisionQuantization::Bar: 					return SongMap.GetBarMap().GetTicksInBarAfterPoint(BarMapPointIndex);
+	case EMidiClockSubdivisionQuantization::Beat:					return (MidiConstants::kTicksPerQuarterNoteInt * 4) / TimeSignature->Denominator;
+	case EMidiClockSubdivisionQuantization::ThirtySecondNote:		return MidiConstants::kTicksPerQuarterNoteInt / 8;
+	case EMidiClockSubdivisionQuantization::SixteenthNote:			return MidiConstants::kTicksPerQuarterNoteInt / 4;
+	case EMidiClockSubdivisionQuantization::EighthNote:				return MidiConstants::kTicksPerQuarterNoteInt / 2;
+	case EMidiClockSubdivisionQuantization::QuarterNote:			return MidiConstants::kTicksPerQuarterNoteInt;
+	case EMidiClockSubdivisionQuantization::HalfNote:				return MidiConstants::kTicksPerQuarterNoteInt * 2;
+	case EMidiClockSubdivisionQuantization::WholeNote:				return MidiConstants::kTicksPerQuarterNoteInt * 4;
+	case EMidiClockSubdivisionQuantization::DottedSixteenthNote:	return (MidiConstants::kTicksPerQuarterNoteInt / 4) + (MidiConstants::kTicksPerQuarterNoteInt / 8);
+	case EMidiClockSubdivisionQuantization::DottedEighthNote:		return (MidiConstants::kTicksPerQuarterNoteInt / 2) + (MidiConstants::kTicksPerQuarterNoteInt / 4);
+	case EMidiClockSubdivisionQuantization::DottedQuarterNote:		return (MidiConstants::kTicksPerQuarterNoteInt)     + (MidiConstants::kTicksPerQuarterNoteInt / 2);
+	case EMidiClockSubdivisionQuantization::DottedHalfNote:			return (MidiConstants::kTicksPerQuarterNoteInt * 2) + (MidiConstants::kTicksPerQuarterNoteInt);
+	case EMidiClockSubdivisionQuantization::DottedWholeNote:		return (MidiConstants::kTicksPerQuarterNoteInt * 4) + (MidiConstants::kTicksPerQuarterNoteInt * 2);
+	case EMidiClockSubdivisionQuantization::SixteenthNoteTriplet:   return (MidiConstants::kTicksPerQuarterNoteInt / 2) / 3;
+	case EMidiClockSubdivisionQuantization::EighthNoteTriplet:		return MidiConstants::kTicksPerQuarterNoteInt / 3;
+	case EMidiClockSubdivisionQuantization::QuarterNoteTriplet:		return (MidiConstants::kTicksPerQuarterNoteInt * 2) / 3;
+	case EMidiClockSubdivisionQuantization::HalfNoteTriplet:        return (MidiConstants::kTicksPerQuarterNoteInt * 4) / 3;
+	default:	/* Beat */											return (MidiConstants::kTicksPerQuarterNoteInt * 4) / TimeSignature->Denominator;
+	}
+}
+
+REGISTER_METASOUND_DATATYPE(HarmonixMetasound::FMidiClock, "MidiClock")
+
+namespace HarmonixMetasound
+{
+	using namespace Metasound;
+	
+	FMidiClock::FMidiClock(const FOperatorSettings& InSettings)
+		: TempoChangesCursor(this)
+		, BlockSize(InSettings.GetNumFramesPerBlock())
+		, CurrentBlockFrameIndex(0)
+		, SampleRate(InSettings.GetSampleRate())
+		, HasSpeedChangeInBlock(false)
+		, HasTempoChangeInBlock(false)
+	{
+		DrivingMidiPlayCursorMgr = MakeShared<FMidiPlayCursorMgr>();
+		DrivingMidiPlayCursorMgr->RegisterHiResPlayCursor(&TempoChangesCursor);
+		IOwnThePlayCursorMgr = true;
+		SpeedChangesInBlock.Add({0, 0.0f, 1.0f});
+		TempoChangesInBlock.Add({0, 0.0f, 120.0f});
+
+		RegisterForGameThreadUpdates();
+	}
+
+	FMidiClock::FMidiClock(const FMidiClock& Other)
+	{
+		*this = Other;
+	}
+
+	FMidiClock& FMidiClock::operator=(const FMidiClock& Other)
+	{
+		if (&Other != this)
+		{
+			bool NeedReRegister = false;
+			if (DrivingMidiPlayCursorMgr && DrivingMidiPlayCursorMgr != Other.DrivingMidiPlayCursorMgr)
+			{
+				NeedReRegister = true;
+				DrivingMidiPlayCursorMgr->UnregisterPlayCursor(&TempoChangesCursor);
+				DrivingMidiPlayCursorMgr = nullptr;
+			}
+
+			BlockSize = Other.BlockSize;
+			CurrentBlockFrameIndex = Other.CurrentBlockFrameIndex;
+			SampleRate = Other.SampleRate;
+			CurrentTransportState = Other.CurrentTransportState;
+
+			TransportChangesInBlock = Other.TransportChangesInBlock;
+			HasSpeedChangeInBlock = Other.HasSpeedChangeInBlock;
+			SpeedChangesInBlock = Other.SpeedChangesInBlock;
+			HasTempoChangeInBlock = Other.HasTempoChangeInBlock;
+			TempoChangesInBlock = Other.TempoChangesInBlock;
+			SmoothingEnabled = Other.SmoothingEnabled;
+
+			if (NeedReRegister)
+			{
+				DrivingMidiPlayCursorMgr = Other.DrivingMidiPlayCursorMgr;
+				if (DrivingMidiPlayCursorMgr)
+				{
+					DrivingMidiPlayCursorMgr->RegisterHiResPlayCursor(&TempoChangesCursor);
+				}
+			}
+
+			RegisterForGameThreadUpdates();
+		}
+		return *this;
+	}
+
+	FMidiClock::FMidiClock(FMidiClock&& Other)
+		: TempoChangesCursor(this)
+	{
+		BlockSize = Other.BlockSize;
+		CurrentBlockFrameIndex = Other.CurrentBlockFrameIndex;
+		SampleRate = Other.SampleRate;
+		CurrentTransportState = Other.CurrentTransportState;
+
+		TransportChangesInBlock = MoveTemp(Other.TransportChangesInBlock);
+		HasSpeedChangeInBlock = Other.HasSpeedChangeInBlock;
+		SpeedChangesInBlock = MoveTemp(Other.SpeedChangesInBlock);
+		HasTempoChangeInBlock = Other.HasTempoChangeInBlock;
+		TempoChangesInBlock = MoveTemp(Other.TempoChangesInBlock);
+		SmoothingEnabled = Other.SmoothingEnabled;
+		
+		if (Other.DrivingMidiPlayCursorMgr)
+		{
+			Other.DrivingMidiPlayCursorMgr->UnregisterPlayCursor(&Other.TempoChangesCursor);
+			DrivingMidiPlayCursorMgr = Other.DrivingMidiPlayCursorMgr;
+		}
+		else
+		{
+			DrivingMidiPlayCursorMgr = MakeShared<FMidiPlayCursorMgr>();
+		}
+		DrivingMidiPlayCursorMgr->RegisterHiResPlayCursor(&TempoChangesCursor);
+
+		RegisterForGameThreadUpdates();
+	}
+
+	FMidiClock::~FMidiClock()
+	{
+		UnregisterForGameThreadUpdates();
+	
+		DrivingMidiPlayCursorMgr->UnregisterPlayCursor(&TempoChangesCursor, false);
+	}
+
+	FMidiClock::FTempoChangesCursor::FTempoChangesCursor()
+	{
+		SetMessageFilter(FMidiPlayCursor::EFilterPassFlags::Tempo);
+	}
+
+	FMidiClock::FTempoChangesCursor::FTempoChangesCursor(FMidiClock* MidiClock) : MyMidiClock(MidiClock)
+	{
+		check(MyMidiClock);
+		SetMessageFilter(FMidiPlayCursor::EFilterPassFlags::Tempo);
+	}
+
+	void FMidiClock::FTempoChangesCursor::OnTempo(int32 TrackIndex, int32 Tick, int32 Tempo, bool IsPreroll /*= false*/)
+	{
+		check(MyMidiClock);
+		int32 BlockFrameIndex = MyMidiClock->CurrentBlockFrameIndex;
+
+		check(BlockFrameIndex >= MyMidiClock->TempoChangesInBlock.Last().BlockSampleFrameIndex);
+		MyMidiClock->HasTempoChangeInBlock = true;
+		float Bpm = MidiConstants::MidiTempoToBPM(Tempo);
+		if (MyMidiClock->TempoChangesInBlock.Last().BlockSampleFrameIndex == BlockFrameIndex)
+		{
+			MyMidiClock->TempoChangesInBlock.Last().Tempo = Bpm;
+		}
+		else
+		{ 
+			MyMidiClock->TempoChangesInBlock.Add({BlockFrameIndex, 0.0f, Bpm});
+		}
+	}
+
+	void FMidiClock::RegisterForGameThreadUpdates()
+	{
+		if (nullptr == GEngine)
+		{
+			UE_LOG(LogHarmonixMidiClock, Error, TEXT("GEngine didn't exist. Failed to register a MIDI clock for game thread updates."))
+			return;
+		}
+
+		UMidiClockUpdateSubsystem* UpdateSubsystem = GEngine->GetEngineSubsystem<UMidiClockUpdateSubsystem>();
+
+		if (nullptr == UpdateSubsystem)
+		{
+			UE_LOG(LogHarmonixMidiClock, Error, TEXT("Failed to get UMidiClockUpdateSubsystem. Low resolution MIDI clock updates won't happen."));
+			return;
+		}
+
+		UpdateSubsystem->TrackClock(this);
+	}
+
+	void FMidiClock::UnregisterForGameThreadUpdates()
+	{
+		if (nullptr == GEngine)
+		{
+			UE_LOG(LogHarmonixMidiClock, Warning, TEXT("GEngine didn't exist while attempting to unregister a MIDI clock with the game thread updater."))
+			return;
+		}
+
+		UMidiClockUpdateSubsystem* UpdateSubsystem = GEngine->GetEngineSubsystem<UMidiClockUpdateSubsystem>();
+
+		if (nullptr == UpdateSubsystem)
+		{
+			UE_LOG(LogHarmonixMidiClock, Warning, TEXT("Failed to get UMidiClockUpdateSubsystem while attempting to unregister a MIDI clock with the game thread updater."));
+			return;
+		}
+
+		UpdateSubsystem->StopTrackingClock(this);
+	}
+
+
+	void FMidiClock::ResetAndStart(int32 FrameIndex, bool SeekToStart)
+	{
+		if (SeekToStart)
+		{
+			SampleCount = 0;
+			DrivingMidiPlayCursorMgr->SeekTo(0, 0, true, false);
+		}
+	
+		AddTransportStateChangeToBlock({FrameIndex,0.0f, EMusicPlayerTransportState::Playing});
+		
+		HasSpeedChangeInBlock = false;
+		SpeedChangesInBlock.SetNum(1);
+		SpeedChangesInBlock[0] = {0, 0.0f, 1.0f};
+
+		HasTempoChangeInBlock = false;
+		TempoChangesInBlock.SetNum(1);
+		TempoChangesInBlock[0] = { 0, 0.0f, 120.0f };
+
+		CurrentBlockFrameIndex = FrameIndex;
+		FramesUntilNextProcess = 0;
+	}
+
+	void FMidiClock::PrepareBlock()
+	{
+		TransportChangesInBlock.Empty(4);
+		if (SpeedChangesInBlock.Num() > 1)
+		{
+			SpeedChangesInBlock[0].Speed = SpeedChangesInBlock.Last().Speed;
+			SpeedChangesInBlock.SetNum(1, false);
+		}
+		HasSpeedChangeInBlock = false;
+		if (TempoChangesInBlock.Num() > 1)
+		{
+			TempoChangesInBlock[0].Tempo = TempoChangesInBlock.Last().Tempo;
+			TempoChangesInBlock.SetNum(1, false);
+		}
+		HasTempoChangeInBlock = false;
+		CurrentBlockFrameIndex = 0;
+
+		MidiClockEventsInBlock.Reset();
+	}
+
+	bool FMidiClock::HasLowResCursors() const
+	{
+		return DrivingMidiPlayCursorMgr->HasLowResCursors();
+	}
+
+	void FMidiClock::UpdateLowResCursors()
+	{
+		DrivingMidiPlayCursorMgr->AdvanceLowResCursors();
+	}
+
+	void FMidiClock::AddTransportStateChangeToBlock(const FMidiTimestampTransportState& NewTransportState)
+	{
+		if (TransportChangesInBlock.IsEmpty() || TransportChangesInBlock.Last().BlockSampleFrameIndex <= NewTransportState.BlockSampleFrameIndex)
+		{
+			CurrentTransportState = NewTransportState;
+			TransportChangesInBlock.Add(NewTransportState);
+		}
+	}
+
+	void FMidiClock::AddSpeedChangeToBlock(const FMidiTimestampSpeed& NewSpeed)
+	{
+		check(SpeedChangesInBlock.IsEmpty() || SpeedChangesInBlock.Last().BlockSampleFrameIndex <= NewSpeed.BlockSampleFrameIndex);
+		if (SpeedChangesInBlock.Last().BlockSampleFrameIndex == NewSpeed.BlockSampleFrameIndex)
+		{
+			SpeedChangesInBlock.Last().Speed = NewSpeed.Speed;
+		}
+		else
+		{
+			SpeedChangesInBlock.Add(NewSpeed);
+		}
+		HasSpeedChangeInBlock = true;
+		DrivingMidiPlayCursorMgr->InformOfCurrentAdvanceRate(NewSpeed.Speed);
+	}
+
+	const TArray<FMidiClockEvent>& FMidiClock::GetMidiClockEventsInBlock() const
+	{
+		return MidiClockEventsInBlock;
+	}
+
+	EMusicPlayerTransportState FMidiClock::GetTransportStateAtBlockSampleFrame(int32 FrameIndex) const
+	{
+		return GetTransportTimestampForBlockSampleFrame(FrameIndex).TransportState;
+	}
+
+	EMusicPlayerTransportState FMidiClock::GetTransportStateAtEndOfBlock() const
+	{
+		return CurrentTransportState.TransportState;
+	}
+
+	const FMidiTimestampTransportState& FMidiClock::GetTransportTimestampForBlockSampleFrame(int32 FrameIndex) const
+	{
+		if (TransportChangesInBlock.IsEmpty())
+		{
+			return CurrentTransportState;
+		}
+		int32 Index = Algo::UpperBoundBy(TransportChangesInBlock, FrameIndex, [](const FMidiTimestampTransportState& t) { return t.BlockSampleFrameIndex; }) - 1;
+		if (Index < 0)
+		{
+			return CurrentTransportState;
+		}
+		return TransportChangesInBlock[Index];
+	}
+
+	float FMidiClock::GetSpeedAtBlockSampleFrame(int32 FrameIndex) const
+	{
+		return GetSpeedTimestampForBlockSampleFrame(FrameIndex).Speed;
+	}
+
+	float FMidiClock::GetSpeedAtEndOfBlock() const
+	{
+		return SpeedChangesInBlock.Last().Speed;
+	}
+
+	const FMidiTimestampSpeed& FMidiClock::GetSpeedTimestampForBlockSampleFrame(int32 FrameIndex) const
+	{
+		int32 Index = Algo::UpperBoundBy(SpeedChangesInBlock, FrameIndex, [](const FMidiTimestampSpeed& t) { return t.BlockSampleFrameIndex; }) - 1;
+		if (Index < 0)
+		{
+			Index = 0;
+		}
+		return SpeedChangesInBlock[Index];
+	}
+
+	float FMidiClock::GetTempoAtBlockSampleFrame(int32 FrameIndex) const
+	{
+		int32 Index = Algo::UpperBoundBy(TempoChangesInBlock, FrameIndex, [](const FMidiTimestampTempo& t) { return t.BlockSampleFrameIndex; }) - 1;
+		if (Index < 0)
+		{
+			Index = 0;
+		}
+		return TempoChangesInBlock[Index].Tempo;
+	}
+
+	float FMidiClock::GetTempoAtEndOfBlock() const
+	{
+		return TempoChangesInBlock.Last().Tempo;
+	}
+
+	int32 FMidiClock::GetNumTempoChangesInBlock() const
+	{
+		return TempoChangesInBlock.Num();
+	}
+
+	FMidiTimestampTempo FMidiClock::GetTempoChangeByIndex(int32 Index) const
+	{
+		if (Index >= 0 && Index < TempoChangesInBlock.Num())
+		{
+			return TempoChangesInBlock[Index]; // intentional copy
+		}
+
+		return InvalidMidiTimestampTempo;
+	}
+
+	int32 FMidiClock::GetCurrentMidiTick() const
+	{
+		return DrivingMidiPlayCursorMgr->GetCurrentHiResTick();
+	}
+
+	int32 FMidiClock::GetCurrentBlockFrameIndex() const
+	{
+		return CurrentBlockFrameIndex;
+	}
+
+	void FMidiClock::AdvanceHiResToMs(int32 BlockFrameIndex, float Ms, bool Broadcast)
+	{
+		CurrentBlockFrameIndex = BlockFrameIndex;
+		DrivingMidiPlayCursorMgr->AdvanceHiResToMs(Ms, Broadcast);
+	}
+
+	void FMidiClock::AttachToTimeAuthority(const FMidiClock& MidiClockRef)
+	{
+		DrivingMidiPlayCursorMgr->AttachToTimeAuthority(MidiClockRef.DrivingMidiPlayCursorMgr);
+	}
+
+	float FMidiClock::GetQuarterNoteIncludingCountIn() const
+	{
+		int32 Tick = DrivingMidiPlayCursorMgr->GetCurrentHiResTick();
+		return (float)Tick / (float)DrivingMidiPlayCursorMgr->GetSongMaps().GetTicksPerQuarterNote();
+	}
+
+	FMusicTimestamp FMidiClock::GetCurrentMusicTimestamp() const
+	{
+		return DrivingMidiPlayCursorMgr->GetMusicTimestampAtMs(GetCurrentHiResMs());
+	}
+
+	FMusicTimestamp FMidiClock::GetMusicTimestampAtBlockOffset(const int32 Offset) const
+	{
+		const float MsAtOffset = GetMsAtBlockOffset(Offset);
+		return DrivingMidiPlayCursorMgr->GetMusicTimestampAtMs(MsAtOffset);
+	}
+
+	float FMidiClock::GetMsAtBlockOffset(int32 Offset) const
+	{
+		const float EndMs = GetCurrentHiResMs();
+		const float MsPerFrame = 1000 / SampleRate;
+		const int32 InvOffset = BlockSize - Offset - 1;
+		return EndMs - InvOffset * MsPerFrame;
+	}
+
+	void FMidiClock::LockForMidiDataChanges()
+	{
+
+		DrivingMidiPlayCursorMgr->LockForMidiDataChanges();
+	}
+
+	void FMidiClock::MidiDataChangesComplete(FMidiPlayCursorMgr::EMidiChangePositionCorrectMode Mode /*= FMidiPlayCursorMgr::EMidiChangePositionCorrectMode::MaintainTick*/)
+	{
+		DrivingMidiPlayCursorMgr->MidiDataChangeComplete(Mode);
+	}
+
+	void FMidiClock::SeekTo(const FMusicSeekTarget& Target, int32 PreRollBars)
+	{
+		int32 Tick = 0;
+		switch (Target.Type)
+		{
+		case ESeekPointType::BarBeat:
+			Tick = DrivingMidiPlayCursorMgr->GetBarMap().MusicTimestampToTick(Target.BarBeat);
+			break;
+		default:
+		case ESeekPointType::Millisecond:
+			Tick = DrivingMidiPlayCursorMgr->GetSongMaps().MsToTick(Target.Ms);
+			break;
+		}
+		DrivingMidiPlayCursorMgr->SeekTo(Tick, PreRollBars,true,false);
+	}
+
+	void FMidiClock::InformOfCurrentAdvanceRate(float AdvanceRate)
+	{
+		DrivingMidiPlayCursorMgr->InformOfCurrentAdvanceRate(AdvanceRate);
+	}
+
+	void FMidiClock::CopySpeedAndTempoChanges(const FMidiClock* InClock, float InSpeedMult)
+	{
+		HasTempoChangeInBlock = InClock->HasTempoChangeInBlock;
+		TempoChangesInBlock = InClock->TempoChangesInBlock;
+
+		HasSpeedChangeInBlock = InClock->HasSpeedChangeInBlock;
+		SpeedChangesInBlock = InClock->SpeedChangesInBlock;
+
+		for (FMidiTimestampSpeed& SpeedChange : SpeedChangesInBlock)
+		{
+			SpeedChange.Speed *= InSpeedMult;
+		}
+	}
+
+	void FMidiClock::WriteAdvance(int32 StartFrameIndex, int32 EndFrameIndex, float InSpeed /*= 1.0f*/)
+	{
+		int32 FramesToProcess = EndFrameIndex - StartFrameIndex;
+		const FTempoMap& TempoMap = DrivingMidiPlayCursorMgr->GetSongMaps().GetTempoMap();
+		while (FramesToProcess > FramesUntilNextProcess)
+		{
+			StartFrameIndex += FramesUntilNextProcess;
+			FSampleCount AdvanceToFrame = SampleCount + (FSampleCount)((float)kMidiGranularity * InSpeed);
+			FramesUntilNextProcess = kMidiGranularity;
+			float AdvanceToMs = ((float)AdvanceToFrame * 1000.0f) / SampleRate;
+			DrivingMidiPlayCursorMgr->InformOfCurrentAdvanceRate(InSpeed);
+			AdvanceHiResToMs(StartFrameIndex, AdvanceToMs, true);
+
+			
+			// Account for looping
+			// recalculate sampleCount based on our new time, if clock looped back, this will be different from AdvanceToFrame
+			if (GetCurrentHiResMs() < AdvanceToMs)
+			{
+				SampleCount = FMath::Max<FSampleCount>(FSampleCount(GetCurrentHiResMs() / 1000.0f * SampleRate), 0);
+			}
+			else
+			{
+				SampleCount = AdvanceToFrame;
+			}
+			
+
+			FramesToProcess = EndFrameIndex - StartFrameIndex;
+		}
+		FramesUntilNextProcess -= FramesToProcess;
+	}
+
+	void FMidiClock::WriteNoAdvance(int32 StartFrameIndex, int32 EndFrameIndex)
+	{
+		EMusicPlayerTransportState ClockTransport = GetTransportStateAtBlockSampleFrame(StartFrameIndex);
+		check(ClockTransport != EMusicPlayerTransportState::Playing);
+	}
+
+	void FMidiClock::SeekTo(int32 BlockFrameIndex, const FMusicSeekTarget& InTarget, int32 InPrerollBars)
+	{
+		CurrentBlockFrameIndex = BlockFrameIndex;
+		SeekTo(InTarget, InPrerollBars);
+		SampleCount = FMath::Max<FSampleCount>(FSampleCount(GetCurrentHiResMs() / 1000.0f * SampleRate), 0);
+		FramesUntilNextProcess = 0;
+	}
+
+	TSharedPtr<FMidiFileData> FMidiClock::MakeClockConductorMidiData(float InTempoBPM, int32 InTimeSigNum, int32 InTimeSigDen)
+	{
+		TSharedPtr<FMidiFileData> OutMidiData = MakeShared<FMidiFileData>();
+
+		// clear it all out for good measure...
+		FTempoMap& TempoMap = OutMidiData->SongMaps.GetTempoMap();
+		TempoMap.Empty();
+		FBarMap& BarMap = OutMidiData->SongMaps.GetBarMap();
+		BarMap.Empty();
+		OutMidiData->Tracks.Empty();
+
+		// create conductor track
+		FMidiTrack& Track = OutMidiData->Tracks.Add_GetRef(FMidiTrack(TEXT("conductor")));
+		OutMidiData->LastEventTick = std::numeric_limits<int32>::max();
+
+		// add time sig info
+		int32 TimeSigNum = FMath::Clamp(InTimeSigNum, 1, 64);
+		int32 TimeSigDen = FMath::Clamp(InTimeSigDen, 1, 64);
+		Track.AddEvent(FMidiEvent(0, FMidiMsg((uint8)TimeSigNum, (uint8)TimeSigDen)));
+		BarMap.AddTimeSignatureAtBarIncludingCountIn(0, TimeSigNum, TimeSigDen);
+
+		// add tempo info
+		float TempoBPM = FMath::Max(1.0f, InTempoBPM);
+		int32 MidiTempo = MidiConstants::BPMToMidiTempo(TempoBPM);
+		Track.AddEvent(FMidiEvent(0, FMidiMsg(MidiTempo)));
+		TempoMap.AddTempoInfoPoint(MidiTempo, 0);
+
+		Track.Sort();
+
+		// max out song length data so the midi can play indefinitely
+		OutMidiData->SongMaps.GetSongLengthData().LastTick = std::numeric_limits<int32>::max();
+		OutMidiData->SongMaps.GetSongLengthData().LengthTicks = std::numeric_limits<int32>::max();
+		OutMidiData->SongMaps.GetSongLengthData().LengthBars = std::numeric_limits<int32>::max();
+
+		return OutMidiData;
+	}
+
+
+	FMidiClockEventCursor::FMidiClockEventCursor(FMidiClockWriteRef InClock)
+		: MidiClock(InClock)
+	{
+		MidiClock->RegisterHiResPlayCursor(this);
+	}
+
+	FMidiClockEventCursor::~FMidiClockEventCursor()
+	{
+
+		MidiClock->UnregisterPlayCursor(this);
+	}
+
+	void FMidiClockEventCursor::Reset(bool ForceNoBroadcast /*= false*/)
+	{
+		FMidiPlayCursor::Reset(ForceNoBroadcast);
+		AddEvent(FMidiClockEvent::MakeResetEvent(MidiClock->GetCurrentBlockFrameIndex(), MidiClock->GetCurrentHiResTick(), ForceNoBroadcast));
+	}
+
+	void FMidiClockEventCursor::OnLoop(int32 LoopStartTick, int32 LoopEndTick)
+	{
+		FMidiPlayCursor::OnLoop(LoopStartTick, LoopEndTick);
+		AddEvent(FMidiClockEvent::MakeLoopEvent(MidiClock->GetCurrentBlockFrameIndex(), LoopStartTick, LoopEndTick));
+	}
+
+	void FMidiClockEventCursor::SeekToTick(int32 Tick)
+	{
+		FMidiPlayCursor::SeekToTick(Tick);
+		AddEvent(FMidiClockEvent::MakeSeekToEvent(MidiClock->GetCurrentBlockFrameIndex(), Tick));
+	}
+
+	void FMidiClockEventCursor::SeekThruTick(int32 Tick)
+	{
+		FMidiPlayCursor::SeekThruTick(Tick);
+		AddEvent(FMidiClockEvent::MakeSeekThruEvent(MidiClock->GetCurrentBlockFrameIndex(), Tick));
+	}
+
+	void FMidiClockEventCursor::AdvanceThruTick(int32 Tick, bool IsPreRoll)
+	{
+		FMidiPlayCursor::AdvanceThruTick(Tick, IsPreRoll);
+		AddEvent(FMidiClockEvent::MakeAdvanceThruEvent(MidiClock->GetCurrentBlockFrameIndex(), Tick, IsPreRoll));
+	}
+
+	void FMidiClockEventCursor::AddEvent(const FMidiClockEvent& InEvent)
+	{
+		TArray<FMidiClockEvent>& Events = MidiClock->MidiClockEventsInBlock;
+		if (!Events.IsEmpty())
+		{
+			check(Events.Last().BlockFrameIndex <= InEvent.BlockFrameIndex);
+		}
+
+		Events.Add(InEvent);
+	}
+
+}
+
+#undef LOCTEXT_NAMESPACE
diff --git a/Engine/Plugins/Runtime/Harmonix/Source/HarmonixMetasound/Private/HarmonixMetasound/DataTypes/MidiClockEvent.cpp b/Engine/Plugins/Runtime/Harmonix/Source/HarmonixMetasound/Private/HarmonixMetasound/DataTypes/MidiClockEvent.cpp
new file mode 100644
--- /dev/null
+++ b/Engine/Plugins/Runtime/Harmonix/Source/HarmonixMetasound/Private/HarmonixMetasound/DataTypes/MidiClockEvent.cpp
@@ -0,0 +1,43 @@
+// Copyright Epic Games, Inc. All Rights Reserved.
+
+#include "HarmonixMetasound/DataTypes/MidiClockEvent.h"
+#include "HarmonixMetasound/DataTypes/MidiClock.h"
+
+
+namespace HarmonixMetasound
+{
+	FMidiClockEvent FMidiClockEvent::MakeResetEvent(int32 InBlockFrameIndex, int32 Tick, bool ForceNoBroadcast)
+	{
+		return FMidiClockEvent(EType::Reset, InBlockFrameIndex, Tick, Tick, false, ForceNoBroadcast);
+	}
+
+	FMidiClockEvent FMidiClockEvent::MakeLoopEvent(int32 InBlockFrameIndex, int32 InStartTick, int32 InEndTick)
+	{
+		return FMidiClockEvent(EType::Loop, InBlockFrameIndex, InStartTick, InEndTick, false, true);
+	}
+
+	FMidiClockEvent FMidiClockEvent::MakeSeekToEvent(int32 InBlockFrameIndex, int32 Tick)
+	{
+		return FMidiClockEvent(EType::SeekTo, InBlockFrameIndex, Tick, Tick, false, true);
+	}
+
+	FMidiClockEvent FMidiClockEvent::MakeSeekThruEvent(int32 InBlockFrameIndex, int32 Tick)
+	{
+		return FMidiClockEvent(EType::SeekThru, InBlockFrameIndex, Tick, Tick, false, true);
+	}
+
+	FMidiClockEvent FMidiClockEvent::MakeAdvanceThruEvent(int32 InBlockFrameIndex, int32 Tick, bool IsPreRoll)
+	{
+		return FMidiClockEvent(EType::AdvanceThru, InBlockFrameIndex, Tick, Tick, IsPreRoll, true);
+	}
+
+	FMidiClockEvent::FMidiClockEvent(EType InType, int32 InBlockFrameIndex, int32 InStartTick, int32 InEndTick, bool InIsPreRoll, bool InForceNoBroadcast)
+		: Type(InType)
+		, IsPreRoll(InIsPreRoll)
+		, ForceNoBroadcast(InForceNoBroadcast)
+		, BlockFrameIndex(InBlockFrameIndex)
+		, StartTick(InStartTick)
+		, EndTick(InEndTick)
+	{}
+
+}
\ No newline at end of file
diff --git a/Engine/Plugins/Runtime/Harmonix/Source/HarmonixMetasound/Private/HarmonixMetasound/DataTypes/MidiStream.cpp b/Engine/Plugins/Runtime/Harmonix/Source/HarmonixMetasound/Private/HarmonixMetasound/DataTypes/MidiStream.cpp
new file mode 100644
--- /dev/null
+++ b/Engine/Plugins/Runtime/Harmonix/Source/HarmonixMetasound/Private/HarmonixMetasound/DataTypes/MidiStream.cpp
@@ -0,0 +1,486 @@
+// Copyright Epic Games, Inc. All Rights Reserved.
+#include "HarmonixMetasound/DataTypes/MidiStream.h"
+#include "Algo/BinarySearch.h"
+#include "Containers/Array.h"
+#include "Containers/UnrealString.h"
+
+#include "MetasoundDataTypeRegistrationMacro.h"
+
+REGISTER_METASOUND_DATATYPE(HarmonixMetasound::FMidiStream, "MidiStream")
+
+namespace HarmonixMetasound
+{
+	DEFINE_LOG_CATEGORY(LogMidiStreamDataType);
+
+	using namespace Metasound;
+
+	FMidiStream::FMidiStream(const FOperatorSettings& InSettings)
+		: NumFramesPerBlock(InSettings.GetNumFramesPerBlock())
+		, SampleRate(InSettings.GetSampleRate())
+	{
+
+	}
+
+	void FMidiStream::PrepareBlock()
+	{
+		TransportChangesInBlock.Empty(4);
+		EventsInBlock.Empty(32);
+	}
+
+	void FMidiStream::AddTransportStateChangeMessage(int32 SampleFrameIndexInBlock, EMusicPlayerTransportState State)
+	{
+		if (State == CurrentTransportState.TransportState)
+		{
+			return;
+		}
+		CurrentTransportState.BlockSampleFrameIndex = SampleFrameIndexInBlock;
+		CurrentTransportState.BlockSampleFrameIndex = 0.0f;
+		CurrentTransportState.TransportState        = State;
+		TransportChangesInBlock.Add(CurrentTransportState);
+	}
+
+	void FMidiStream::AddMidiEvent(const FMidiStreamEvent& Event)
+	{
+		check (EventsInBlock.IsEmpty() || EventsInBlock.Last().BlockSampleFrameIndex <= Event.BlockSampleFrameIndex);
+		EventsInBlock.Add(Event);
+
+		UpdateActiveVoice(Event);
+	}
+
+	void FMidiStream::InsertMidiEvent(const FMidiStreamEvent& Event)
+	{
+		int32 AtIndex = Algo::UpperBound(EventsInBlock, Event, [](const FMidiStreamEvent& NewEvent, const FMidiStreamEvent& ExistingEvent){ return NewEvent.BlockSampleFrameIndex < ExistingEvent.BlockSampleFrameIndex; });
+		EventsInBlock.Insert(Event, AtIndex);
+
+		UpdateActiveVoice(Event);
+	}
+
+	void FMidiStream::AddNoteOffEventOrCancelPendingNoteOn(const FMidiStreamEvent& Event)
+	{
+		check(Event.MidiMessage.IsNoteOff());
+		int32 NumRemoved = EventsInBlock.RemoveAll([&](const FMidiStreamEvent& EventInList)
+			{
+				return EventInList.MidiMessage.IsNoteOn() && EventInList.GetVoiceId() == Event.GetVoiceId();
+			});
+		if (NumRemoved == 0)
+		{
+			AddMidiEvent(Event);
+		}
+	}
+
+	void FMidiStream::InsertNoteOffEventOrCancelPendingNoteOn(const FMidiStreamEvent& Event)
+	{
+		check(Event.MidiMessage.IsNoteOff());
+		int32 NumRemoved = EventsInBlock.RemoveAll([&](const FMidiStreamEvent& EventInList)
+			{
+				return EventInList.MidiMessage.IsNoteOn() && EventInList.GetVoiceId() == Event.GetVoiceId();
+			});
+		if (NumRemoved == 0)
+		{
+			InsertMidiEvent(Event);
+		}
+	}
+
+	const FString* FMidiStream::GetMidiTrackText(int32 TrackNumber, int32 TextIndex) const
+	{
+		if (!MidiFileSourceOfEvents || TrackNumber < 0 || TextIndex < 0)
+		{
+			return nullptr;
+		}
+
+		TSharedPtr<FMidiFileData> MidiFile = MidiFileSourceOfEvents->GetMidiFile();
+		if (TrackNumber >= MidiFile->Tracks.Num())
+		{
+			return nullptr;
+		}
+		const FMidiTextRepository* Repository = MidiFile->Tracks[TrackNumber].GetTextRepository();
+		if (TextIndex >= Repository->Num())
+		{
+			return nullptr;
+		}
+		return &(*Repository)[TextIndex];
+	}
+
+	void FMidiStream::Copy(const TArray<FMidiStreamReadRef>& InStreams, bool EventsOnly)
+	{
+		if (!InStreams.Num())
+		{
+			return;
+		}
+
+		if (InStreams.Num() == 1)
+		{
+			Copy(InStreams[0]);
+			return;
+		}
+
+		TArray<FMidiStreamReadRef> SameClockStreams;
+		FilterArrayToStreamsWithTheSameClock(InStreams, SameClockStreams);
+		if (!EventsOnly)
+		{
+			CopyTransportEvents(SameClockStreams);
+		}
+		CopyMidiEvents(SameClockStreams);
+	}
+
+	void FMidiStream::Copy(FMidiStreamReadRef InStream, bool EventsOnly)
+	{
+		if (GetMidiClockSource() && GetMidiClockSource()->Get() != InStream->GetMidiClockSource()->Get())
+		{
+			UE_LOG(LogMidiStreamDataType, Warning, TEXT("Cannot copy midi events from one stream to another if they have different clocks."));
+			return;
+		}
+		if (EventsOnly)
+		{
+			CopyTransportEvents(InStream);
+		}
+		CopyMidiEvents(InStream);
+	}
+
+	void FMidiStream::FilterArrayToStreamsWithTheSameClock(const TArray<FMidiStreamReadRef>& InStreams, TArray<FMidiStreamReadRef>& OutStreams)
+	{
+		// This has no clock so pass all events
+		if (!GetMidiClockSource())
+		{
+			OutStreams.Append(InStreams);
+			return;
+		}
+
+		for (const FMidiStreamReadRef& Stream : InStreams)
+		{
+			// Allow the copy if the stream matches or the stream has no clock
+			// Streams with no clock are ok because they won't conflict with our clock
+			// And chances are, that stream has no events
+			if (!Stream->GetMidiClockSource() || Stream->GetMidiClockSource()->Get() == GetMidiClockSource()->Get())
+			{
+				OutStreams.Add(Stream);
+			}
+			else
+			{
+				UE_LOG(LogMidiStreamDataType, Warning, TEXT("Cannot copy midi events from one stream to another if they have different clocks."));
+			}
+		}
+	}
+
+	void FMidiStream::CopyTransportEvents(const TArray<FMidiStreamReadRef>& InStreams)
+	{
+		TransportChangesInBlock.Empty(4);
+
+		using TransportIterator = TArray<FMidiTimestampTransportState>::RangedForConstIteratorType;
+
+		int32 NumInStreams = InStreams.Num();
+		TArray<TransportIterator> WalkingIterators;
+		WalkingIterators.Reserve(NumInStreams);
+		TArray<TransportIterator> IteratorEnds;
+		IteratorEnds.Reserve(NumInStreams);
+
+		for (int32 i = 0; i < NumInStreams; ++i)
+		{
+			WalkingIterators.Add(InStreams[i]->GetTransportChangesInBlock().begin());
+			IteratorEnds.Add(InStreams[i]->GetTransportChangesInBlock().end());
+		}
+
+		while(1)
+		{
+			int32 StreamWithEarliestEvent = -1;
+			int32 EarliestBlockSampleFrameIndex = MAX_int32;
+
+			for (int32 IteratorIndex = 0; IteratorIndex < NumInStreams; ++IteratorIndex)
+			{
+				if (WalkingIterators[IteratorIndex] != IteratorEnds[IteratorIndex])
+				{
+					if ((*WalkingIterators[IteratorIndex]).BlockSampleFrameIndex < EarliestBlockSampleFrameIndex)
+					{
+						EarliestBlockSampleFrameIndex = (*WalkingIterators[IteratorIndex]).BlockSampleFrameIndex;
+						StreamWithEarliestEvent = IteratorIndex;
+					}
+				}
+			}
+
+			// done?
+			if (StreamWithEarliestEvent == -1)
+			{
+				break;
+			}
+
+			AddTransportStateChangeMessage((*WalkingIterators[StreamWithEarliestEvent]).BlockSampleFrameIndex, (*WalkingIterators[StreamWithEarliestEvent]).TransportState);
+
+			++WalkingIterators[StreamWithEarliestEvent];
+		}
+	}
+
+	void FMidiStream::CopyTransportEvents(const FMidiStreamReadRef& InStream)
+	{
+		const TArray<FMidiTimestampTransportState>& TransportChanges = InStream->GetTransportChangesInBlock();
+		TransportChangesInBlock.Empty(TransportChanges.Num());
+		for (auto& Change : TransportChanges)
+		{
+			AddTransportStateChangeMessage(Change.BlockSampleFrameIndex, Change.TransportState);
+		}
+	}
+
+	void FMidiStream::CopyMidiEvents(const TArray<FMidiStreamReadRef>& InStreams)
+	{
+		using MidiStreamIterator = TArray<FMidiStreamEvent>::RangedForConstIteratorType;
+		
+		int32 NumInStreams = InStreams.Num();
+		TArray<MidiStreamIterator> WalkingIterators;
+		WalkingIterators.Reserve(NumInStreams);
+		TArray<MidiStreamIterator> IteratorEnds;
+		IteratorEnds.Reserve(NumInStreams);
+
+		int32 NumTotalEvents = 0;
+		for (int32 i = 0; i < NumInStreams; ++i)
+		{
+			WalkingIterators.Add(InStreams[i]->GetEventsInBlock().begin());
+			IteratorEnds.Add(InStreams[i]->GetEventsInBlock().end());
+			NumTotalEvents += InStreams[i]->GetEventsInBlock().Num();
+
+			// Ensure that the active voices from the source streams are in this stream
+			ActiveVoices.Append(InStreams[i]->ActiveVoices);
+		}
+
+		EventsInBlock.Empty(FMath::Max(32, NumTotalEvents));
+
+		while (1)
+		{
+			int32 StreamWithEarliestEvent = -1;
+			int32 EarliestBlockSampleIndex = MAX_int32;
+
+			for (int32 IteratorIndex = 0; IteratorIndex < NumInStreams; ++IteratorIndex)
+			{
+				if (WalkingIterators[IteratorIndex] != IteratorEnds[IteratorIndex])
+				{
+					if ((*WalkingIterators[IteratorIndex]).BlockSampleFrameIndex == EarliestBlockSampleIndex)
+					{
+						if ((*WalkingIterators[StreamWithEarliestEvent]).MidiMessage.IsNoteOn() && (*WalkingIterators[IteratorIndex]).MidiMessage.IsNoteOff())
+						{
+							StreamWithEarliestEvent = IteratorIndex;
+						}
+					}
+					else if ((*WalkingIterators[IteratorIndex]).BlockSampleFrameIndex < EarliestBlockSampleIndex)
+					{
+						EarliestBlockSampleIndex = (*WalkingIterators[IteratorIndex]).BlockSampleFrameIndex;
+						StreamWithEarliestEvent = IteratorIndex;
+					}
+				}
+			}
+
+			// done?
+			if (StreamWithEarliestEvent == -1)
+			{
+				break;
+			}
+
+			AddMidiEvent(*WalkingIterators[StreamWithEarliestEvent]);
+			++WalkingIterators[StreamWithEarliestEvent];
+		}
+	}
+
+	void FMidiStream::MergeMidiEvents(const TArray<FMidiStreamReadRef>& InStreams)
+	{
+		using MidiStreamIterator = TArray<FMidiStreamEvent>::RangedForConstIteratorType;
+
+		int32 NumInStreams = InStreams.Num();
+		TArray<MidiStreamIterator> WalkingIterators;
+		WalkingIterators.Reserve(NumInStreams);
+		TArray<MidiStreamIterator> IteratorEnds;
+		IteratorEnds.Reserve(NumInStreams);
+
+		int32 NumTotalEvents = 0;
+		for (int32 i = 0; i < NumInStreams; ++i)
+		{
+			WalkingIterators.Add(InStreams[i]->GetEventsInBlock().begin());
+			IteratorEnds.Add(InStreams[i]->GetEventsInBlock().end());
+			NumTotalEvents += InStreams[i]->GetEventsInBlock().Num();
+		}
+
+		EventsInBlock.Reserve(FMath::Max(32, NumTotalEvents + EventsInBlock.Num()));
+
+		while (1)
+		{
+			int32 StreamWithEarliestEvent = -1;
+			int32 EarliestBlockSampleIndex = MAX_int32;
+
+			for (int32 IteratorIndex = 0; IteratorIndex < NumInStreams; ++IteratorIndex)
+			{
+				if (WalkingIterators[IteratorIndex] != IteratorEnds[IteratorIndex])
+				{
+					if ((*WalkingIterators[IteratorIndex]).BlockSampleFrameIndex < EarliestBlockSampleIndex)
+					{
+						EarliestBlockSampleIndex = (*WalkingIterators[IteratorIndex]).BlockSampleFrameIndex;
+						StreamWithEarliestEvent = IteratorIndex;
+					}
+				}
+			}
+
+			// done?
+			if (StreamWithEarliestEvent == -1)
+			{
+				break;
+			}
+
+			InsertMidiEvent(*WalkingIterators[StreamWithEarliestEvent]);
+			++WalkingIterators[StreamWithEarliestEvent];
+		}
+	}
+
+	void FMidiStream::CopyMidiEvents(const FMidiStreamReadRef& InStream)
+	{
+		const TArray<FMidiStreamEvent>& MidiEvents = InStream->GetEventsInBlock();
+		EventsInBlock.Empty(FMath::Max(MidiEvents.Num(), 32));
+		// Ensure that the active voices from the source match this stream
+		ActiveVoices = InStream->ActiveVoices;
+		for (auto& Event : MidiEvents)
+		{
+			AddMidiEvent(Event);
+		}
+	}
+
+	void FMidiStream::MergeMidiEvents(const FMidiStreamReadRef& InStream)
+	{
+		const TArray<FMidiStreamEvent>& NewMidiEvents = InStream->GetEventsInBlock();
+		EventsInBlock.Reserve(FMath::Max(32, EventsInBlock.Num() + NewMidiEvents.Num()));
+		for (auto& Event : NewMidiEvents)
+		{
+			InsertMidiEvent(Event);
+		}
+	}
+
+	void FMidiStream::UpdateActiveVoice(const FMidiStreamEvent& Event)
+	{
+		if (Event.MidiMessage.IsNoteMessage())
+		{
+			if (Event.MidiMessage.IsNoteOn())
+			{
+				if (FMidiVoiceId VoiceId = Event.GetVoiceId(); !ActiveVoices.Contains(VoiceId))
+				{
+					ActiveVoices.Emplace(MoveTemp(VoiceId));
+				}
+			}
+			else if (Event.MidiMessage.IsNoteOff())
+			{
+				ActiveVoices.Remove(Event.GetVoiceId());
+			}
+			// Kill all or all off
+			else
+			{
+				ActiveVoices.Reset();
+			}
+		}
+	}
+
+	uint64 FMidiStreamEventTrackChannelFilter::StringToBitfield(const FString& InString, FString& OutErrorMessage)
+	{
+		OutErrorMessage.Empty();
+		uint64 Bitfield = 0;
+		auto SetBit = [&](uint32 bit) { Bitfield |= (uint64(1) << bit); };
+
+		if (InString.Equals(TEXT("*")))
+		{
+			Bitfield = std::numeric_limits<uint64>::max();
+			return Bitfield;
+		}
+
+		TArray<FString> CommaParts;
+		InString.ParseIntoArray(CommaParts, TEXT(","));
+		for (const FString& CommaPart : CommaParts)
+		{
+			TArray<FString> DashParts;
+			CommaPart.ParseIntoArray(DashParts, TEXT("-"));
+			if (DashParts.IsEmpty() || DashParts.Num() > 2)
+			{
+				OutErrorMessage = FString::Printf(TEXT("Unexpected filter expression -> \"%s\""), *CommaPart);
+				return 0;
+			}
+
+			if (DashParts.Num() == 1)
+			{
+				int32 AsNumber;
+				if (!LexTryParseString(AsNumber, *DashParts[0]))
+				{
+					OutErrorMessage = FString::Printf(TEXT("Error parsing part of filter string -> \"%s\""), *DashParts[0]);
+					return 0;
+				}
+
+				if (AsNumber <= 0)
+				{
+					OutErrorMessage = TEXT("Track and Channel filter indexes should be '1' based!");
+					return 0;
+				}
+
+				SetBit(AsNumber - 1);
+			}
+			else if (DashParts.Num() == 2)
+			{
+				int32 AsNumberStart;
+				int32 AsNumberEnd;
+				if (!LexTryParseString(AsNumberStart, *DashParts[0]))
+				{
+					OutErrorMessage = FString::Printf(TEXT("Error parsing part of filter string -> \"%s\""), *DashParts[0]);
+					return 0;
+				}
+
+				if (!LexTryParseString(AsNumberEnd, *DashParts[1]))
+				{
+					OutErrorMessage = FString::Printf(TEXT("Error parsing part of filter string -> \"%s\""), *DashParts[1]);
+					return 0;
+				}
+
+				if (AsNumberStart <= 0)
+				{
+					OutErrorMessage = TEXT("Track and Channel filter indexes should be '1' based!");
+					return 0;
+				}
+
+				if (AsNumberEnd < AsNumberStart)
+				{
+					OutErrorMessage = FString::Printf(TEXT("End of range must be greater than start of range (\"%s\")!"), *CommaPart);
+					return 0;
+				}
+
+				for (int32 i = AsNumberStart; i < (AsNumberEnd + 1); ++i)
+				{
+					SetBit(i - 1);
+				}
+			}
+		}
+		return Bitfield;
+	}
+
+	bool FMidiStreamEventTrackChannelFilter::SetMidiChannelFilterFromString(const FString& FilterString, FString& OutErrorMessage)
+	{
+		MidiChannelFilter = (uint16)StringToBitfield(FilterString, OutErrorMessage);
+		return OutErrorMessage.IsEmpty();
+	}
+
+	bool FMidiStreamEventTrackChannelFilter::SetTrackFilterFromString(const FString& FilterString, FString& OutErrorMessage)
+	{
+		TrackFilter = (uint64)StringToBitfield(FilterString, OutErrorMessage);
+		return OutErrorMessage.IsEmpty();
+	}
+
+	void FMidiVoiceTracker::Process(const FMidiStream& MidiStream, const FKillVoiceFn& KillVoiceFn)
+	{
+		// Remove active voices that have note offs this block
+		for (const auto& MidiEvent : MidiStream.EventsInBlock)
+		{
+			if (MidiEvent.MidiMessage.IsNoteOff())
+			{
+				ActiveVoices.Remove(MidiEvent.GetVoiceId());
+			}
+		}
+
+		// Check for voices that are no longer active and notify the caller
+		for (auto It = ActiveVoices.CreateIterator(); It; ++It)
+		{
+			if (!MidiStream.ActiveVoices.Contains(*It))
+			{
+				KillVoiceFn(*It);
+			}
+		}
+
+		ActiveVoices = MidiStream.ActiveVoices;
+	}
+
+}
diff --git a/Engine/Plugins/Runtime/Harmonix/Source/HarmonixMetasound/Private/HarmonixMetasound/Nodes/MediaStreamerTransportControllerNode.cpp b/Engine/Plugins/Runtime/Harmonix/Source/HarmonixMetasound/Private/HarmonixMetasound/Nodes/MediaStreamerTransportControllerNode.cpp
new file mode 100644
--- /dev/null
+++ b/Engine/Plugins/Runtime/Harmonix/Source/HarmonixMetasound/Private/HarmonixMetasound/Nodes/MediaStreamerTransportControllerNode.cpp
@@ -0,0 +1,344 @@
+// Copyright Epic Games, Inc. All Rights Reserved.
+
+#include "MetasoundExecutableOperator.h"
+#include "MetasoundFacade.h"
+#include "MetasoundNodeInterface.h"
+#include "MetasoundParamHelper.h"
+#include "MetasoundSampleCounter.h"
+#include "MetasoundStandardNodesCategories.h"
+#include "MetasoundVertex.h"
+
+#include "HarmonixMetasound/Common.h"
+#include "HarmonixMetasound/DataTypes/MidiAsset.h"
+#include "HarmonixMetasound/DataTypes/MidiClock.h"
+#include "HarmonixMetasound/DataTypes/MusicTransport.h"
+#include "HarmonixMetasound/DataTypes/MusicTimestamp.h"
+#include <algorithm>
+
+#define LOCTEXT_NAMESPACE "HarmonixMetaSound"
+
+DEFINE_LOG_CATEGORY_STATIC(MediaStreamerTransportControllerOperator, Log, All);
+
+namespace HarmonixMetasound
+{
+	using namespace Metasound;
+
+	namespace MediaStreamerTransportControllerVertexNames
+	{
+		METASOUND_PARAM(InputTriggerOnStreamStarts, "On Stream Starts", "Triggers when stream starts delivering data.")
+		METASOUND_PARAM(InputTriggerOnStreamStops, "On Stream Stops", "Triggers when stream stops delivering data.")
+		METASOUND_PARAM(InputPlaybackTime, "Playback Time", "Returns the current PTS of the meda stream.")
+		METASOUND_PARAM(InputPlaybackTimeEndBlock, "Playback Time at end of block", "Returns the current PTS of the meda stream at the end of the current block.")
+		METASOUND_PARAM(InputPlaybackRate, "Playback Rate", "Returns the current playback rate of the meda stream.")
+	}
+
+	class FMediaStreamerTransportControllerOperator : public TExecutableOperator<FMediaStreamerTransportControllerOperator>
+	{
+	public:
+		static const FNodeClassMetadata& GetNodeInfo();
+		static const FVertexInterface& GetVertexInterface();
+		static TUniquePtr<IOperator> CreateOperator(const FBuildOperatorParams& InParams, FBuildResults& OutResults);
+
+		FMediaStreamerTransportControllerOperator(
+			const FBuildOperatorParams& InParams,
+			const FMidiAssetReadRef& InMidiAsset,
+			const FInt32ReadRef& InPrerollBars,
+			const FTriggerReadRef& InTriggerOnStreamStarts,
+			const FTriggerReadRef& InTriggerOnStreamStops,
+			const FTimeReadRef& InPlaybackTime,
+			const FTimeReadRef& InPlaybackTimeEndBlock,
+			const FFloatReadRef& InPlaybackRate);
+
+		virtual void BindInputs(FInputVertexInterfaceData& InVertexData) override;
+		virtual void BindOutputs(FOutputVertexInterfaceData& InVertexData) override;
+		virtual FDataReferenceCollection GetInputs() const override;
+		virtual FDataReferenceCollection GetOutputs() const override;
+
+		void Reset(const FResetParams& ResetParams);
+		void Execute();
+
+	private:
+		//** INPUTS
+		FMidiAssetReadRef MidiAssetInPin;
+		FInt32ReadRef SeekPrerollBarsInPin;
+		FTriggerReadRef TriggerOnStreamStartsInPin;
+		FTriggerReadRef TriggerOnStreamStopsInPin;
+		FTimeReadRef PlaybackTimeInPin;
+		FTimeReadRef PlaybackTimeEndBlockInPin;
+		FFloatReadRef PlaybackRateInPin;
+
+		//** OUTPUTS
+		FMusicTransportEventStreamWriteRef TransportOutPin;
+		FMidiClockWriteRef MidiClockOutPin;
+
+		//** DATA
+		FSampleCount BlockSize;
+		float        SampleRate;
+		FSampleCount SampleCount;
+		static const int32 kMidiGranularity = 128;
+
+		FTime        StopTime;
+
+		EMusicPlayerTransportRequest LastTransportRequest = EMusicPlayerTransportRequest::None;
+	};
+
+	class FMediaStreamerTransportControllerNode : public FNodeFacade
+	{
+	public:
+		FMediaStreamerTransportControllerNode(const FNodeInitData& InInitData) :
+			FNodeFacade(InInitData.InstanceName, InInitData.InstanceID, TFacadeOperatorClass<FMediaStreamerTransportControllerOperator>())
+		{}
+		virtual ~FMediaStreamerTransportControllerNode() = default;
+	};
+
+	METASOUND_REGISTER_NODE(FMediaStreamerTransportControllerNode)
+}
+
+const Metasound::FNodeClassMetadata& HarmonixMetasound::FMediaStreamerTransportControllerOperator::GetNodeInfo()
+{
+	auto InitNodeInfo = []() -> FNodeClassMetadata
+	{
+		FNodeClassMetadata Info;
+		Info.ClassName = { HarmonixNodeNamespace, TEXT("MediaStreamerTransportController"), TEXT("") };
+		Info.MajorVersion = 0;
+		Info.MinorVersion = 1;
+		Info.DisplayName = METASOUND_LOCTEXT("MediaStreamerTransportControllerNode_DisplayName", "Media Streamer Transport Controller");
+		Info.Description = METASOUND_LOCTEXT("MediaStreamerTransportControllerNode_Description", "Provides a transport and midi clock controlled by signals from a Media Streamer using the tempo map in the specified Midi file.");
+		Info.Author = PluginAuthor;
+		Info.PromptIfMissing = PluginNodeMissingPrompt;
+		Info.DefaultInterface = GetVertexInterface();
+		Info.CategoryHierarchy.Emplace(NodeCategories::Music);
+
+		return Info;
+	};
+
+	static const FNodeClassMetadata Info = InitNodeInfo();
+
+	return Info;
+}
+
+const Metasound::FVertexInterface& HarmonixMetasound::FMediaStreamerTransportControllerOperator::GetVertexInterface()
+{
+	using namespace CommonPinNames;
+	using namespace MediaStreamerTransportControllerVertexNames;
+
+	static const FVertexInterface Interface(
+		FInputVertexInterface(
+			TInputDataVertex<FMidiAsset>(METASOUND_GET_PARAM_NAME_AND_METADATA(Inputs::MidiFileAsset)),
+			TInputDataVertex<int32>(METASOUND_GET_PARAM_NAME_AND_METADATA(Inputs::PrerollBars), 8),
+			TInputDataVertex<FTrigger>(METASOUND_GET_PARAM_NAME_AND_METADATA(InputTriggerOnStreamStarts)),
+			TInputDataVertex<FTrigger>(METASOUND_GET_PARAM_NAME_AND_METADATA(InputTriggerOnStreamStops)),
+			TInputDataVertex<FTime>(METASOUND_GET_PARAM_NAME_AND_METADATA(InputPlaybackTime)),
+			TInputDataVertex<FTime>(METASOUND_GET_PARAM_NAME_AND_METADATA(InputPlaybackTimeEndBlock)),
+			TInputDataVertex<float>(METASOUND_GET_PARAM_NAME_AND_METADATA(InputPlaybackRate))
+		),
+		FOutputVertexInterface(
+			TOutputDataVertex<FMusicTransportEventStream>(METASOUND_GET_PARAM_NAME_AND_METADATA(Outputs::Transport)),
+			TOutputDataVertex<FMidiClock>(METASOUND_GET_PARAM_NAME_AND_METADATA(Outputs::MidiClock))
+		)
+	);
+
+	return Interface;
+}
+
+TUniquePtr<Metasound::IOperator> HarmonixMetasound::FMediaStreamerTransportControllerOperator::CreateOperator(const FBuildOperatorParams& InParams, FBuildResults& OutResults)
+{
+	using namespace CommonPinNames;
+	using namespace MediaStreamerTransportControllerVertexNames;
+
+	const FInputVertexInterfaceData& InputData = InParams.InputData;
+
+	FMidiAssetReadRef MidiAsset = InputData.GetOrConstructDataReadReference<FMidiAsset>(METASOUND_GET_PARAM_NAME(Inputs::MidiFileAsset));
+	FInt32ReadRef PrerollBars = InputData.GetOrCreateDefaultDataReadReference<int32>(METASOUND_GET_PARAM_NAME(Inputs::PrerollBars), InParams.OperatorSettings);
+
+	FTriggerReadRef TriggerOnStreamStarts = InputData.GetOrConstructDataReadReference<FTrigger>(METASOUND_GET_PARAM_NAME(InputTriggerOnStreamStarts), InParams.OperatorSettings);
+	FTriggerReadRef TriggerOnStreamStops = InputData.GetOrConstructDataReadReference<FTrigger>(METASOUND_GET_PARAM_NAME(InputTriggerOnStreamStops), InParams.OperatorSettings);
+	FTimeReadRef PlaybackTime = InputData.GetOrConstructDataReadReference<FTime>(METASOUND_GET_PARAM_NAME(InputPlaybackTime));
+	FTimeReadRef PlaybackTimeEndBlock = InputData.GetOrConstructDataReadReference<FTime>(METASOUND_GET_PARAM_NAME(InputPlaybackTimeEndBlock));
+	FFloatReadRef PlaybackRate = InputData.GetOrConstructDataReadReference<float>(METASOUND_GET_PARAM_NAME(InputPlaybackRate));
+
+	return MakeUnique<FMediaStreamerTransportControllerOperator>(InParams, MidiAsset, PrerollBars, TriggerOnStreamStarts, TriggerOnStreamStops, PlaybackTime, PlaybackTimeEndBlock, PlaybackRate);
+}
+
+HarmonixMetasound::FMediaStreamerTransportControllerOperator::FMediaStreamerTransportControllerOperator(
+	const FBuildOperatorParams& InParams,
+	const FMidiAssetReadRef& InMidiAsset,
+	const FInt32ReadRef& InPrerollBars,
+	const FTriggerReadRef& InTriggerOnStreamStarts,
+	const FTriggerReadRef& InTriggerOnStreamStops,
+	const FTimeReadRef& InPlaybackTime,
+	const FTimeReadRef& InPlaybackTimeEndBlock,
+	const FFloatReadRef& InPlaybackRate) :
+	MidiAssetInPin(InMidiAsset),
+	SeekPrerollBarsInPin(InPrerollBars),
+	TriggerOnStreamStartsInPin(InTriggerOnStreamStarts),
+	TriggerOnStreamStopsInPin(InTriggerOnStreamStops),
+	PlaybackTimeInPin(InPlaybackTime),
+	PlaybackTimeEndBlockInPin(InPlaybackTimeEndBlock),
+	PlaybackRateInPin(InPlaybackRate),
+	TransportOutPin(FMusicTransportEventStreamWriteRef::CreateNew(InParams.OperatorSettings)),
+	MidiClockOutPin(FMidiClockWriteRef::CreateNew(InParams.OperatorSettings)),
+	BlockSize(InParams.OperatorSettings.GetNumFramesPerBlock()),
+	SampleRate(InParams.OperatorSettings.GetSampleRate())
+{
+	Reset(InParams);
+}
+
+void HarmonixMetasound::FMediaStreamerTransportControllerOperator::BindInputs(FInputVertexInterfaceData& InVertexData)
+{
+	using namespace CommonPinNames;
+	using namespace MediaStreamerTransportControllerVertexNames;
+
+	InVertexData.BindReadVertex(METASOUND_GET_PARAM_NAME(Inputs::MidiFileAsset), MidiAssetInPin);
+	InVertexData.BindReadVertex(METASOUND_GET_PARAM_NAME(Inputs::PrerollBars), SeekPrerollBarsInPin);
+	InVertexData.BindReadVertex(METASOUND_GET_PARAM_NAME(InputTriggerOnStreamStarts), TriggerOnStreamStartsInPin);
+	InVertexData.BindReadVertex(METASOUND_GET_PARAM_NAME(InputTriggerOnStreamStops), TriggerOnStreamStopsInPin);
+	InVertexData.BindReadVertex(METASOUND_GET_PARAM_NAME(InputPlaybackTime), PlaybackTimeInPin);
+	InVertexData.BindReadVertex(METASOUND_GET_PARAM_NAME(InputPlaybackTimeEndBlock), PlaybackTimeEndBlockInPin);
+	InVertexData.BindReadVertex(METASOUND_GET_PARAM_NAME(InputPlaybackRate), PlaybackRateInPin);
+}
+
+void HarmonixMetasound::FMediaStreamerTransportControllerOperator::BindOutputs(FOutputVertexInterfaceData& InVertexData)
+{
+	using namespace CommonPinNames;
+	using namespace MediaStreamerTransportControllerVertexNames;
+
+	InVertexData.BindReadVertex(METASOUND_GET_PARAM_NAME(Outputs::Transport), TransportOutPin);
+	InVertexData.BindReadVertex(METASOUND_GET_PARAM_NAME(Outputs::MidiClock), MidiClockOutPin);
+}
+
+Metasound::FDataReferenceCollection HarmonixMetasound::FMediaStreamerTransportControllerOperator::GetInputs() const
+{
+	// This should never be called. Bind(...) is called instead. This method
+	// exists as a stop-gap until the API can be deprecated and removed.
+	checkNoEntry();
+	return {};
+}
+
+Metasound::FDataReferenceCollection HarmonixMetasound::FMediaStreamerTransportControllerOperator::GetOutputs() const
+{
+	// This should never be called. Bind(...) is called instead. This method
+	// exists as a stop-gap until the API can be deprecated and removed.
+	checkNoEntry();
+	return {};
+}
+
+void HarmonixMetasound::FMediaStreamerTransportControllerOperator::Reset(const FResetParams& ResetParams)
+{
+	TransportOutPin->Reset();
+	MidiClockOutPin->PrepareBlock();
+	BlockSize = ResetParams.OperatorSettings.GetNumFramesPerBlock();
+	SampleRate = ResetParams.OperatorSettings.GetSampleRate();
+	SampleCount = 0;
+	StopTime = FTime{};
+	LastTransportRequest = EMusicPlayerTransportRequest::None;
+		
+	const FMidiFileProxyPtr& MidiProxy = MidiAssetInPin->GetMidiProxy();
+	MidiClockOutPin->AttachToMidiResource(MidiProxy ? MidiProxy->GetMidiFile() : nullptr);
+}
+
+void HarmonixMetasound::FMediaStreamerTransportControllerOperator::Execute()
+{
+	TransportOutPin->Reset();
+	MidiClockOutPin->PrepareBlock();
+	MidiClockOutPin->AddSpeedChangeToBlock({ 0, 0.0f, *PlaybackRateInPin });
+
+	// Parse triggers and render audio
+	int32 PlayTrigIndex = 0;
+	int32 NextPlayFrame = 0;
+	const int32 NumPlayTrigs = TriggerOnStreamStartsInPin->NumTriggeredInBlock();
+
+	int32 StopTrigIndex = 0;
+	int32 NextStopFrame = 0;
+	const int32 NumStopTrigs = TriggerOnStreamStopsInPin->NumTriggeredInBlock();
+
+	int32 CurrAudioFrame = 0;
+	int32 NextAudioFrame = 0;
+	const int32 LastAudioFrame = BlockSize - 1;
+	const int32 NoTrigger = BlockSize << 1;
+
+	while (NextAudioFrame < LastAudioFrame)
+	{
+		// get the next Play and Stop indices
+		// (play)
+		if (PlayTrigIndex < NumPlayTrigs)
+		{
+			NextPlayFrame = (*TriggerOnStreamStartsInPin)[PlayTrigIndex];
+		}
+		else
+		{
+			NextPlayFrame = NoTrigger;
+		}
+
+		// (stop)
+		if (StopTrigIndex < NumStopTrigs)
+		{
+			NextStopFrame = (*TriggerOnStreamStopsInPin)[StopTrigIndex];
+		}
+		else
+		{
+			NextStopFrame = NoTrigger;
+		}
+
+		// determine the next audio frame we are going to render up to
+		NextAudioFrame = FMath::Min(NextPlayFrame, NextStopFrame);
+
+		// no more triggers, rendering to the end of the block
+		if (NextAudioFrame == NoTrigger)
+		{
+			NextAudioFrame = BlockSize;
+		}
+
+		CurrAudioFrame = NextAudioFrame;
+
+		// execute the next trigger
+		if (CurrAudioFrame == NextPlayFrame)
+		{
+			switch (LastTransportRequest)
+			{
+			default:
+				LastTransportRequest = EMusicPlayerTransportRequest::Play;
+				TransportOutPin->AddTransportRequest(EMusicPlayerTransportRequest::Play, CurrAudioFrame);
+				MidiClockOutPin->ResetAndStart(CurrAudioFrame);
+				break;
+			case EMusicPlayerTransportRequest::Continue:
+			case EMusicPlayerTransportRequest::Play:
+				break;
+			case EMusicPlayerTransportRequest::Pause:
+				LastTransportRequest = EMusicPlayerTransportRequest::Continue;
+				TransportOutPin->AddTransportRequest(EMusicPlayerTransportRequest::Continue, CurrAudioFrame);
+				break;
+			}
+
+			FTime PlayTime = *PlaybackTimeEndBlockInPin - *PlaybackRateInPin * (BlockSize - CurrAudioFrame) / SampleRate;
+			if (!FMath::IsNearlyEqual(PlayTime.GetSeconds(), StopTime.GetSeconds(), 0.5 / SampleRate))
+			{
+				FMusicSeekTarget Target = { ESeekPointType::Millisecond };
+				Target.Ms = float(PlayTime.GetSeconds()) * 1000.f;
+				TransportOutPin->AddSeekRequest(CurrAudioFrame, Target);
+				MidiClockOutPin->SeekTo(Target, *SeekPrerollBarsInPin);
+			}
+
+			++PlayTrigIndex;
+		}
+
+		if (CurrAudioFrame == NextStopFrame)
+		{
+			StopTime = *PlaybackTimeInPin + *PlaybackRateInPin * CurrAudioFrame / SampleRate;
+			LastTransportRequest = EMusicPlayerTransportRequest::Pause;
+			TransportOutPin->AddTransportRequest(EMusicPlayerTransportRequest::Pause, CurrAudioFrame);
+
+			++StopTrigIndex;
+		}
+	}
+
+	switch (LastTransportRequest)
+	{
+	case EMusicPlayerTransportRequest::Continue:
+	case EMusicPlayerTransportRequest::Play:
+		MidiClockOutPin->AdvanceHiResToMs(BlockSize, float(PlaybackTimeEndBlockInPin->GetSeconds()) * 1000.f, true);
+		break;
+	}
+}
+
+#undef LOCTEXT_NAMESPACE // "HarmonixMetaSound"
diff --git a/Engine/Plugins/Runtime/Harmonix/Source/HarmonixMetasound/Private/HarmonixMetasound/Nodes/MetronomeNode.cpp b/Engine/Plugins/Runtime/Harmonix/Source/HarmonixMetasound/Private/HarmonixMetasound/Nodes/MetronomeNode.cpp
new file mode 100644
--- /dev/null
+++ b/Engine/Plugins/Runtime/Harmonix/Source/HarmonixMetasound/Private/HarmonixMetasound/Nodes/MetronomeNode.cpp
@@ -0,0 +1,627 @@
+// Copyright Epic Games, Inc. All Rights Reserved.
+
+#include "MetasoundExecutableOperator.h"
+#include "MetasoundFacade.h"
+#include "MetasoundNodeInterface.h"
+#include "MetasoundParamHelper.h"
+#include "MetasoundSampleCounter.h"
+#include "MetasoundStandardNodesCategories.h"
+#include "MetasoundVertex.h"
+
+#include "HarmonixMetasound/Common.h"
+#include "HarmonixMetasound/DataTypes/MidiClock.h"
+#include "HarmonixMetasound/DataTypes/MusicTransport.h"
+#include <algorithm>
+
+#include "HAL/PlatformMath.h"
+
+#define LOCTEXT_NAMESPACE "HarmonixMetaSound"
+
+DEFINE_LOG_CATEGORY_STATIC(LogMetronomeNode, Log, All);
+
+namespace HarmonixMetasound
+{
+	using namespace Metasound;
+
+	class FMetronomeOperator : public TExecutableOperator<FMetronomeOperator>, public FMusicTransportControllable, public FMidiPlayCursor
+	{
+	public:
+		static const FNodeClassMetadata& GetNodeInfo();
+		static const FVertexInterface& GetVertexInterface();
+		static TUniquePtr<IOperator> CreateOperator(const FBuildOperatorParams& InParams, FBuildResults& OutResults);
+
+		FMetronomeOperator(const FBuildOperatorParams& InParams,
+						   const FMusicTransportEventStreamReadRef& InTransport,
+						   const bool  InLoop,
+						   const int32 InLoopLengthBars, 
+						   const FInt32ReadRef& InTimSigNumerator,
+						   const FInt32ReadRef& InTimeSigDenominator,
+						   const FFloatReadRef& InTempo,
+		                   const FFloatReadRef& InSpeedMultiplier,
+		                   const int32 InSeekPrerollBars);
+
+		virtual void BindInputs(FInputVertexInterfaceData& InVertexData) override;
+		virtual void BindOutputs(FOutputVertexInterfaceData& InVertexData) override;
+		virtual FDataReferenceCollection GetInputs() const override;
+		virtual FDataReferenceCollection GetOutputs() const override;
+
+		void Reset(const FResetParams& Params);
+		void Execute();
+		void UpdateLoopOffsetTickFromTick(int32 Tick);
+
+		//~ BEGIN FMidiPlayCursor Overrides
+		virtual void Reset(bool ForceNoBroadcast = false) override;
+		virtual void OnLoop(int32 LoopStartTick, int32 LoopEndTick) override;
+		virtual void SeekToTick(int32 Tick) override;
+		virtual void SeekThruTick(int32 Tick) override;
+		virtual void AdvanceThruTick(int32 Tick, bool IsPreRoll) override;
+		//~ END FMidiPlayCursor Overrides
+
+	private:
+		void Init();
+		
+		//** INPUTS
+		FMusicTransportEventStreamReadRef TransportInPin;
+		const bool  LoopInPin;
+		const int32 LoopLengthBarsInPin;
+		FInt32ReadRef TimeSigNumInPin;
+		FInt32ReadRef TimeSigDenomInPin;
+		FFloatReadRef TempoInPin;
+		FFloatReadRef SpeedMultInPin;
+		const int32 SeekPreRollBarsInPin;
+
+		//** OUTPUTS
+		FMidiClockWriteRef MidiClockOutPin;
+
+		//** DATA
+		FMidiClock MetronomeClock;
+		FMidiClockEventCursor MidiClockEventCursor;
+		TSharedPtr<FMidiFileData> MidiData;
+		FSampleCount BlockSize;
+		float        SampleRate;
+		float        CurrentTempo;
+		int32        CurrentTimeSigNum;
+		int32        CurrentTimeSigDenom;
+		int32		 LastClockTickUpdate = -1;
+		int32		 LoopOffsetTick = 0;
+
+		void BuildMidiData();
+		void UpdateMidi();
+		void AddTempoChangeForMidi(float TempoBPM);
+		void AddTimeSigChangeForMidi(int32 TimeSigNum, int32 TimeSigDenom);
+		void AddTransportStateChangeToMidiClock(int32 StartFrameIndex, EMusicPlayerTransportState NewTransportState);
+
+		FMidiClock& GetDrivingMidiClock() { return LoopInPin ? MetronomeClock : (*MidiClockOutPin); }
+	};
+
+	class FMetronomeNode : public FNodeFacade
+	{
+	public:
+		FMetronomeNode(const FNodeInitData& InInitData)
+			: FNodeFacade(InInitData.InstanceName, InInitData.InstanceID, TFacadeOperatorClass<FMetronomeOperator>())
+		{}
+		virtual ~FMetronomeNode() = default;
+	};
+
+	METASOUND_REGISTER_NODE(FMetronomeNode)
+
+	const FNodeClassMetadata& FMetronomeOperator::GetNodeInfo()
+	{
+		auto InitNodeInfo = []() -> FNodeClassMetadata
+		{
+			FNodeClassMetadata Info;
+			Info.ClassName        = { HarmonixNodeNamespace, TEXT("Metronome"), TEXT("") };
+			Info.MajorVersion     = 0;
+			Info.MinorVersion     = 1;
+			Info.DisplayName      = METASOUND_LOCTEXT("MetronomeNode_DisplayName", "Metronome Midi Clock Generator");
+			Info.Description      = METASOUND_LOCTEXT("MetronomeNode_Description", "Provides a midi clock at the specified tempo and speed.");
+			Info.Author           = PluginAuthor;
+			Info.PromptIfMissing  = PluginNodeMissingPrompt;
+			Info.DefaultInterface = GetVertexInterface();
+			Info.CategoryHierarchy.Emplace(NodeCategories::Music);
+
+			return Info;
+		};
+
+		static const FNodeClassMetadata Info = InitNodeInfo();
+
+		return Info;
+	}
+
+	const FVertexInterface& FMetronomeOperator::GetVertexInterface()
+	{
+		using namespace CommonPinNames;
+
+		static const FVertexInterface Interface(
+			FInputVertexInterface(
+				TInputDataVertex<FMusicTransportEventStream>(METASOUND_GET_PARAM_NAME_AND_METADATA(Inputs::Transport)),
+				TInputConstructorVertex<bool>(METASOUND_GET_PARAM_NAME_AND_METADATA(Inputs::Loop), false),
+				TInputConstructorVertex<int32>(METASOUND_GET_PARAM_NAME_AND_METADATA(Inputs::LoopLengthBars), 4),
+				TInputDataVertex<int32>(METASOUND_GET_PARAM_NAME_AND_METADATA(Inputs::TimeSigNumerator), 4),
+				TInputDataVertex<int32>(METASOUND_GET_PARAM_NAME_AND_METADATA(Inputs::TimeSigDenominator), 4),
+				TInputDataVertex<float>(METASOUND_GET_PARAM_NAME_AND_METADATA(Inputs::Tempo), 120.0f),
+				TInputDataVertex<float>(METASOUND_GET_PARAM_NAME_AND_METADATA(Inputs::Speed), 1.0f),
+				TInputConstructorVertex<int32>(METASOUND_GET_PARAM_NAME_AND_METADATA(Inputs::PrerollBars), 8)
+			),
+			FOutputVertexInterface(
+				TOutputDataVertex<FMidiClock>(METASOUND_GET_PARAM_NAME_AND_METADATA(Outputs::MidiClock))
+			)
+		);
+
+		return Interface;
+	}
+
+	TUniquePtr<IOperator> FMetronomeOperator::CreateOperator(const FBuildOperatorParams& InParams, FBuildResults& OutResults)
+	{
+		using namespace CommonPinNames;
+
+		const FMetronomeNode& TempoClockNode = static_cast<const FMetronomeNode&>(InParams.Node);
+
+		const FInputVertexInterfaceData& InputData = InParams.InputData;
+		const FOperatorSettings& Settings = InParams.OperatorSettings;
+
+		FMusicTransportEventStreamReadRef InTransport = InputData.GetOrConstructDataReadReference<FMusicTransportEventStream>(METASOUND_GET_PARAM_NAME(Inputs::Transport), Settings);
+		bool InLoop = InputData.GetOrCreateDefaultValue<bool>(METASOUND_GET_PARAM_NAME(Inputs::Loop), Settings);
+		int32 InLoopLengthBars = InputData.GetOrCreateDefaultValue<int32>(METASOUND_GET_PARAM_NAME(Inputs::LoopLengthBars), Settings);
+		FInt32ReadRef InTimeSigNumerator   = InputData.GetOrCreateDefaultDataReadReference<int32>(METASOUND_GET_PARAM_NAME(Inputs::TimeSigNumerator), Settings);
+		FInt32ReadRef InTimeSigDenominator = InputData.GetOrCreateDefaultDataReadReference<int32>(METASOUND_GET_PARAM_NAME(Inputs::TimeSigDenominator), Settings);
+		FFloatReadRef InTempo = InputData.GetOrCreateDefaultDataReadReference<float>(METASOUND_GET_PARAM_NAME(Inputs::Tempo), Settings);
+		FFloatReadRef InSpeed = InputData.GetOrCreateDefaultDataReadReference<float>(METASOUND_GET_PARAM_NAME(Inputs::Speed), Settings);
+		int32 InPreRollBars = InputData.GetOrCreateDefaultValue<int32>(METASOUND_GET_PARAM_NAME(Inputs::PrerollBars), Settings);
+
+		return MakeUnique<FMetronomeOperator>(InParams, InTransport, InLoop, InLoopLengthBars, InTimeSigNumerator, InTimeSigDenominator, InTempo, InSpeed, InPreRollBars);
+	}
+
+	FMetronomeOperator::FMetronomeOperator(const FBuildOperatorParams& InParams, 
+	                                       const FMusicTransportEventStreamReadRef& InTransport,
+										   const bool  InLoop,
+										   const int32 InLoopLengthBars,
+	                                       const FInt32ReadRef& InTimSigNumerator,
+	                                       const FInt32ReadRef& InTimeSigDenominator,
+	                                       const FFloatReadRef& InTempo,
+	                                       const FFloatReadRef& InSpeedMultiplier,
+		                                   const int32 InPreRollBars)
+		: FMusicTransportControllable(EMusicPlayerTransportState::Prepared) 
+		, TransportInPin(InTransport)
+		, LoopInPin(InLoop)
+		, LoopLengthBarsInPin(InLoopLengthBars)
+		, TimeSigNumInPin(InTimSigNumerator)
+		, TimeSigDenomInPin(InTimeSigDenominator)
+		, TempoInPin(InTempo)
+		, SpeedMultInPin(InSpeedMultiplier)
+		, SeekPreRollBarsInPin(InPreRollBars)
+		, MidiClockOutPin(FMidiClockWriteRef::CreateNew(InParams.OperatorSettings))
+		, MetronomeClock(InParams.OperatorSettings)
+		, MidiClockEventCursor(MidiClockOutPin)
+		, BlockSize(InParams.OperatorSettings.GetNumFramesPerBlock())
+		, SampleRate(InParams.OperatorSettings.GetSampleRate())
+		, CurrentTempo(*TempoInPin)
+		, CurrentTimeSigNum(FMath::Clamp(*TimeSigNumInPin, 1, 64))
+		, CurrentTimeSigDenom(FMath::Clamp(*TimeSigDenomInPin, 1, 64))
+	{
+		MetronomeClock.RegisterHiResPlayCursor(this);
+
+		if (LoopInPin)
+		{
+			MidiClockOutPin->AttachToTimeAuthority(MetronomeClock);
+		}
+
+		Reset(InParams);
+
+		Init();
+	}
+
+	void FMetronomeOperator::BindInputs(FInputVertexInterfaceData& InVertexData)
+	{
+		using namespace CommonPinNames;
+		InVertexData.BindReadVertex(METASOUND_GET_PARAM_NAME(Inputs::Transport), TransportInPin);
+		InVertexData.SetValue(METASOUND_GET_PARAM_NAME(Inputs::Loop), LoopInPin);
+		InVertexData.SetValue(METASOUND_GET_PARAM_NAME(Inputs::LoopLengthBars), LoopLengthBarsInPin);
+		InVertexData.BindReadVertex(METASOUND_GET_PARAM_NAME(Inputs::TimeSigNumerator), TimeSigNumInPin);
+		InVertexData.BindReadVertex(METASOUND_GET_PARAM_NAME(Inputs::TimeSigDenominator), TimeSigDenomInPin);
+		InVertexData.BindReadVertex(METASOUND_GET_PARAM_NAME(Inputs::Tempo), TempoInPin);
+		InVertexData.BindReadVertex(METASOUND_GET_PARAM_NAME(Inputs::Speed), SpeedMultInPin);
+		InVertexData.SetValue(METASOUND_GET_PARAM_NAME(Inputs::PrerollBars), SeekPreRollBarsInPin);
+
+		Init();
+	}
+
+	void FMetronomeOperator::BindOutputs(FOutputVertexInterfaceData& InVertexData)
+	{
+		using namespace CommonPinNames;
+		InVertexData.BindReadVertex(METASOUND_GET_PARAM_NAME(Outputs::MidiClock), MidiClockOutPin);
+	}
+
+	FDataReferenceCollection FMetronomeOperator::GetInputs() const
+	{
+		// This should never be called. Bind(...) is called instead. This method
+		// exists as a stop-gap until the API can be deprecated and removed.
+		checkNoEntry();
+		return {};
+	}
+
+	FDataReferenceCollection FMetronomeOperator::GetOutputs() const
+	{
+		// This should never be called. Bind(...) is called instead. This method
+		// exists as a stop-gap until the API can be deprecated and removed.
+		checkNoEntry();
+		return {};
+	}
+
+	void FMetronomeOperator::Reset(const FResetParams& Params)
+	{
+		BlockSize = Params.OperatorSettings.GetNumFramesPerBlock();
+		SampleRate = Params.OperatorSettings.GetSampleRate();
+		
+		LastClockTickUpdate = -1;
+	}
+
+	void FMetronomeOperator::Reset(bool ForceNoBroadcast /*= false*/)
+	{
+		FMidiPlayCursor::Reset(ForceNoBroadcast);
+		int32 Tick = MetronomeClock.GetCurrentHiResTick();
+		UpdateLoopOffsetTickFromTick(Tick);
+		int32 MappedTick = Tick - LoopOffsetTick;
+
+		float Ms = MidiClockOutPin->GetSongMaps().TickToMs(Tick + 1);
+		FMusicSeekTarget SeekTarget;
+		SeekTarget.Type = ESeekPointType::Millisecond;
+		SeekTarget.Ms = Ms;
+		MidiClockOutPin->SeekTo(MetronomeClock.GetCurrentBlockFrameIndex(), SeekTarget, SeekPreRollBarsInPin);
+	}
+
+	void FMetronomeOperator::OnLoop(int32 LoopStartTick, int32 LoopEndTick)
+	{
+		FMidiPlayCursor::OnLoop(LoopStartTick, LoopEndTick);
+	}
+
+	void FMetronomeOperator::SeekToTick(int32 Tick)
+	{
+		FMidiPlayCursor::SeekToTick(Tick);
+		UpdateLoopOffsetTickFromTick(Tick);
+		int32 MappedTick = Tick - LoopOffsetTick;
+
+		float Ms = MidiClockOutPin->GetSongMaps().TickToMs(Tick + 1);
+		FMusicSeekTarget SeekTarget;
+		SeekTarget.Type = ESeekPointType::Millisecond;
+		SeekTarget.Ms = Ms;
+		MidiClockOutPin->SeekTo(MetronomeClock.GetCurrentBlockFrameIndex(), SeekTarget, SeekPreRollBarsInPin);
+	}
+
+	void FMetronomeOperator::SeekThruTick(int32 Tick)
+	{
+		FMidiPlayCursor::SeekThruTick(Tick);
+
+		UpdateLoopOffsetTickFromTick(Tick);
+		int32 MappedTick = Tick - LoopOffsetTick;
+
+		float Ms = MidiClockOutPin->GetSongMaps().TickToMs(Tick + 1);
+		FMusicSeekTarget SeekTarget;
+		SeekTarget.Type = ESeekPointType::Millisecond;
+		SeekTarget.Ms = Ms;
+		MidiClockOutPin->SeekTo(MetronomeClock.GetCurrentBlockFrameIndex(), SeekTarget, SeekPreRollBarsInPin);
+	}
+
+	void FMetronomeOperator::AdvanceThruTick(int32 InTick, bool IsPreRoll)
+	{
+		FMidiPlayCursor::AdvanceThruTick(InTick, IsPreRoll);
+
+		int32 MappedTick = InTick - LoopOffsetTick;
+		float CurrentSpeed = MetronomeClock.GetSpeedAtBlockSampleFrame(MetronomeClock.GetCurrentBlockFrameIndex());
+		float Ms = MidiClockOutPin->GetSongMaps().TickToMs(MappedTick);
+		float AdvanceRatio = Owner->GetSongMaps().GetTempoAtTick(InTick) / MidiClockOutPin->GetSongMaps().GetTempoAtTick(MappedTick);
+		MidiClockOutPin->InformOfCurrentAdvanceRate(CurrentSpeed * AdvanceRatio);
+		MidiClockOutPin->AdvanceHiResToMs(MetronomeClock.GetCurrentBlockFrameIndex(), Ms, true);
+
+		if (MidiClockOutPin->DoesLoop())
+		{
+			if (MidiClockOutPin->GetCurrentHiResTick() < MappedTick)
+			{
+				UpdateLoopOffsetTickFromTick(InTick);
+			}
+		}
+	}
+
+	void FMetronomeOperator::Init()
+	{
+		BuildMidiData();
+		
+		FTransportInitFn InitFn = [this](EMusicPlayerTransportState CurrentState)
+		{
+			FMidiClock& DrivingMidiClock = GetDrivingMidiClock();
+			
+			switch (CurrentState)
+			{
+			case EMusicPlayerTransportState::Invalid:
+			case EMusicPlayerTransportState::Preparing:
+				MidiClockOutPin->AddTransportStateChangeToBlock({ 0, 0.0f, EMusicPlayerTransportState::Prepared });
+				return EMusicPlayerTransportState::Prepared;
+
+			case EMusicPlayerTransportState::Prepared:
+				return CurrentState;
+
+			case EMusicPlayerTransportState::Starting:
+				DrivingMidiClock.ResetAndStart(0, true);
+				return EMusicPlayerTransportState::Playing;
+
+			case EMusicPlayerTransportState::Playing:
+				DrivingMidiClock.WriteAdvance(0, 0, *SpeedMultInPin);
+				return EMusicPlayerTransportState::Playing;
+
+			case EMusicPlayerTransportState::Seeking:
+				checkNoEntry();
+				return EMusicPlayerTransportState::Invalid;
+
+			case EMusicPlayerTransportState::Continuing:
+				MidiClockOutPin->AddTransportStateChangeToBlock({ 0, 0.0f, EMusicPlayerTransportState::Playing });
+				return EMusicPlayerTransportState::Playing;
+
+			case EMusicPlayerTransportState::Pausing:
+				MidiClockOutPin->AddTransportStateChangeToBlock({ 0, 0.0f, EMusicPlayerTransportState::Paused });
+				return EMusicPlayerTransportState::Paused;
+
+			case EMusicPlayerTransportState::Paused:
+				return EMusicPlayerTransportState::Paused;
+
+			case EMusicPlayerTransportState::Stopping:
+				MidiClockOutPin->AddTransportStateChangeToBlock({ 0, 0.0f, EMusicPlayerTransportState::Prepared });
+				return EMusicPlayerTransportState::Prepared;
+
+			case EMusicPlayerTransportState::Killing:
+				MidiClockOutPin->AddTransportStateChangeToBlock({ 0, 0.0f, EMusicPlayerTransportState::Prepared });
+				return EMusicPlayerTransportState::Prepared;
+
+			default:
+				checkNoEntry();
+				return EMusicPlayerTransportState::Invalid;
+			}
+		};
+		
+		FMusicTransportControllable::Init(*TransportInPin, MoveTemp(InitFn));
+	}
+
+	void FMetronomeOperator::Execute()
+	{
+		MidiClockOutPin->PrepareBlock();
+
+		FMidiClock& DrivingMidiClock = GetDrivingMidiClock();
+		DrivingMidiClock.PrepareBlock();
+
+		if (*SpeedMultInPin != DrivingMidiClock.GetSpeedAtEndOfBlock())
+		{
+			DrivingMidiClock.AddSpeedChangeToBlock({ 0, 0.0f, *SpeedMultInPin });
+		}
+
+		// only update our midi data if the clock is advancing
+		// update the tempo and time sig before we advance our clock 
+		int32 ClockTick = DrivingMidiClock.GetCurrentHiResTick();
+		if (ClockTick >= 0 && ClockTick > LastClockTickUpdate)
+		{
+			UpdateMidi();
+			LastClockTickUpdate = ClockTick;
+		}
+		else if (ClockTick == -1 && LastClockTickUpdate == -1)
+		{
+			// check if we need to update our midi data based on new inputs
+			int32 InTimeSigNum = FMath::Clamp(*TimeSigNumInPin, 1, 64);
+			int32 InTimeSigDenom = FMath::Clamp(*TimeSigDenomInPin, 1, 64);
+			bool RebuildMidi = *TempoInPin > 0 && !FMath::IsNearlyEqual(CurrentTempo, *TempoInPin);
+			RebuildMidi |= InTimeSigNum != CurrentTimeSigNum || InTimeSigDenom != CurrentTimeSigDenom;
+			if (RebuildMidi)
+			{
+				BuildMidiData();
+			}
+		}
+
+		TransportSpanProcessor TransportHandler = [this, &DrivingMidiClock](int32 StartFrameIndex, int32 EndFrameIndex, EMusicPlayerTransportState CurrentState)
+		{
+			switch (CurrentState)
+			{
+			case EMusicPlayerTransportState::Invalid:
+			case EMusicPlayerTransportState::Preparing:
+				AddTransportStateChangeToMidiClock(StartFrameIndex, EMusicPlayerTransportState::Prepared);
+				DrivingMidiClock.WriteNoAdvance(StartFrameIndex, EndFrameIndex);
+				return EMusicPlayerTransportState::Prepared;
+
+			case EMusicPlayerTransportState::Prepared:
+				DrivingMidiClock.WriteNoAdvance(StartFrameIndex, EndFrameIndex);
+				return CurrentState;
+
+			case EMusicPlayerTransportState::Starting:
+				// We need to rebuild the midi data as there may have been a previous 'play' 
+				// and during that 'play' there may have been tempo changes added. When doing a
+				// 'start from a stop' we want a clean slate...
+				BuildMidiData();
+				// Play from the beginning if we haven't received a seek call while we were stopped...
+				if (!ReceivedSeekWhileStopped())
+				{
+					LastClockTickUpdate = -1;
+				}
+				DrivingMidiClock.ResetAndStart(StartFrameIndex, !ReceivedSeekWhileStopped());
+				DrivingMidiClock.WriteAdvance(StartFrameIndex, EndFrameIndex, *SpeedMultInPin);
+				return EMusicPlayerTransportState::Playing;
+
+			case EMusicPlayerTransportState::Playing:
+				DrivingMidiClock.WriteAdvance(StartFrameIndex, EndFrameIndex, *SpeedMultInPin);
+				return EMusicPlayerTransportState::Playing;
+
+			case EMusicPlayerTransportState::Seeking:
+				DrivingMidiClock.SeekTo(StartFrameIndex, TransportInPin->GetNextSeekDestination(), SeekPreRollBarsInPin);
+				// Here we will return that we want to be in the same state we were in before this request to 
+				// seek since we can seek "instantaneously"...
+				return GetTransportState();
+
+			case EMusicPlayerTransportState::Continuing:
+				AddTransportStateChangeToMidiClock(StartFrameIndex, EMusicPlayerTransportState::Playing);
+				DrivingMidiClock.WriteAdvance(StartFrameIndex, EndFrameIndex, *SpeedMultInPin);
+				return EMusicPlayerTransportState::Playing;
+
+			case EMusicPlayerTransportState::Pausing:
+				AddTransportStateChangeToMidiClock(StartFrameIndex, EMusicPlayerTransportState::Paused);
+				DrivingMidiClock.WriteNoAdvance(StartFrameIndex, EndFrameIndex);
+				return EMusicPlayerTransportState::Paused;
+
+			case EMusicPlayerTransportState::Paused:
+				DrivingMidiClock.WriteNoAdvance(StartFrameIndex, EndFrameIndex);
+				return EMusicPlayerTransportState::Paused;
+
+			case EMusicPlayerTransportState::Stopping:
+				AddTransportStateChangeToMidiClock(StartFrameIndex, EMusicPlayerTransportState::Prepared);
+				DrivingMidiClock.WriteNoAdvance(StartFrameIndex, EndFrameIndex);
+				return EMusicPlayerTransportState::Prepared;
+
+			case EMusicPlayerTransportState::Killing:
+				AddTransportStateChangeToMidiClock(StartFrameIndex, EMusicPlayerTransportState::Prepared);
+				DrivingMidiClock.WriteNoAdvance(StartFrameIndex, EndFrameIndex);
+				return EMusicPlayerTransportState::Prepared;
+
+			default:
+				checkNoEntry();
+				return EMusicPlayerTransportState::Invalid;
+			}
+		};
+		ExecuteTransportSpans(TransportInPin, BlockSize, TransportHandler);
+
+		if (LoopInPin)
+		{
+			MidiClockOutPin->CopySpeedAndTempoChanges(&MetronomeClock);
+		}
+	}
+
+	void FMetronomeOperator::UpdateLoopOffsetTickFromTick(int32 Tick)
+	{
+		// only update our loop offset if we're actually looping
+		// if our external clock is seeking us
+		// we need to update our loop offset ticks
+		if (MidiClockOutPin->DoesLoop())
+		{
+			int32 LoopStartTick = MidiClockOutPin->GetLoopStartTick();
+			int32 LoopEndTick = MidiClockOutPin->GetLoopEndTick();
+			int32 LoopLengthTicks = LoopEndTick - LoopStartTick;
+			if (LoopLengthTicks > 0)
+			{
+				// the whole number part of the division tells us 
+				// how many times we "should have looped" based on the incoming tick
+				int32 LoopNum = (Tick - LoopStartTick) / LoopLengthTicks;
+
+				// So we can multiply it back to get the LoopOffsetTick
+				LoopOffsetTick = LoopLengthTicks * LoopNum + LoopStartTick;
+			}
+		}
+	}
+
+	void FMetronomeOperator::BuildMidiData()
+	{
+		// make sure we have valid values
+		CurrentTempo = FMath::Max(1.0f, *TempoInPin);
+		CurrentTimeSigNum = FMath::Max(1, *TimeSigNumInPin);
+		CurrentTimeSigDenom = FMath::Max(1, *TimeSigDenomInPin);
+	
+		MidiData = FMidiClock::MakeClockConductorMidiData(CurrentTempo, CurrentTimeSigNum, CurrentTimeSigDenom);
+		
+		if (LoopInPin)
+		{
+			MetronomeClock.AttachToMidiResource(MidiData);
+
+			// midi clock out will follow the tempo of the metronome
+			// so just assign it to some reasonable values. They will be ignored
+			TSharedPtr<FMidiFileData> MidiDataOut = FMidiClock::MakeClockConductorMidiData(120.0f, CurrentTimeSigNum, CurrentTimeSigDenom);
+			MidiClockOutPin->AttachToMidiResource(MidiDataOut);
+			
+			int32 LoopEndTick = MidiDataOut->SongMaps.GetBarMap().BarIncludingCountInToTick(FMath::Max(LoopLengthBarsInPin, 1));
+			MidiClockOutPin->SetLoop(0, LoopEndTick);
+		}
+		else
+		{
+			// if we're not looping, then the MidiClockOut is going to be the driving midi clock
+			// so attach the midi data directly to it
+			MidiClockOutPin->AttachToMidiResource(MidiData);
+			MidiClockOutPin->ClearLoop();
+		}
+	}
+
+	void FMetronomeOperator::UpdateMidi()
+	{
+		FMidiClock& DrivingMidiClock = GetDrivingMidiClock();
+		bool HasMidiChanges = false;
+		if (*TempoInPin > 0 && !FMath::IsNearlyEqual(CurrentTempo, *TempoInPin))
+		{
+			if (!HasMidiChanges)
+			{
+				DrivingMidiClock.LockForMidiDataChanges();
+			}
+			AddTempoChangeForMidi(*TempoInPin);
+			HasMidiChanges = true;
+		};
+
+		int32 InTimeSigNum = FMath::Clamp(*TimeSigNumInPin, 1, 64);
+		int32 InTimeSigDenom = FMath::Clamp(*TimeSigDenomInPin, 1, 64);
+		if (InTimeSigNum != CurrentTimeSigNum || InTimeSigDenom != CurrentTimeSigDenom)
+		{
+			CurrentTimeSigNum = InTimeSigNum;
+			CurrentTimeSigDenom = InTimeSigDenom;
+			if (ensure(!MidiClockOutPin->DoesLoop()))
+			{
+				if (!HasMidiChanges)
+				{
+					DrivingMidiClock.LockForMidiDataChanges();
+				}
+				AddTimeSigChangeForMidi(InTimeSigNum, InTimeSigDenom);
+				HasMidiChanges = true;
+			}
+			else
+			{
+				UE_LOG(LogMetronomeNode, Warning, TEXT("Changing Time Sig. on looping metronome not supported." 
+					"Changing time signature will require changing loop length which is currently not supported."))
+			}
+		}
+
+		if (HasMidiChanges)
+		{
+			DrivingMidiClock.MidiDataChangesComplete();
+		}
+	}
+
+	void FMetronomeOperator::AddTempoChangeForMidi(float InTempoBPM)
+	{
+		CurrentTempo = InTempoBPM;
+		int32 AtTick = GetDrivingMidiClock().GetCurrentHiResTick() + 1;
+		MidiData->AddTempoChange(0, AtTick, CurrentTempo);
+	}
+
+	void FMetronomeOperator::AddTimeSigChangeForMidi(int32 InTimeSigNum, int32 InTimeSigDenom)
+	{
+		CurrentTimeSigNum = InTimeSigNum;
+		CurrentTimeSigDenom = InTimeSigDenom;
+		int32 AtTick = GetDrivingMidiClock().GetCurrentHiResTick() + 1;
+		// round to the next bar boundary, the bar we're actually going to apply the time sig change to
+		int32 AtBar = FMath::CeilToInt32(MidiData->SongMaps.GetBarIncludingCountInAtTick(AtTick));
+		int32 NumTimeSigPoints = MidiData->SongMaps.GetBarMap().GetNumTimeSignaturePoints();
+
+		// check if there's already a time signature point at the bar we're trying to update
+		// the metronome clock increases monotonically, so we just have to check the _last_ point
+		FTimeSignaturePoint& Point = MidiData->SongMaps.GetBarMap().GetTimeSignaturePoint(NumTimeSigPoints - 1);
+		if (Point.BarIndex == AtBar)
+		{
+			// update that instead of adding a new one
+			Point.TimeSignature.Numerator = InTimeSigNum;
+			Point.TimeSignature.Denominator = InTimeSigDenom;
+		}
+		else
+		{
+			MidiData->AddTimeSigChange(0, AtTick, CurrentTimeSigNum, CurrentTimeSigDenom);
+		}
+	}
+
+	void FMetronomeOperator::AddTransportStateChangeToMidiClock(int32 StartFrameIndex, EMusicPlayerTransportState NewTransportState)
+	{
+		FMidiTimestampTransportState NewState{ StartFrameIndex, 0.0f, NewTransportState };
+		MidiClockOutPin->AddTransportStateChangeToBlock(NewState);
+		if (LoopInPin)
+		{
+			MetronomeClock.AddTransportStateChangeToBlock(NewState);
+		}
+	}
+}
+
+#undef LOCTEXT_NAMESPACE // "HarmonixMetaSound"
diff --git a/Engine/Plugins/Runtime/Harmonix/Source/HarmonixMetasound/Private/HarmonixMetasound/Nodes/MidiCCTriggerNode.cpp b/Engine/Plugins/Runtime/Harmonix/Source/HarmonixMetasound/Private/HarmonixMetasound/Nodes/MidiCCTriggerNode.cpp
new file mode 100644
--- /dev/null
+++ b/Engine/Plugins/Runtime/Harmonix/Source/HarmonixMetasound/Private/HarmonixMetasound/Nodes/MidiCCTriggerNode.cpp
@@ -0,0 +1,235 @@
+// Copyright Epic Games, Inc. All Rights Reserved.
+
+#include "MetasoundExecutableOperator.h"
+#include "MetasoundFacade.h"
+#include "MetasoundNodeInterface.h"
+#include "MetasoundParamHelper.h"
+#include "MetasoundSampleCounter.h"
+#include "MetasoundStandardNodesCategories.h"
+#include "MetasoundVertex.h"
+#include "MetasoundTrigger.h"
+#include "HarmonixMetasound/Nodes/MidiCCTriggerNode.h"
+#include "HarmonixMetasound/Common.h"
+#include "HarmonixMetasound/DataTypes/MidiStream.h"
+#include "HarmonixMetasound/DataTypes/MusicTransport.h"
+#include "HarmonixMidi/MidiMsg.h"
+#include "MetasoundEnumRegistrationMacro.h"
+#include "HarmonixMetasound/DataTypes/MidiControllerID.h"
+
+DEFINE_LOG_CATEGORY_STATIC(LogMidiCCTrigger, Log, All);
+
+#define LOCTEXT_NAMESPACE "HarmonixMetaSound"
+
+namespace HarmonixMetasound::Nodes::MidiCCTriggerNode
+{
+	const Metasound::FNodeClassName& GetClassName()
+	{
+		static const Metasound::FNodeClassName ClassName
+		{
+			HarmonixNodeNamespace,
+				TEXT("MidiCCTrigger"),
+				""
+		};
+		return ClassName;
+	}
+
+	namespace Inputs
+	{
+		EXTERN_METASOUND_PARAM_ALIAS(Enable, CommonPinNames::Inputs::Enable);
+		EXTERN_METASOUND_PARAM_ALIAS(MidiTrackNumber, CommonPinNames::Inputs::MidiTrackNumber);
+		EXTERN_METASOUND_PARAM_ALIAS(MidiChannelNumber, CommonPinNames::Inputs::MidiChannelNumber);
+		EXTERN_METASOUND_PARAM_D(InputMidiControllerID, "Standard Midi Controller ID", "Standard Midi Controller ID (0-127)");
+		EXTERN_METASOUND_PARAM_ALIAS(MidiStream, CommonPinNames::Inputs::MidiStream);
+	}
+
+	namespace Outputs
+	{
+		EXTERN_METASOUND_PARAM_D(OutputControlChangeValueInt32, "Midi Control Change Value (Int32)", "Midi Control Change value (0-127)");
+		EXTERN_METASOUND_PARAM_D(OutputControlChangeValueFloat, "Midi Control Change Value (Float)", "normalized Midi Control Change value (0.0-1.0)");
+		EXTERN_METASOUND_PARAM_D(OutputTrigger, "Trigger Out", "A trigger when a Midi Control Change message is found");
+	}
+
+	class FOp final : public Metasound::TExecutableOperator<FOp>
+	{
+	public:
+		struct FInputs
+		{
+			Metasound::FBoolReadRef       Enable;
+			Metasound::FInt32ReadRef      TrackNumber;
+			Metasound::FInt32ReadRef      ChannelNumber;
+			Metasound::FEnumStdMidiControllerIDReadRef  ControllerID;
+			HarmonixMetasound::FMidiStreamReadRef MidiStream;
+		};
+
+		struct FOutputs
+		{
+			Metasound::FInt32WriteRef ControlChangeValueInt32;
+			Metasound::FFloatWriteRef ControlChangeValueFloat;
+			Metasound::FTriggerWriteRef TriggerOut;
+		};
+
+		static const Metasound::FVertexInterface& GetVertexInterface()
+		{
+			const auto MakeInterface = []() -> Metasound::FVertexInterface
+			{
+				using namespace Metasound;
+
+				return
+				{
+					FInputVertexInterface
+					{
+						TInputDataVertex<bool>(METASOUND_GET_PARAM_NAME_AND_METADATA(Inputs::Enable), true),
+						TInputDataVertex<int32>(METASOUND_GET_PARAM_NAME_AND_METADATA(Inputs::MidiTrackNumber), 1),
+						TInputDataVertex<int32>(METASOUND_GET_PARAM_NAME_AND_METADATA(Inputs::MidiChannelNumber), 1),
+						TInputDataVertex<FEnumStdMidiControllerID>(METASOUND_GET_PARAM_NAME_AND_METADATA(Inputs::InputMidiControllerID)),
+						TInputDataVertex<FMidiStream>(METASOUND_GET_PARAM_NAME_AND_METADATA(Inputs::MidiStream))
+					},
+					FOutputVertexInterface
+					{
+						TOutputDataVertex<int32>(METASOUND_GET_PARAM_NAME_AND_METADATA(Outputs::OutputControlChangeValueInt32)),
+						TOutputDataVertex<float>(METASOUND_GET_PARAM_NAME_AND_METADATA(Outputs::OutputControlChangeValueFloat)),
+						TOutputDataVertex<FTrigger>(METASOUND_GET_PARAM_NAME_AND_METADATA(Outputs::OutputTrigger))
+					}
+				};
+			};
+
+			static const Metasound::FVertexInterface Interface = MakeInterface();
+
+			return Interface;
+		}
+
+		static const Metasound::FNodeClassMetadata& GetNodeInfo()
+		{
+			auto InitNodeInfo = []() -> FNodeClassMetadata
+			{
+				FNodeClassMetadata Info;
+				Info.ClassName = { HarmonixNodeNamespace, TEXT("MidiCCTrigger"), TEXT("") };
+				Info.MajorVersion = 0;
+				Info.MinorVersion = 1;
+				Info.DisplayName = METASOUND_LOCTEXT("MidiCCTriggerNode_DisplayName", "Midi CC Trigger");
+				Info.Description = METASOUND_LOCTEXT("MidiCCTriggerNode_Description", "Find the Midi CC messages in a midi stream and output them as triggers, floats, ints by the specified Midi track, Midi channel, and Midi Controller ID.");
+				Info.Author = PluginAuthor;
+				Info.PromptIfMissing = PluginNodeMissingPrompt;
+				Info.DefaultInterface = GetVertexInterface();
+				Info.CategoryHierarchy.Emplace(NodeCategories::Music);
+				return Info;
+			};
+
+			static const FNodeClassMetadata Info = InitNodeInfo();
+
+			return Info;
+		}
+
+		static TUniquePtr<IOperator> CreateOperator(const Metasound::FBuildOperatorParams& InParams, Metasound::FBuildResults& OutResults)
+		{
+			const Metasound::FOperatorSettings& OperatorSettings = InParams.OperatorSettings;
+			const FInputVertexInterfaceData& InputData = InParams.InputData;
+
+			FInputs Inputs
+			{
+				InputData.GetOrCreateDefaultDataReadReference<bool>(Inputs::EnableName,OperatorSettings),
+				InputData.GetOrCreateDefaultDataReadReference<int32>(Inputs::MidiTrackNumberName, OperatorSettings),
+				InputData.GetOrCreateDefaultDataReadReference<int32>(Inputs::MidiChannelNumberName, OperatorSettings),
+				InputData.GetOrCreateDefaultDataReadReference<FEnumStdMidiControllerID>(Inputs::InputMidiControllerIDName, OperatorSettings),
+				InputData.GetOrCreateDefaultDataReadReference<HarmonixMetasound::FMidiStream>(Inputs::MidiStreamName, OperatorSettings)
+			};
+
+			FOutputs Outputs
+			{
+				Metasound::FInt32WriteRef::CreateNew(0),
+				Metasound::FFloatWriteRef::CreateNew(0.0f),
+				Metasound::FTriggerWriteRef::CreateNew(InParams.OperatorSettings)
+			};
+
+			return MakeUnique<FOp>(InParams, MoveTemp(Inputs), MoveTemp(Outputs));
+		}
+
+		FOp(const Metasound::FBuildOperatorParams& Params, FInputs&& InInputs, FOutputs&& InOutputs)
+			: Inputs(MoveTemp(InInputs))
+			, Outputs(MoveTemp(InOutputs))
+		{
+			Reset(Params);
+		}
+
+		virtual void BindInputs(FInputVertexInterfaceData& InVertexData) override
+		{
+			InVertexData.BindReadVertex(Inputs::EnableName, Inputs.Enable);
+			InVertexData.BindReadVertex(Inputs::MidiTrackNumberName, Inputs.TrackNumber);
+			InVertexData.BindReadVertex(Inputs::MidiChannelNumberName, Inputs.ChannelNumber);
+			InVertexData.BindReadVertex(Inputs::InputMidiControllerIDName, Inputs.ControllerID);
+			InVertexData.BindReadVertex(Inputs::MidiStreamName, Inputs.MidiStream);
+		}
+
+		virtual void BindOutputs(FOutputVertexInterfaceData& InVertexData) override
+		{
+			InVertexData.BindReadVertex(Outputs::OutputControlChangeValueInt32Name, Outputs.ControlChangeValueInt32);
+			InVertexData.BindReadVertex(Outputs::OutputControlChangeValueFloatName, Outputs.ControlChangeValueFloat);
+			InVertexData.BindReadVertex(Outputs::OutputTriggerName, Outputs.TriggerOut);
+		}
+
+		void Reset(const FResetParams& ResetParams)
+		{
+			Outputs.TriggerOut->Reset();
+			*Outputs.ControlChangeValueInt32 = 0;
+			*Outputs.ControlChangeValueFloat = 0.0f;
+		}
+
+		void Execute()
+		{
+			Outputs.TriggerOut->AdvanceBlock();
+
+			//reset output values to 0 upon entering a new block
+			*Outputs.ControlChangeValueInt32 = 0;
+			*Outputs.ControlChangeValueFloat = 0.0f;
+			
+			if (!*Inputs.Enable)
+			{
+				return;
+			}
+
+			const TArray<FMidiStreamEvent>& MidiEvents = Inputs.MidiStream->GetEventsInBlock();
+			for (const FMidiStreamEvent& Event : MidiEvents)
+			{
+				if (!Event.MidiMessage.IsStd())
+				{
+					continue;
+				}
+
+				if (Event.TrackIndex != *Inputs.TrackNumber || Event.MidiMessage.GetStdChannel() + 1 != *Inputs.ChannelNumber
+					|| Event.MidiMessage.GetStdStatus() != MidiConstants::kControl)
+				{
+					continue;
+				}
+
+				EStdMidiControllerID InCCID = (*Inputs.ControllerID);
+				//check if current CC message's controller ID matches with input Controller ID
+				if (Event.MidiMessage.GetStdData1() == static_cast<uint8>(InCCID))
+				{
+					//raw midi control change value (0-127)
+					*Outputs.ControlChangeValueInt32 = (int32)Event.MidiMessage.GetStdData2();
+
+					//normalize control change value(0.0-1.0)
+					*Outputs.ControlChangeValueFloat = (float)Event.MidiMessage.GetStdData2() / 127.0f;
+
+					Outputs.TriggerOut->TriggerFrame(Event.BlockSampleFrameIndex);
+				}
+				
+			}
+		}
+		private:
+			FInputs Inputs;
+			FOutputs Outputs;
+	};
+
+	class FMidiCCTriggerNode final : public Metasound::FNodeFacade
+	{
+	public:
+		explicit FMidiCCTriggerNode(const Metasound::FNodeInitData& InInitData)
+			: FNodeFacade(InInitData.InstanceName, InInitData.InstanceID, Metasound::TFacadeOperatorClass<FOp>())
+		{}
+	};
+
+	METASOUND_REGISTER_NODE(FMidiCCTriggerNode);
+}
+
+#undef LOCTEXT_NAMESPACE
\ No newline at end of file
