[Harmonix] Move Harmonix plugin into Engine/Plugins/Runtime to make it available to licensees in UE 5.4

diff --git a/Engine/Plugins/Runtime/Harmonix/Source/HarmonixMetasound/Private/HarmonixMetasound/Nodes/MidiClockTimestampTriggerNode.cpp b/Engine/Plugins/Runtime/Harmonix/Source/HarmonixMetasound/Private/HarmonixMetasound/Nodes/MidiClockTimestampTriggerNode.cpp
new file mode 100644
--- /dev/null
+++ b/Engine/Plugins/Runtime/Harmonix/Source/HarmonixMetasound/Private/HarmonixMetasound/Nodes/MidiClockTimestampTriggerNode.cpp
@@ -0,0 +1,303 @@
+// Copyright Epic Games, Inc. All Rights Reserved.
+
+#include "Tickable.h"
+#include "Containers/Queue.h"
+
+#include <atomic>
+#include <limits>
+
+#include "MetasoundExecutableOperator.h"
+#include "MetasoundFacade.h"
+#include "MetasoundNodeInterface.h"
+#include "MetasoundParamHelper.h"
+#include "MetasoundSampleCounter.h"
+#include "MetasoundStandardNodesCategories.h"
+#include "MetasoundEnumRegistrationMacro.h"
+#include "MetasoundVertex.h"
+
+#include "HarmonixMetasound/Common.h"
+#include "HarmonixMetasound/DataTypes/MidiClock.h"
+
+DEFINE_LOG_CATEGORY_STATIC(LogMidiClockTimestampTrigger, Log, All);
+
+#define LOCTEXT_NAMESPACE "HarmonixMetaSound"
+
+namespace HarmonixMetasound
+{
+	using namespace Metasound;
+
+	class FMidiClockTimestampTriggerOperator : public TExecutableOperator<FMidiClockTimestampTriggerOperator>, public FMidiPlayCursor
+	{
+	public:
+		static const FNodeClassMetadata& GetNodeInfo();
+		static const FVertexInterface& GetVertexInterface();
+		static TUniquePtr<IOperator> CreateOperator(const FBuildOperatorParams& InParams, FBuildResults& OutResults);
+
+		FMidiClockTimestampTriggerOperator(const FBuildOperatorParams& InParams,
+									const FBoolReadRef&      InEnabled,
+		                            const FMidiClockReadRef& InMidiClock,
+									const FInt32ReadRef&     InBar,
+									const FFloatReadRef&     InBeat,
+									const FBoolReadRef&      InTriggerDuringSeek);
+		virtual ~FMidiClockTimestampTriggerOperator() override;
+
+		virtual void BindInputs(FInputVertexInterfaceData& InVertexData) override;
+		virtual void BindOutputs(FOutputVertexInterfaceData& InVertexData) override;
+		virtual FDataReferenceCollection GetInputs() const override;
+		virtual FDataReferenceCollection GetOutputs() const override;
+
+		void Reset(const FResetParams& ResetParams);
+		
+		void Execute();
+
+	private:
+		//** INPUTS
+		FMidiClockReadRef   MidiClockInPin;
+		FBoolReadRef        EnableInPin;
+		FInt32ReadRef		BarInPin;
+		FFloatReadRef		BeatInPin;
+		FBoolReadRef		TriggerDuringSeekInPin;
+
+		//** OUTPUTS
+		FTriggerWriteRef   TriggerOutPin;
+
+ 		//** DATA (current state)
+		int32 Bar         = 1;
+		float Beat        = 1.0f;
+		int32 TriggerTick = 0;
+		bool  AdvanceBlockNeeded = true;
+
+		//** BEGIN FMidiPlayCursor
+		virtual void SeekToTick(int32 Tick) override { SeekThruTick(Tick - 1); }
+		virtual void SeekThruTick(int32 Tick) override;
+		virtual void AdvanceThruTick(int32 Tick, bool IsPreRoll) override;
+		// We have to override this to disambiguate the FMidiPlayCursor Reset and the MS operator Reset
+		virtual void Reset(bool ForceNoBroadcast) override { FMidiPlayCursor::Reset(ForceNoBroadcast); }
+		//** END FMidiPlayCursor
+
+		void CalculateTriggerTick();
+	};
+
+	class FMidiClockTimestampTriggerNode : public FNodeFacade
+	{
+	public:
+		FMidiClockTimestampTriggerNode(const FNodeInitData& InInitData)
+			: FNodeFacade(InInitData.InstanceName, InInitData.InstanceID, TFacadeOperatorClass<FMidiClockTimestampTriggerOperator>())
+		{}
+		virtual ~FMidiClockTimestampTriggerNode() = default;
+	};
+
+	METASOUND_REGISTER_NODE(FMidiClockTimestampTriggerNode)
+		
+	const FNodeClassMetadata& FMidiClockTimestampTriggerOperator::GetNodeInfo()
+	{
+		auto InitNodeInfo = []() -> FNodeClassMetadata
+		{
+			FNodeClassMetadata Info;
+			Info.ClassName        = { HarmonixNodeNamespace, TEXT("MidiClockTimestampTrigger"), TEXT("")};
+			Info.MajorVersion     = 0;
+			Info.MinorVersion     = 1;
+			Info.DisplayName      = METASOUND_LOCTEXT("MidiClockTimestampTriggerNode_DisplayName", "Midi Clock Timestamp Trigger");
+			Info.Description      = METASOUND_LOCTEXT("MidiClockTimestampTriggerNode_Description", "Watches a midi clock and outputs triggers at the specified musical timestamp.");
+			Info.Author           = PluginAuthor;
+			Info.PromptIfMissing  = PluginNodeMissingPrompt;
+			Info.DefaultInterface = GetVertexInterface();
+			Info.CategoryHierarchy.Emplace(NodeCategories::Music);
+			return Info;
+		};
+
+		static const FNodeClassMetadata Info = InitNodeInfo();
+
+		return Info;
+	}
+
+	namespace MidiClockTimestampTriggerPinNames
+	{
+		METASOUND_PARAM(TriggerDuringSeek, "Trigger During Seek", "Whether a trigger should be generated is a seek over the timestamp is detected.")
+		METASOUND_PARAM(TriggerOutput, "Trigger Out", "A trigger when the timestamp is detected.")
+	}
+
+	const FVertexInterface& FMidiClockTimestampTriggerOperator::GetVertexInterface()
+	{
+		using namespace MidiClockTimestampTriggerPinNames;
+		using namespace CommonPinNames;
+
+		static const FVertexInterface Interface(
+			FInputVertexInterface(
+				TInputDataVertex<FMidiClock>(METASOUND_GET_PARAM_NAME_AND_METADATA(Inputs::MidiClock)),
+				TInputDataVertex<bool>(METASOUND_GET_PARAM_NAME_AND_METADATA(Inputs::Enable), true),
+				TInputDataVertex<int32>(METASOUND_GET_PARAM_NAME_AND_METADATA(Inputs::Bar), 0),
+				TInputDataVertex<float>(METASOUND_GET_PARAM_NAME_AND_METADATA(Inputs::FloatBeat), 0.0f),
+				TInputDataVertex<bool>(METASOUND_GET_PARAM_NAME_AND_METADATA(TriggerDuringSeek), false)
+				),
+			FOutputVertexInterface(
+				TOutputDataVertex<FTrigger>(METASOUND_GET_PARAM_NAME_AND_METADATA(TriggerOutput))
+			)
+		);
+
+		return Interface;
+	}
+
+	TUniquePtr<IOperator> FMidiClockTimestampTriggerOperator::CreateOperator(const FBuildOperatorParams& InParams, FBuildResults& OutResults)
+	{
+		using namespace MidiClockTimestampTriggerPinNames;
+		using namespace CommonPinNames;
+
+		const FMidiClockTimestampTriggerNode& TheNode = static_cast<const FMidiClockTimestampTriggerNode&>(InParams.Node);
+
+		const FInputVertexInterfaceData& InputData = InParams.InputData;
+		FMidiClockReadRef InMidiClock     = InputData.GetOrConstructDataReadReference<FMidiClock>(METASOUND_GET_PARAM_NAME(Inputs::MidiClock), InParams.OperatorSettings);
+		FBoolReadRef  InEnabled           = InputData.GetOrCreateDefaultDataReadReference<bool>(METASOUND_GET_PARAM_NAME(Inputs::Enable), InParams.OperatorSettings);
+		FInt32ReadRef InBar           = InputData.GetOrCreateDefaultDataReadReference<int32>(METASOUND_GET_PARAM_NAME(Inputs::Bar), InParams.OperatorSettings);
+		FFloatReadRef InBeat              = InputData.GetOrCreateDefaultDataReadReference<float>(METASOUND_GET_PARAM_NAME(Inputs::FloatBeat), InParams.OperatorSettings);
+		FBoolReadRef  InTriggerDuringSeek = InputData.GetOrCreateDefaultDataReadReference<bool>(METASOUND_GET_PARAM_NAME(TriggerDuringSeek), InParams.OperatorSettings);
+
+		return MakeUnique<FMidiClockTimestampTriggerOperator>(InParams,
+			InEnabled,
+			InMidiClock,
+			InBar,
+			InBeat,
+			InTriggerDuringSeek);
+	}
+
+	FMidiClockTimestampTriggerOperator::FMidiClockTimestampTriggerOperator(const FBuildOperatorParams& InParams,
+															 const FBoolReadRef&      InEnabled,
+															 const FMidiClockReadRef& InMidiClock,
+															 const FInt32ReadRef&     InBar,
+															 const FFloatReadRef&     InBeat,
+															 const FBoolReadRef&      InTriggerDuringSeek)
+		: MidiClockInPin(InMidiClock)
+		, EnableInPin(InEnabled)
+		, BarInPin(InBar)
+		, BeatInPin(InBeat)
+		, TriggerDuringSeekInPin(InTriggerDuringSeek)
+		, TriggerOutPin(FTriggerWriteRef::CreateNew(InParams.OperatorSettings))
+	{
+		Reset(InParams);
+	}
+
+	FMidiClockTimestampTriggerOperator::~FMidiClockTimestampTriggerOperator()
+	{
+	}
+
+	void FMidiClockTimestampTriggerOperator::BindInputs(FInputVertexInterfaceData& InVertexData)
+	{
+		using namespace CommonPinNames;
+		using namespace MidiClockTimestampTriggerPinNames;
+
+		InVertexData.BindReadVertex(METASOUND_GET_PARAM_NAME(Inputs::Enable),    EnableInPin);
+		InVertexData.BindReadVertex(METASOUND_GET_PARAM_NAME(Inputs::MidiClock), MidiClockInPin);
+		InVertexData.BindReadVertex(METASOUND_GET_PARAM_NAME(Inputs::Bar),  BarInPin);
+		InVertexData.BindReadVertex(METASOUND_GET_PARAM_NAME(Inputs::FloatBeat), BeatInPin);
+		InVertexData.BindReadVertex(METASOUND_GET_PARAM_NAME(TriggerDuringSeek), TriggerDuringSeekInPin);
+	}
+
+	void FMidiClockTimestampTriggerOperator::BindOutputs(FOutputVertexInterfaceData& InVertexData)
+	{
+		InVertexData.BindReadVertex(METASOUND_GET_PARAM_NAME(MidiClockTimestampTriggerPinNames::TriggerOutput), TriggerOutPin);
+	}
+	
+	FDataReferenceCollection FMidiClockTimestampTriggerOperator::GetInputs() const
+	{
+		// This should never be called. Bind(...) is called instead. This method
+		// exists as a stop-gap until the API can be deprecated and removed.
+		checkNoEntry();
+		return {};
+	}
+
+	FDataReferenceCollection FMidiClockTimestampTriggerOperator::GetOutputs() const
+	{
+		// This should never be called. Bind(...) is called instead. This method
+		// exists as a stop-gap until the API can be deprecated and removed.
+		checkNoEntry();
+		return {};
+	}
+
+	void FMidiClockTimestampTriggerOperator::Reset(const FResetParams& ResetParams)
+	{
+		TriggerOutPin->Reset();
+		
+		FMidiPlayCursor::Reset(true);
+		
+		SetMessageFilter(FMidiPlayCursor::EFilterPassFlags::None);
+		MidiClockInPin->RegisterHiResPlayCursor(this);
+
+		Bar = *BarInPin;
+		Beat = *BeatInPin;
+		TriggerTick = 0;
+		AdvanceBlockNeeded = true;
+
+		CalculateTriggerTick();
+	}
+
+	void FMidiClockTimestampTriggerOperator::CalculateTriggerTick()
+	{
+		const FSongMaps& SongMaps = MidiClockInPin->GetSongMaps();
+		int32 TicksPerBeat;
+		int32 BarStartTick = SongMaps.GetBarMap().MusicTimestampBarToTick(Bar, nullptr, &TicksPerBeat);
+		int32 TicksPer32nd = MidiConstants::kTicksPerQuarterNoteInt / 8;
+		int32 BeatStartTick = (FMath::FloorToInt32((Beat - 1.0f) * TicksPerBeat) + (TicksPer32nd / 2)) / TicksPer32nd;
+		BeatStartTick *= TicksPer32nd;
+		TriggerTick = BarStartTick + BeatStartTick;
+	}
+
+	void FMidiClockTimestampTriggerOperator::Execute()
+	{
+		// first let's see if our configuration has changed at all...
+		if (*BarInPin != Bar || *BeatInPin != Beat)
+		{
+			Bar = *BarInPin;
+			Beat = *BeatInPin;
+			CalculateTriggerTick();
+		}
+
+		AdvanceBlockNeeded = true;
+	}
+
+	void FMidiClockTimestampTriggerOperator::SeekThruTick(int32 Tick)
+	{
+		if (AdvanceBlockNeeded)
+		{
+			TriggerOutPin->AdvanceBlock();
+			AdvanceBlockNeeded = false;
+		}
+
+		int32 TickProceedingThisAdvance = CurrentTick;
+		FMidiPlayCursor::SeekThruTick(Tick);
+
+		if (Tick < TickProceedingThisAdvance || !*TriggerDuringSeekInPin || !*EnableInPin)
+		{
+			return;
+		}
+
+		if (TickProceedingThisAdvance < TriggerTick && CurrentTick >= TriggerTick)
+		{
+			TriggerOutPin->TriggerFrame(MidiClockInPin->GetCurrentBlockFrameIndex());
+		}
+	}
+
+	void FMidiClockTimestampTriggerOperator::AdvanceThruTick(int32 Tick, bool IsPreRoll)
+	{
+		if (AdvanceBlockNeeded)
+		{
+			TriggerOutPin->AdvanceBlock();
+			AdvanceBlockNeeded = false;
+		}
+
+		int32 TickProceedingThisAdvance = CurrentTick;
+		FMidiPlayCursor::AdvanceThruTick(Tick, IsPreRoll);
+
+		// don't trigger during preroll...
+		if (IsPreRoll || !*EnableInPin)
+		{
+			return;
+		}
+
+		if (TickProceedingThisAdvance < TriggerTick && CurrentTick >= TriggerTick)
+		{
+			TriggerOutPin->TriggerFrame(MidiClockInPin->GetCurrentBlockFrameIndex());
+		}
+	}
+}
+
+#undef LOCTEXT_NAMESPACE // "HarmonixMetaSound"
diff --git a/Engine/Plugins/Runtime/Harmonix/Source/HarmonixMetasound/Private/HarmonixMetasound/Nodes/MidiNoteTriggerNode.cpp b/Engine/Plugins/Runtime/Harmonix/Source/HarmonixMetasound/Private/HarmonixMetasound/Nodes/MidiNoteTriggerNode.cpp
new file mode 100644
--- /dev/null
+++ b/Engine/Plugins/Runtime/Harmonix/Source/HarmonixMetasound/Private/HarmonixMetasound/Nodes/MidiNoteTriggerNode.cpp
@@ -0,0 +1,350 @@
+// Copyright Epic Games, Inc. All Rights Reserved.
+
+#include "MetasoundExecutableOperator.h"
+#include "MetasoundFacade.h"
+#include "MetasoundNodeInterface.h"
+#include "MetasoundParamHelper.h"
+#include "MetasoundSampleCounter.h"
+#include "MetasoundStandardNodesCategories.h"
+#include "MetasoundVertex.h"
+#include "MetasoundTrigger.h"
+
+#include "HarmonixMetasound/Common.h"
+#include "HarmonixMetasound/DataTypes/MidiStream.h"
+#include "HarmonixMetasound/DataTypes/MusicTransport.h"
+
+DEFINE_LOG_CATEGORY_STATIC(LogMidiNoteTrigger, Log, All);
+
+#define LOCTEXT_NAMESPACE "HarmonixMetaSound"
+
+namespace HarmonixMetasound
+{
+	using namespace Metasound;
+
+	class FMidiNoteTriggerOperator : public TExecutableOperator<FMidiNoteTriggerOperator>
+	{
+	public:
+		static const FNodeClassMetadata& GetNodeInfo();
+		static const FVertexInterface& GetVertexInterface();
+		static TUniquePtr<IOperator> CreateOperator(const FBuildOperatorParams& InParams, FBuildResults& OutResults);
+
+		FMidiNoteTriggerOperator(const FBuildOperatorParams& InParams,
+								 const FBoolReadRef&       InEnabled,
+								 const FMidiStreamReadRef& InMidiStream,
+								 const FInt32ReadRef&      InTrackNumber,
+								 const FInt32ReadRef&      InChannelNumber,
+								 const FInt32ReadRef&      InMinNoteNumber,
+								 const FInt32ReadRef&      InMaxNoteNumber,
+								 const FInt32ReadRef&      InMinVelocity,
+								 const FInt32ReadRef&      InMaxVelocity,
+								 const FBoolReadRef&       InMidiClockAffectsFrequency);
+
+		virtual void BindInputs(FInputVertexInterfaceData& InVertexData) override;
+		virtual void BindOutputs(FOutputVertexInterfaceData& InVertexData) override;
+		virtual FDataReferenceCollection GetInputs() const override;
+		virtual FDataReferenceCollection GetOutputs() const override;
+
+		void Reset(const FResetParams& ResetParams);
+		
+		void Execute();
+
+	private:
+		//** INPUTS
+		FBoolReadRef       EnableInPin;
+		FMidiStreamReadRef MidiStreamInPin;
+		FInt32ReadRef      TrackNumberInPin;
+		FInt32ReadRef      ChannelNumInPin;
+		FInt32ReadRef      MinNoteInPin;
+		FInt32ReadRef      MaxNoteInPin;
+		FInt32ReadRef      MinVelInPin;
+		FInt32ReadRef      MaxVelInPin;
+		FBoolReadRef       ClockAffectsFrequencyInPin;
+
+		//** OUTPUTS
+		FTriggerWriteRef NoteOnOutPin;
+		FTriggerWriteRef NoteOffOutPin;
+		FInt32WriteRef   NoteNumOutPin;
+		FFloatWriteRef   FreqOutPin;
+		FInt32WriteRef   VelOutPin;
+		FFloatWriteRef   NormVelOutPin;
+
+		//** DATA
+		int8 SoundingNote = -1;
+		FMidiVoiceId PlayingId;
+	};
+
+	class FMidiNoteTriggerNode : public FNodeFacade
+	{
+	public:
+		FMidiNoteTriggerNode(const FNodeInitData& InInitData)
+			: FNodeFacade(InInitData.InstanceName, InInitData.InstanceID, TFacadeOperatorClass<FMidiNoteTriggerOperator>())
+		{}
+		virtual ~FMidiNoteTriggerNode() = default;
+	};
+
+	METASOUND_REGISTER_NODE(FMidiNoteTriggerNode)
+
+	const FNodeClassMetadata& FMidiNoteTriggerOperator::GetNodeInfo()
+	{
+		auto InitNodeInfo = []() -> FNodeClassMetadata
+		{
+			FNodeClassMetadata Info;
+			Info.ClassName        = { HarmonixNodeNamespace, TEXT("MidiNoteTrigger"), TEXT("")};
+			Info.MajorVersion     = 0;
+			Info.MinorVersion     = 1;
+			Info.DisplayName      = METASOUND_LOCTEXT("MidiNoteTriggerNode_DisplayName", "Midi Note Trigger");
+			Info.Description      = METASOUND_LOCTEXT("MidiNoteTriggerNode_Description", "Receives a midi stream, filters for the desired messages, and outputs triggers.");
+			Info.Author           = PluginAuthor;
+			Info.PromptIfMissing  = PluginNodeMissingPrompt;
+			Info.DefaultInterface = GetVertexInterface();
+			Info.CategoryHierarchy.Emplace(NodeCategories::Music);
+			return Info;
+		};
+
+		static const FNodeClassMetadata Info = InitNodeInfo();
+
+		return Info;
+	}
+
+	const FVertexInterface& FMidiNoteTriggerOperator::GetVertexInterface()
+	{
+		using namespace CommonPinNames;
+
+		static const FVertexInterface Interface(
+			FInputVertexInterface(
+				TInputDataVertex<bool>(METASOUND_GET_PARAM_NAME_AND_METADATA(Inputs::Enable), true),
+				TInputDataVertex<FMidiStream>(METASOUND_GET_PARAM_NAME_AND_METADATA(Inputs::MidiStream)),
+				TInputDataVertex<int32>(METASOUND_GET_PARAM_NAME_AND_METADATA(Inputs::MidiTrackNumber),1),
+				TInputDataVertex<int32>(METASOUND_GET_PARAM_NAME_AND_METADATA(Inputs::MidiChannelNumber), 1),
+				TInputDataVertex<int32>(METASOUND_GET_PARAM_NAME_AND_METADATA(Inputs::MinMidiNote), 0),
+				TInputDataVertex<int32>(METASOUND_GET_PARAM_NAME_AND_METADATA(Inputs::MaxMidiNote), 127),
+				TInputDataVertex<int32>(METASOUND_GET_PARAM_NAME_AND_METADATA(Inputs::MinMidiVelocity), 0),
+				TInputDataVertex<int32>(METASOUND_GET_PARAM_NAME_AND_METADATA(Inputs::MaxMidiVelocity), 127),
+				TInputDataVertex<bool>(METASOUND_GET_PARAM_NAME_AND_METADATA(Inputs::ClockSpeedToFrequency), false)
+				),
+			FOutputVertexInterface(
+				TOutputDataVertex<FTrigger>(METASOUND_GET_PARAM_NAME_AND_METADATA(Outputs::NoteOn)),
+				TOutputDataVertex<FTrigger>(METASOUND_GET_PARAM_NAME_AND_METADATA(Outputs::NoteOff)),
+				TOutputDataVertex<int32>(METASOUND_GET_PARAM_NAME_AND_METADATA(Outputs::MidiNoteNumber)),
+				TOutputDataVertex<float>(METASOUND_GET_PARAM_NAME_AND_METADATA(Outputs::Frequency)),
+				TOutputDataVertex<int32>(METASOUND_GET_PARAM_NAME_AND_METADATA(Outputs::MidiVelocity)),
+				TOutputDataVertex<float>(METASOUND_GET_PARAM_NAME_AND_METADATA(Outputs::NormalizedVelocity))
+				)
+		);
+
+		return Interface;
+	}
+
+	TUniquePtr<IOperator> FMidiNoteTriggerOperator::CreateOperator(const FBuildOperatorParams& InParams, FBuildResults& OutResults)
+	{
+		using namespace CommonPinNames;
+
+		const FMidiNoteTriggerNode& LoggerNode = static_cast<const FMidiNoteTriggerNode&>(InParams.Node);
+
+		const FInputVertexInterfaceData& InputData = InParams.InputData;
+		FBoolReadRef InEnabled          = InputData.GetOrCreateDefaultDataReadReference<bool>(METASOUND_GET_PARAM_NAME(Inputs::Enable), InParams.OperatorSettings);
+		FMidiStreamReadRef InMidiStream = InputData.GetOrConstructDataReadReference<FMidiStream>(METASOUND_GET_PARAM_NAME(Inputs::MidiStream), InParams.OperatorSettings);
+		FInt32ReadRef InTrackNumber     = InputData.GetOrCreateDefaultDataReadReference<int32>(METASOUND_GET_PARAM_NAME(Inputs::MidiTrackNumber), InParams.OperatorSettings);
+		FInt32ReadRef InMidiChannel     = InputData.GetOrCreateDefaultDataReadReference<int32>(METASOUND_GET_PARAM_NAME(Inputs::MidiChannelNumber), InParams.OperatorSettings);
+		FInt32ReadRef InMinNote         = InputData.GetOrCreateDefaultDataReadReference<int32>(METASOUND_GET_PARAM_NAME(Inputs::MinMidiNote), InParams.OperatorSettings);
+		FInt32ReadRef InMaxNote         = InputData.GetOrCreateDefaultDataReadReference<int32>(METASOUND_GET_PARAM_NAME(Inputs::MaxMidiNote), InParams.OperatorSettings);
+		FInt32ReadRef InMinVelocity     = InputData.GetOrCreateDefaultDataReadReference<int32>(METASOUND_GET_PARAM_NAME(Inputs::MinMidiVelocity), InParams.OperatorSettings);
+		FInt32ReadRef InMaxVelocity     = InputData.GetOrCreateDefaultDataReadReference<int32>(METASOUND_GET_PARAM_NAME(Inputs::MaxMidiVelocity), InParams.OperatorSettings);
+		FBoolReadRef  InClockToFreq     = InputData.GetOrCreateDefaultDataReadReference<bool>(METASOUND_GET_PARAM_NAME(Inputs::ClockSpeedToFrequency), InParams.OperatorSettings);
+
+		return MakeUnique<FMidiNoteTriggerOperator>(InParams, InEnabled, InMidiStream, InTrackNumber, InMidiChannel, InMinNote, InMaxNote, InMinVelocity, InMaxVelocity, InClockToFreq);
+	}
+
+	FMidiNoteTriggerOperator::FMidiNoteTriggerOperator(const FBuildOperatorParams& InParams,
+													   const FBoolReadRef&       InEnabled,
+													   const FMidiStreamReadRef& InMidiStream,
+													   const FInt32ReadRef&      InTrackNumber,
+													   const FInt32ReadRef&      InChannelNumber,
+													   const FInt32ReadRef&      InMinNoteNumber,
+													   const FInt32ReadRef&      InMaxNoteNumber,
+													   const FInt32ReadRef&      InMinVelocity,
+													   const FInt32ReadRef&      InMaxVelocity,
+													   const FBoolReadRef&       InClockToFreq)
+		: EnableInPin(InEnabled)
+		, MidiStreamInPin(InMidiStream)
+		, TrackNumberInPin(InTrackNumber)
+		, ChannelNumInPin(InChannelNumber)
+		, MinNoteInPin(InMinNoteNumber)
+		, MaxNoteInPin(InMaxNoteNumber)
+		, MinVelInPin(InMinVelocity)
+		, MaxVelInPin(InMaxVelocity)
+		, ClockAffectsFrequencyInPin(InClockToFreq)
+		, NoteOnOutPin(FTriggerWriteRef::CreateNew(InParams.OperatorSettings))
+		, NoteOffOutPin(FTriggerWriteRef::CreateNew(InParams.OperatorSettings))
+		, NoteNumOutPin(FInt32WriteRef::CreateNew(0))
+		, FreqOutPin(FFloatWriteRef::CreateNew(0.0f))
+		, VelOutPin(FInt32WriteRef::CreateNew(0))
+		, NormVelOutPin(FFloatWriteRef::CreateNew(0.0f))
+	{
+		Reset(InParams);
+	}
+
+	void FMidiNoteTriggerOperator::BindInputs(FInputVertexInterfaceData& InVertexData)
+	{
+		using namespace CommonPinNames;
+
+		InVertexData.BindReadVertex(METASOUND_GET_PARAM_NAME(Inputs::Enable), EnableInPin);
+		InVertexData.BindReadVertex(METASOUND_GET_PARAM_NAME(Inputs::MidiStream),   MidiStreamInPin);
+		InVertexData.BindReadVertex(METASOUND_GET_PARAM_NAME(Inputs::MidiTrackNumber),  TrackNumberInPin);
+		InVertexData.BindReadVertex(METASOUND_GET_PARAM_NAME(Inputs::MidiChannelNumber), ChannelNumInPin);
+		InVertexData.BindReadVertex(METASOUND_GET_PARAM_NAME(Inputs::MinMidiNote), MinNoteInPin);
+		InVertexData.BindReadVertex(METASOUND_GET_PARAM_NAME(Inputs::MaxMidiNote), MaxNoteInPin);
+		InVertexData.BindReadVertex(METASOUND_GET_PARAM_NAME(Inputs::MinMidiVelocity), MinVelInPin);
+		InVertexData.BindReadVertex(METASOUND_GET_PARAM_NAME(Inputs::MaxMidiVelocity), MaxVelInPin);
+		InVertexData.BindReadVertex(METASOUND_GET_PARAM_NAME(Inputs::ClockSpeedToFrequency), ClockAffectsFrequencyInPin);
+	}
+
+	void FMidiNoteTriggerOperator::BindOutputs(FOutputVertexInterfaceData& InVertexData)
+	{
+		using namespace CommonPinNames;
+
+		InVertexData.BindReadVertex(METASOUND_GET_PARAM_NAME(Outputs::NoteOn), NoteOnOutPin);
+		InVertexData.BindReadVertex(METASOUND_GET_PARAM_NAME(Outputs::NoteOff), NoteOffOutPin);
+		InVertexData.BindReadVertex(METASOUND_GET_PARAM_NAME(Outputs::MidiNoteNumber), NoteNumOutPin);
+		InVertexData.BindReadVertex(METASOUND_GET_PARAM_NAME(Outputs::Frequency), FreqOutPin);
+		InVertexData.BindReadVertex(METASOUND_GET_PARAM_NAME(Outputs::MidiVelocity), VelOutPin);
+		InVertexData.BindReadVertex(METASOUND_GET_PARAM_NAME(Outputs::NormalizedVelocity), NormVelOutPin);
+	}
+	
+	FDataReferenceCollection FMidiNoteTriggerOperator::GetInputs() const
+	{
+		// This should never be called. Bind(...) is called instead. This method
+		// exists as a stop-gap until the API can be deprecated and removed.
+		checkNoEntry();
+		return {};
+	}
+
+	FDataReferenceCollection FMidiNoteTriggerOperator::GetOutputs() const
+	{
+		// This should never be called. Bind(...) is called instead. This method
+		// exists as a stop-gap until the API can be deprecated and removed.
+		checkNoEntry();
+		return {};
+	}
+
+	void FMidiNoteTriggerOperator::Reset(const FResetParams& ResetParams)
+	{
+		NoteOnOutPin->Reset();
+		NoteOffOutPin->Reset();
+		*NoteNumOutPin = 0;
+		*FreqOutPin = 0;
+		*VelOutPin = 0;
+		*NormVelOutPin = 0;
+		SoundingNote = -1;
+		PlayingId = FMidiVoiceId::None();
+	}
+
+	void FMidiNoteTriggerOperator::Execute()
+	{
+		NoteOnOutPin->AdvanceBlock();
+		NoteOffOutPin->AdvanceBlock();
+
+		if (!*EnableInPin)
+		{
+			if (SoundingNote >= 0)
+			{
+				*VelOutPin     = 0;
+				*NormVelOutPin = 0.0f;
+				NoteOffOutPin->TriggerFrame(0);
+				SoundingNote   = -1;
+			}
+			return;
+		}
+
+		auto TransportChanges = MidiStreamInPin->GetTransportChangesInBlock();
+		if (!TransportChanges.IsEmpty())
+		{
+			if (TransportChanges.Last().TransportState != EMusicPlayerTransportState::Playing)
+			{
+				if (SoundingNote >= 0)
+				{
+					*VelOutPin     = 0;
+					*NormVelOutPin = 0.0f;
+					NoteOffOutPin->TriggerFrame(0);
+					SoundingNote   = -1;
+				}
+				return;
+			}
+		}
+
+		int32 NoteOffTriggerFrame = -1;
+		const TArray<FMidiStreamEvent>& MidiEvents = MidiStreamInPin->GetEventsInBlock();
+		for (auto& Event : MidiEvents)
+		{
+			int32 trackPin   = *TrackNumberInPin;
+			int32 channelPin = *ChannelNumInPin;
+			bool IsNoteOn    = Event.MidiMessage.IsNoteOn();
+			bool IsNoteOff   = Event.MidiMessage.IsNoteOff();
+		
+			if (Event.TrackIndex != *TrackNumberInPin ||
+				(!IsNoteOn && !IsNoteOff) ||
+				Event.MidiMessage.GetStdChannel() + 1 != *ChannelNumInPin)
+			{
+				continue;
+			}
+
+			if (Event.MidiMessage.GetStdData1() < *MinNoteInPin ||
+				Event.MidiMessage.GetStdData1() > *MaxNoteInPin)
+			{
+				continue;
+			}
+
+			if (Event.MidiMessage.IsNoteOn()) 
+			{
+				if (Event.MidiMessage.GetStdData2() < *MinVelInPin ||
+					Event.MidiMessage.GetStdData2() > *MaxVelInPin)
+				{
+					continue;
+				}
+
+				if (SoundingNote > -1)
+				{
+					// stopping sounding note!
+					NoteOffOutPin->TriggerFrame(Event.BlockSampleFrameIndex);
+					NoteOffTriggerFrame = Event.BlockSampleFrameIndex;
+					SoundingNote = -1;
+				}
+				PlayingId      = Event.GetVoiceId();
+				*VelOutPin     = Event.MidiMessage.GetStdData2();
+				*NormVelOutPin = (float)Event.MidiMessage.GetStdData2() / 127.0f;
+				*NoteNumOutPin = Event.MidiMessage.GetStdData1();
+				NoteOnOutPin->TriggerFrame(NoteOffTriggerFrame == Event.BlockSampleFrameIndex ? NoteOffTriggerFrame + 1 : Event.BlockSampleFrameIndex);
+				SoundingNote   = Event.MidiMessage.GetStdData1();;
+			}
+			else
+			{
+				if (Event.GetVoiceId() == PlayingId)
+				{
+					*VelOutPin = 0;
+					*NormVelOutPin = 0.0f;
+					*NoteNumOutPin = Event.MidiMessage.GetStdData1();
+					NoteOffOutPin->TriggerFrame(Event.BlockSampleFrameIndex);
+					NoteOffTriggerFrame = Event.BlockSampleFrameIndex;
+					SoundingNote = -1;
+					PlayingId = FMidiVoiceId::None();
+				}
+			}
+		}
+		
+		// We do this here so that even during the playback of a note we can
+		// adjust the pitch if there is a clock source and it has a changing speed. 
+		if (SoundingNote != -1)
+		{
+			float Frequency = powf(2.0f, (float)(SoundingNote - 69) / 12.0f) * 440.0f;
+			const FMidiClockReadRef* ClockSource = MidiStreamInPin->GetMidiClockSource();
+			if (ClockSource && *ClockAffectsFrequencyInPin)
+			{
+				Frequency *= (*ClockSource)->GetSpeedAtEndOfBlock();
+			}
+			*FreqOutPin = Frequency;
+		}
+	}
+}
+
+#undef LOCTEXT_NAMESPACE // "HarmonixMetaSound"
+
diff --git a/Engine/Plugins/Runtime/Harmonix/Source/HarmonixMetasound/Private/HarmonixMetasound/Nodes/MidiPlayerNode.cpp b/Engine/Plugins/Runtime/Harmonix/Source/HarmonixMetasound/Private/HarmonixMetasound/Nodes/MidiPlayerNode.cpp
new file mode 100644
--- /dev/null
+++ b/Engine/Plugins/Runtime/Harmonix/Source/HarmonixMetasound/Private/HarmonixMetasound/Nodes/MidiPlayerNode.cpp
@@ -0,0 +1,856 @@
+// Copyright Epic Games, Inc. All Rights Reserved.
+
+#include "MetasoundExecutableOperator.h"
+#include "MetasoundFacade.h"
+#include "MetasoundNodeInterface.h"
+#include "MetasoundParamHelper.h"
+#include "MetasoundSampleCounter.h"
+#include "MetasoundStandardNodesCategories.h"
+#include "MetasoundVertex.h"
+
+#include "HarmonixMetasound/Common.h"
+#include "HarmonixMetasound/DataTypes/MidiStream.h"
+#include "HarmonixMetasound/DataTypes/MidiClock.h"
+#include "HarmonixMetasound/DataTypes/MidiClockEvent.h"
+#include "HarmonixMetasound/DataTypes/MidiAsset.h"
+#include "HarmonixMetasound/DataTypes/MusicTransport.h"
+
+#include "HarmonixMidi/MidiPlayCursorMgr.h"
+#include "HarmonixMidi/MidiPlayCursor.h"
+#include "HarmonixMidi/MidiVoiceId.h"
+
+#define LOCTEXT_NAMESPACE "HarmonixMetaSound"
+
+namespace HarmonixMetasound
+{
+	using namespace Metasound;
+
+	namespace MidiPlayerNodePinNames
+	{
+		METASOUND_PARAM(KillVoicesOnSeek, "Kill Voices On Seek", "If true, a \"Kill All Voices\" midi message will be sent when seeking. Otherwise an \"All Notes Off\" will be sent which allows to ADSR release phases.")
+		METASOUND_PARAM(KillVoicesOnMidiChange, "Kill Voices On Midi File Change", "If true, a \"Kill All Voices\" midi message will be sent when the midi file asset is changed. Otherwise an \"All Notes Off\" will be sent which allows to ADSR release phases.")
+	}
+
+	class FMidiPlayerOperator : public TExecutableOperator<FMidiPlayerOperator>, public FMidiPlayCursor, public FMidiVoiceGeneratorBase, public FMusicTransportControllable
+	{
+	public:
+		static const FNodeClassMetadata& GetNodeInfo();
+		static const FVertexInterface&   GetVertexInterface();
+		static TUniquePtr<IOperator>     CreateOperator(const FBuildOperatorParams& InParams, FBuildResults& OutResults);
+
+		FMidiPlayerOperator(const FOperatorSettings& InSettings, 
+							const FMidiAssetReadRef& InMidiAsset,
+							const FMusicTransportEventStreamReadRef& InTransport,
+							const FBoolReadRef& InLoop,
+							const FFloatReadRef& InSpeedMultiplier,
+                            const int32 InPrerollBars,
+							const bool bInKillVoicesOnSeek,
+							const bool bInKillVoicesOnMidiChange);
+
+		virtual void BindInputs(FInputVertexInterfaceData& InVertexData) override;
+		virtual void BindOutputs(FOutputVertexInterfaceData& InVertexData) override;
+		
+		virtual void Reset(const FResetParams& Params);
+		virtual void Execute();
+		bool IsPlaying() const;
+
+		//~ BEGIN FMidiPlayCursor Overrides
+		virtual void Reset(bool ForceNoBroadcast = false) override;
+		virtual void OnLoop(int32 LoopStartTick, int32 LoopEndTick) override;
+		virtual void SeekToTick(int32 Tick) override;
+		virtual void SeekThruTick(int32 Tick) override;
+		virtual void AdvanceThruTick(int32 Tick, bool IsPreRoll) override;
+
+		virtual void OnMidiMessage(int32 TrackIndex, int32 Tick, uint8 Status, uint8 Data1, uint8 Data2, bool IsPreroll) override;
+		virtual void OnTempo(int32 TrackIndex, int32 Tick, int32 Tempo, bool IsPreroll) override;
+		virtual void OnText(int32 TrackIndex, int32 Tick, int32 TextIndex, const FString& Str, uint8 Type, bool IsPreroll = false) override;
+		virtual void OnPreRollNoteOn(int32 TrackIndex, int32 EventTick, int32 CurrentTick, float PreRollMs, uint8 Status, uint8 Data1, uint8 Data2) override;
+		//~ END FMidiPlayCursor Overrides
+
+	protected:
+		//** INPUTS
+		FMidiAssetReadRef MidiAssetInPin;
+		FMusicTransportEventStreamReadRef TransportInPin;
+		FBoolReadRef  LoopInPin;
+		FFloatReadRef SpeedMultInPin;
+		int32 PrerollBars;
+		bool bKillVoicesOnSeek;
+		bool bKillVoicesOnMidiChange;
+
+		//** OUTPUTS
+		FMidiStreamWriteRef MidiOutPin;
+		FMidiClockWriteRef  MidiClockOut;
+
+		//** DATA
+		FMidiClockEventCursor MidiClockEventCursor;
+		FMidiFileProxyPtr CurrentMidiFile;
+		FSampleCount BlockSize      = 0;
+		int32 CurrentBlockSpanStart = 0;
+		
+		virtual void SetupNewMidiFile(const FMidiFileProxyPtr& NewMidi);
+	};
+
+	class FExternallyClockedMidiPlayerOperator : public FMidiPlayerOperator
+	{
+	public:
+		FExternallyClockedMidiPlayerOperator(const FOperatorSettings& InSettings,
+		                                     const FMidiAssetReadRef& InMidiAsset,
+		                                     const FMusicTransportEventStreamReadRef& InTransport,
+		                                     const FMidiClockReadRef& InMidiClock,
+											 const FBoolReadRef& InLoop,
+		                                     const FFloatReadRef& InSpeedMultiplier,
+			                                 const int32 InPrerollBars,
+											 const bool bInKillVoicesOnSeek,
+											 const bool bInKillVoicesOnMidiChange);
+
+		virtual void BindInputs(FInputVertexInterfaceData& InVertexData) override;
+		virtual void BindOutputs(FOutputVertexInterfaceData& InVertexData) override;
+		virtual FDataReferenceCollection GetInputs() const override;
+		virtual FDataReferenceCollection GetOutputs() const override;
+
+		virtual void Reset(const FResetParams& Params) override;
+
+		virtual void Execute() override;
+
+	private:
+		//** INPUTS **********************************
+		FMidiClockReadRef MidiClockIn;
+
+		//** DATA   **********************************
+		int32 LoopOffsetTick = 0;
+
+		void UpdateLoopOffsetTickFromTick(int32 Tick);
+	};
+
+	class FSelfClockedMidiPlayerOperator : public FMidiPlayerOperator
+	{
+	public:
+		FSelfClockedMidiPlayerOperator(const FOperatorSettings& InSettings,
+		                               const FMidiAssetReadRef& InMidiAsset,
+		                               const FMusicTransportEventStreamReadRef& InTransport,
+									   const FBoolReadRef& InLoop,
+		                               const FFloatReadRef& InSpeedMultiplier,
+			                           const int32 InPrerollBars,
+									   const bool bInKillVoicesOnSeek,
+									   const bool bInKillVoicesOnMidiChange );
+
+		virtual void Reset(const FResetParams& Params) override;
+
+		virtual void Execute() override;
+
+		virtual void BindInputs(FInputVertexInterfaceData& InVertexData) override;
+		virtual void BindOutputs(FOutputVertexInterfaceData& InVertexData) override;
+		virtual FDataReferenceCollection GetInputs() const override;
+		virtual FDataReferenceCollection GetOutputs() const override;
+	};
+
+	class FMidiPlayerNode : public FNodeFacade
+	{
+	public:
+		FMidiPlayerNode(const FNodeInitData& InInitData)
+			: FNodeFacade(InInitData.InstanceName, InInitData.InstanceID, TFacadeOperatorClass<FMidiPlayerOperator>())
+		{}
+		virtual ~FMidiPlayerNode() = default;
+	};
+
+	METASOUND_REGISTER_NODE(FMidiPlayerNode)
+
+	const FNodeClassMetadata& FMidiPlayerOperator::GetNodeInfo()
+	{
+		auto InitNodeInfo = []() -> FNodeClassMetadata
+		{
+			FNodeClassMetadata Info;
+			Info.ClassName        = { HarmonixNodeNamespace, TEXT("MidiPlayer"), TEXT("") };
+			Info.MajorVersion     = 0;
+			Info.MinorVersion     = 1;
+			Info.DisplayName      = METASOUND_LOCTEXT("MidiPlayerNode_DisplayName", "Midi Player");
+			Info.Description      = METASOUND_LOCTEXT("MidiPlayerNode_Description", "Plays a standard midi file.");
+			Info.Author           = PluginAuthor;
+			Info.PromptIfMissing  = PluginNodeMissingPrompt;
+			Info.DefaultInterface = GetVertexInterface();
+			Info.CategoryHierarchy.Emplace(NodeCategories::Music);
+
+			return Info;
+		};
+
+		static const FNodeClassMetadata Info = InitNodeInfo();
+
+		return Info;
+	}
+
+	const FVertexInterface& FMidiPlayerOperator::GetVertexInterface()
+	{
+		using namespace CommonPinNames;
+
+		static const FVertexInterface Interface(
+			FInputVertexInterface(
+				TInputDataVertex<FMidiAsset>(METASOUND_GET_PARAM_NAME_AND_METADATA(Inputs::MidiFileAsset)),
+				TInputDataVertex<FMusicTransportEventStream>(METASOUND_GET_PARAM_NAME_AND_METADATA(Inputs::Transport)),
+				TInputDataVertex<FMidiClock>(METASOUND_GET_PARAM_NAME_AND_METADATA(Inputs::MidiClock)),
+				TInputDataVertex<bool>(METASOUND_GET_PARAM_NAME_AND_METADATA(Inputs::Loop), false),
+				TInputDataVertex<float>(METASOUND_GET_PARAM_NAME_AND_METADATA(Inputs::Speed), 1.0f),
+				TInputConstructorVertex<int32>(METASOUND_GET_PARAM_NAME_AND_METADATA(Inputs::PrerollBars), 8),
+				TInputConstructorVertex<bool>(METASOUND_GET_PARAM_NAME_AND_METADATA(MidiPlayerNodePinNames::KillVoicesOnSeek), false),
+				TInputConstructorVertex<bool>(METASOUND_GET_PARAM_NAME_AND_METADATA(MidiPlayerNodePinNames::KillVoicesOnMidiChange), false)
+			),
+			FOutputVertexInterface(
+				TOutputDataVertex<FMidiStream>(METASOUND_GET_PARAM_NAME_AND_METADATA(Outputs::MidiStream)),
+				TOutputDataVertex<FMidiClock>(METASOUND_GET_PARAM_NAME_AND_METADATA(Outputs::MidiClock))
+			)
+		);
+
+		return Interface;
+	}
+
+	TUniquePtr<IOperator> FMidiPlayerOperator::CreateOperator(const FBuildOperatorParams& InParams, FBuildResults& OutResults)
+	{
+		using namespace CommonPinNames;
+
+		const FMidiPlayerNode& MidiPlayerNode = static_cast<const FMidiPlayerNode&>(InParams.Node);
+
+		const FInputVertexInterfaceData& InputData = InParams.InputData;
+
+		const FOperatorSettings& Settings = InParams.OperatorSettings;
+
+		FMidiAssetReadRef InMidiAsset = InputData.GetOrConstructDataReadReference<FMidiAsset>(METASOUND_GET_PARAM_NAME(Inputs::MidiFileAsset));
+		FMusicTransportEventStreamReadRef InTransport  = InputData.GetOrConstructDataReadReference<FMusicTransportEventStream>(METASOUND_GET_PARAM_NAME(Inputs::Transport), Settings);
+		FBoolReadRef InLoop = InputData.GetOrConstructDataReadReference<bool>(METASOUND_GET_PARAM_NAME(Inputs::Loop), false);
+		FFloatReadRef InSpeed = InputData.GetOrConstructDataReadReference<float>(METASOUND_GET_PARAM_NAME(Inputs::Speed), 1.0f);
+		int32 InPrerollBars = InputData.GetOrCreateDefaultValue<int32>(METASOUND_GET_PARAM_NAME(Inputs::PrerollBars), Settings);
+		bool bInKillVoicesOnSeek = InputData.GetOrCreateDefaultValue<bool>(METASOUND_GET_PARAM_NAME(MidiPlayerNodePinNames::KillVoicesOnSeek), Settings);
+		bool bInKillVoicesOnMidiChange = InputData.GetOrCreateDefaultValue<bool>(METASOUND_GET_PARAM_NAME(MidiPlayerNodePinNames::KillVoicesOnMidiChange), Settings);
+		if (InputData.IsVertexBound(METASOUND_GET_PARAM_NAME(Inputs::MidiClock)))
+		{
+			FMidiClockReadRef InMidiClock = InputData.GetOrConstructDataReadReference<FMidiClock>(METASOUND_GET_PARAM_NAME(Inputs::MidiClock), Settings);
+			return MakeUnique<FExternallyClockedMidiPlayerOperator>(InParams.OperatorSettings, InMidiAsset, InTransport, InMidiClock, InLoop, InSpeed, InPrerollBars, bInKillVoicesOnSeek, bInKillVoicesOnMidiChange);
+		}
+
+		return MakeUnique<FSelfClockedMidiPlayerOperator>(InParams.OperatorSettings, InMidiAsset, InTransport, InLoop, InSpeed, InPrerollBars, bInKillVoicesOnSeek, bInKillVoicesOnMidiChange);
+	}
+
+	FMidiPlayerOperator::FMidiPlayerOperator(const FOperatorSettings& InSettings, 
+											 const FMidiAssetReadRef& InMidiAsset,
+											 const FMusicTransportEventStreamReadRef& InTransport,
+											 const FBoolReadRef& InLoop,
+											 const FFloatReadRef& InSpeed,
+		                                     const int32 InPrerollBars,
+											 const bool bInKillVoicesOnSeek,
+											 const bool bInKillVoicesOnMidiChange)
+		: FMusicTransportControllable(EMusicPlayerTransportState::Prepared)
+		, MidiAssetInPin(InMidiAsset)
+		, TransportInPin(InTransport)
+		, LoopInPin(InLoop)
+		, SpeedMultInPin(InSpeed)
+		, PrerollBars(InPrerollBars)
+		, bKillVoicesOnSeek(bInKillVoicesOnSeek)
+		, bKillVoicesOnMidiChange(bInKillVoicesOnMidiChange)
+		, MidiOutPin(FMidiStreamWriteRef::CreateNew(InSettings))
+		, MidiClockOut(FMidiClockWriteRef::CreateNew(InSettings))
+		, MidiClockEventCursor(MidiClockOut)
+		, BlockSize(InSettings.GetNumFramesPerBlock())
+	{
+		MidiClockOut->RegisterHiResPlayCursor(this);
+		MidiOutPin->SetClockSource(MidiClockOut);
+		MidiOutPin->AddTransportStateChangeMessage(0, EMusicPlayerTransportState::Prepared);
+	}
+
+	FExternallyClockedMidiPlayerOperator::FExternallyClockedMidiPlayerOperator(const FOperatorSettings& InSettings,
+																			   const FMidiAssetReadRef& InMidiAsset, 
+																			   const FMusicTransportEventStreamReadRef& InTransport,
+																			   const FMidiClockReadRef& InMidiClock, 
+																			   const FBoolReadRef& InLoop,
+																			   const FFloatReadRef& InSpeedMultiplier,
+		                                                                       const int32 InPrerollBars,
+																			   const bool bInKillVoicesOnSeek,
+																			   const bool bInKillVoicesOnMidiChange)
+		: FMidiPlayerOperator(InSettings, InMidiAsset, InTransport, InLoop, InSpeedMultiplier, InPrerollBars, bInKillVoicesOnSeek, bInKillVoicesOnMidiChange)
+		, MidiClockIn(InMidiClock)
+	{
+	}
+
+	FSelfClockedMidiPlayerOperator::FSelfClockedMidiPlayerOperator(const FOperatorSettings& InSettings, 
+																   const FMidiAssetReadRef& InMidiAsset, 
+																   const FMusicTransportEventStreamReadRef& InTransport,
+																   const FBoolReadRef& InLoop,
+																   const FFloatReadRef& InSpeedMultiplier,
+		                                                           int32 InPrerollBars,
+																   const bool bInKillVoicesOnSeek,
+																   const bool bInKillVoicesOnMidiChange)
+		: FMidiPlayerOperator(InSettings, InMidiAsset, InTransport, InLoop, InSpeedMultiplier, InPrerollBars, bInKillVoicesOnSeek, bInKillVoicesOnMidiChange)
+	{}
+
+	void FSelfClockedMidiPlayerOperator::Reset(const FResetParams& Params)
+	{
+		FMidiPlayerOperator::Reset(Params);
+	}
+
+	void FExternallyClockedMidiPlayerOperator::BindInputs(FInputVertexInterfaceData& InVertexData)
+	{
+		FMidiPlayerOperator::BindInputs(InVertexData);
+		using namespace CommonPinNames;
+		InVertexData.BindReadVertex(METASOUND_GET_PARAM_NAME(Inputs::MidiClock), MidiClockIn);
+
+		MidiClockOut->AttachToTimeAuthority(*MidiClockIn);
+	}
+
+	void FExternallyClockedMidiPlayerOperator::BindOutputs(FOutputVertexInterfaceData& InVertexData)
+	{
+		FMidiPlayerOperator::BindOutputs(InVertexData);
+	}
+
+	FDataReferenceCollection FExternallyClockedMidiPlayerOperator::GetInputs() const
+	{
+		// This should never be called. Bind(...) is called instead. This method
+		// exists as a stop-gap until the API can be deprecated and removed.
+		checkNoEntry();
+		return {};
+	}
+
+	FDataReferenceCollection FExternallyClockedMidiPlayerOperator::GetOutputs() const
+	{
+		// This should never be called. Bind(...) is called instead. This method
+		// exists as a stop-gap until the API can be deprecated and removed.
+		checkNoEntry();
+		return {};
+	}
+
+	void FExternallyClockedMidiPlayerOperator::Reset(const FResetParams& Params)
+	{
+		FMidiPlayerOperator::Reset(Params);
+		LoopOffsetTick = 0;
+	}
+
+	void FSelfClockedMidiPlayerOperator::BindInputs(FInputVertexInterfaceData& InVertexData)
+	{
+		FMidiPlayerOperator::BindInputs(InVertexData);
+	}
+
+	void FSelfClockedMidiPlayerOperator::BindOutputs(FOutputVertexInterfaceData& InVertexData)
+	{
+		FMidiPlayerOperator::BindOutputs(InVertexData);
+	}
+
+	FDataReferenceCollection FSelfClockedMidiPlayerOperator::GetInputs() const
+	{
+		// This should never be called. Bind(...) is called instead. This method
+		// exists as a stop-gap until the API can be deprecated and removed.
+		checkNoEntry();
+		return {};
+	}
+
+	FDataReferenceCollection FSelfClockedMidiPlayerOperator::GetOutputs() const
+	{
+		// This should never be called. Bind(...) is called instead. This method
+		// exists as a stop-gap until the API can be deprecated and removed.
+		checkNoEntry();
+		return {};
+	}
+
+	void FMidiPlayerOperator::Execute()
+	{
+		MidiOutPin->PrepareBlock();
+		MidiClockOut->PrepareBlock();
+
+		if (CurrentMidiFile != MidiAssetInPin->GetMidiProxy())
+		{
+			SetupNewMidiFile(MidiAssetInPin->GetMidiProxy());
+		}
+	}
+
+	bool FMidiPlayerOperator::IsPlaying() const
+	{
+		return GetTransportState() == EMusicPlayerTransportState::Playing
+			|| GetTransportState() == EMusicPlayerTransportState::Starting
+			|| GetTransportState() == EMusicPlayerTransportState::Continuing;
+	}
+
+	void FExternallyClockedMidiPlayerOperator::Execute()
+	{
+		FMidiPlayerOperator::Execute();
+
+		MidiClockOut->CopySpeedAndTempoChanges(MidiClockIn.Get(), *SpeedMultInPin);
+
+		int32 CurrentMidiClockEventIndex = -1;
+		TransportSpanPostProcessor HandleMidiClockEventsInBlock = [&](int32 StartFrameIndex, int32 EndFrameIndex, EMusicPlayerTransportState CurrentState)
+		{
+			const TArray<FMidiClockEvent>& MidiClockEvents = MidiClockIn->GetMidiClockEventsInBlock();
+			for (int32 EventIndex = FMath::Max(CurrentMidiClockEventIndex, 0); EventIndex < MidiClockEvents.Num(); ++EventIndex)
+			{
+				const FMidiClockEvent& Event = MidiClockEvents[EventIndex];
+
+				// we've run past the end of the events. Might as well stop iterating.
+				if (Event.BlockFrameIndex >= EndFrameIndex)
+				{
+					break;
+				}
+
+				if (Event.BlockFrameIndex >= StartFrameIndex && Event.BlockFrameIndex < EndFrameIndex)
+				{
+					// We should be advancing midi clock events in order (never repeating)!
+					check(CurrentMidiClockEventIndex < EventIndex);
+					CurrentMidiClockEventIndex = EventIndex;
+					switch (Event.Type)
+					{
+					case FMidiClockEvent::EType::Reset:
+					{
+						UpdateLoopOffsetTickFromTick(Event.StartTick);
+						int32 Tick = Event.StartTick - LoopOffsetTick;
+
+						float Ms = MidiClockOut->GetSongMaps().TickToMs(Tick);
+						FMusicSeekTarget SeekTarget;
+						SeekTarget.Type = ESeekPointType::Millisecond;
+						SeekTarget.Ms = Ms;
+						MidiClockOut->SeekTo(Event.BlockFrameIndex, SeekTarget, PrerollBars);
+						break;
+					}
+					case FMidiClockEvent::EType::Loop:
+					{
+						// loops should actually be handled by seeking and advancing...
+						break;
+					}
+					case FMidiClockEvent::EType::SeekTo:
+					{
+						UpdateLoopOffsetTickFromTick(Event.StartTick);
+						int32 Tick = Event.StartTick - LoopOffsetTick;
+
+						float Ms = MidiClockOut->GetSongMaps().TickToMs(Tick);
+						FMusicSeekTarget SeekTarget;
+						SeekTarget.Type = ESeekPointType::Millisecond;
+						SeekTarget.Ms = Ms;
+						MidiClockOut->SeekTo(Event.BlockFrameIndex, SeekTarget, PrerollBars);
+						break;
+					}
+					case FMidiClockEvent::EType::SeekThru:
+					{
+						UpdateLoopOffsetTickFromTick(Event.StartTick);
+						int32 Tick = Event.StartTick - LoopOffsetTick;
+
+						float Ms = MidiClockOut->GetSongMaps().TickToMs(Tick + 1);
+						FMusicSeekTarget SeekTarget;
+						SeekTarget.Type = ESeekPointType::Millisecond;
+						SeekTarget.Ms = Ms;
+						MidiClockOut->SeekTo(Event.BlockFrameIndex, SeekTarget, PrerollBars);
+						break;
+					}
+					case FMidiClockEvent::EType::AdvanceThru:
+					{
+						
+						// The MidiClock handles looping on its own, so we can conveniently advance it
+						int32 Tick = Event.StartTick - LoopOffsetTick;
+						float Ms = MidiClockOut->GetSongMaps().TickToMs(Tick);
+
+						float ClockInSpeed = MidiClockIn->GetSpeedAtBlockSampleFrame(StartFrameIndex);
+						// midi clock needs to know how fast its advancing based on their authority
+						MidiClockOut->InformOfCurrentAdvanceRate(ClockInSpeed * *SpeedMultInPin);
+						MidiClockOut->AdvanceHiResToMs(Event.BlockFrameIndex, Ms, true);
+
+						// we have to update our loop offset index _after_ the advance
+						// which we can do a little bit differently here
+						
+						// check if the clock looped back
+						if (MidiClockOut->DoesLoop())
+						{
+							int32 NewTick = MidiClockOut->GetCurrentHiResTick();
+							if (NewTick < Tick)
+							{
+								UpdateLoopOffsetTickFromTick(Event.StartTick);
+							}
+						}
+						break;
+					}
+					}
+				}
+			}
+		};
+
+		TransportSpanProcessor TransportHandler = [&](int32 StartFrameIndex, int32 EndFrameIndex, EMusicPlayerTransportState CurrentState)
+		{
+			CurrentBlockSpanStart = StartFrameIndex;
+			switch (CurrentState)
+			{
+			case EMusicPlayerTransportState::Invalid:
+			case EMusicPlayerTransportState::Preparing:
+				MidiOutPin->AddTransportStateChangeMessage(StartFrameIndex, EMusicPlayerTransportState::Prepared);
+				return EMusicPlayerTransportState::Prepared;
+
+			case EMusicPlayerTransportState::Prepared:
+				MidiOutPin->AddTransportStateChangeMessage(StartFrameIndex, CurrentState);
+				return CurrentState;
+
+			case EMusicPlayerTransportState::Starting:
+				//SampleCounter.SetNumSamples(0);
+				MidiOutPin->AddTransportStateChangeMessage(StartFrameIndex, EMusicPlayerTransportState::Playing);
+				return EMusicPlayerTransportState::Playing;
+
+			case EMusicPlayerTransportState::Playing:
+				return EMusicPlayerTransportState::Playing;
+
+			case EMusicPlayerTransportState::Seeking:
+				// We have nothing to do because our incoming midi clock would have seeked us. 
+				return GetTransportState();
+
+			case EMusicPlayerTransportState::Continuing:
+				MidiOutPin->AddTransportStateChangeMessage(StartFrameIndex, EMusicPlayerTransportState::Continuing);
+				return EMusicPlayerTransportState::Playing;
+
+			case EMusicPlayerTransportState::Pausing:
+				MidiOutPin->AddTransportStateChangeMessage(StartFrameIndex, EMusicPlayerTransportState::Pausing);
+				return EMusicPlayerTransportState::Paused;
+
+			case EMusicPlayerTransportState::Paused:
+				return EMusicPlayerTransportState::Paused;
+
+			case EMusicPlayerTransportState::Stopping:
+				MidiOutPin->AddTransportStateChangeMessage(StartFrameIndex, EMusicPlayerTransportState::Stopping);
+				{
+					FMidiStreamEvent MidiEvent(this, FMidiMsg::CreateAllNotesOff());
+					MidiEvent.BlockSampleFrameIndex = MidiClockOut->GetCurrentBlockFrameIndex();
+					MidiEvent.BlockSampleFrameOffset = 0.0f;
+					MidiEvent.AuthoredMidiTick = MidiClockOut->GetCurrentHiResTick();
+					MidiEvent.CurrentMidiTick = MidiClockOut->GetCurrentHiResTick();
+					MidiEvent.MsOffset = 0.0f;
+					MidiEvent.TrackIndex = 0;
+					MidiOutPin->AddMidiEvent(MidiEvent);
+				}
+				return EMusicPlayerTransportState::Prepared;
+
+			case EMusicPlayerTransportState::Killing:
+				MidiOutPin->AddTransportStateChangeMessage(StartFrameIndex, EMusicPlayerTransportState::Killing);
+				return EMusicPlayerTransportState::Prepared;
+
+			default:
+				checkNoEntry();
+				return EMusicPlayerTransportState::Invalid;
+			}
+		};
+		ExecuteTransportSpans(TransportInPin, BlockSize, TransportHandler, HandleMidiClockEventsInBlock);
+
+		GetTransportState();
+	}
+
+	void FExternallyClockedMidiPlayerOperator::UpdateLoopOffsetTickFromTick(int32 Tick)
+	{
+		// only update our loop offset if we're actually looping
+		// if our external clock is seeking us
+		// we need to update our loop offset ticks
+		if (MidiClockOut->DoesLoop())
+		{
+			int32 LoopStartTick = MidiClockOut->GetLoopStartTick();
+			int32 LoopEndTick = MidiClockOut->GetLoopEndTick();
+			int32 LoopLengthTicks = LoopEndTick - LoopStartTick;
+			if (LoopLengthTicks > 0)
+			{
+				// the whole number part of the division tells us 
+				// how many times we "should have looped" based on the incoming tick
+				int32 LoopNum = (Tick - LoopStartTick) / LoopLengthTicks;
+
+				// So we can multiply it back to get the LoopOffsetTick
+				LoopOffsetTick = LoopLengthTicks * LoopNum + LoopStartTick;
+			}
+		}
+	}
+
+	void FSelfClockedMidiPlayerOperator::Execute()
+	{
+		FMidiPlayerOperator::Execute();
+
+		MidiClockOut->AddSpeedChangeToBlock({ 0, 0.0f, *SpeedMultInPin });
+
+		TransportSpanPostProcessor MidiClockTransportHandler = [&](int32 StartFrameIndex, int32 EndFrameIndex, EMusicPlayerTransportState CurrentState)
+		{
+			switch (CurrentState)
+			{
+			case EMusicPlayerTransportState::Playing:
+				MidiClockOut->WriteAdvance(StartFrameIndex, EndFrameIndex, *SpeedMultInPin);
+				return;
+
+			case EMusicPlayerTransportState::Continuing:
+				MidiClockOut->WriteAdvance(StartFrameIndex, EndFrameIndex, *SpeedMultInPin);
+				return;
+			}
+		};
+
+		TransportSpanProcessor TransportHandler = [this](int32 StartFrameIndex, int32 EndFrameIndex, EMusicPlayerTransportState CurrentState)
+		{
+			switch (CurrentState)
+			{
+			case EMusicPlayerTransportState::Invalid:
+			case EMusicPlayerTransportState::Preparing:
+				// midi clock out
+				MidiClockOut->AddTransportStateChangeToBlock({ StartFrameIndex, 0.0f, EMusicPlayerTransportState::Prepared });
+				MidiClockOut->WriteNoAdvance(StartFrameIndex, EndFrameIndex);
+
+				// midi out
+				MidiOutPin->AddTransportStateChangeMessage(StartFrameIndex, EMusicPlayerTransportState::Prepared);
+				return EMusicPlayerTransportState::Prepared;
+
+			case EMusicPlayerTransportState::Prepared:
+				// midi clock out
+				MidiClockOut->WriteNoAdvance(StartFrameIndex, EndFrameIndex);
+
+				// midi out
+				MidiOutPin->AddTransportStateChangeMessage(StartFrameIndex, CurrentState);
+				return CurrentState;
+
+			case EMusicPlayerTransportState::Starting:
+				// midi clock out 
+				MidiClockOut->ResetAndStart(StartFrameIndex, !ReceivedSeekWhileStopped());
+
+				// midi out
+				MidiOutPin->AddTransportStateChangeMessage(StartFrameIndex, EMusicPlayerTransportState::Playing);
+				return EMusicPlayerTransportState::Playing;
+
+			case EMusicPlayerTransportState::Playing:
+				return EMusicPlayerTransportState::Playing;
+
+			case EMusicPlayerTransportState::Seeking:
+				// midi clock out
+				MidiClockOut->SeekTo(StartFrameIndex, TransportInPin->GetNextSeekDestination(), PrerollBars);
+
+				// We have nothing to do because our incoming midi clock would have seeked us. 
+				return GetTransportState();
+
+			case EMusicPlayerTransportState::Continuing:
+				// midi clock
+				MidiClockOut->AddTransportStateChangeToBlock({ StartFrameIndex, 0.0f, EMusicPlayerTransportState::Playing });
+
+				// midi out
+				MidiOutPin->AddTransportStateChangeMessage(StartFrameIndex, EMusicPlayerTransportState::Continuing);
+				return EMusicPlayerTransportState::Playing;
+
+			case EMusicPlayerTransportState::Pausing:
+				// midi clock out
+				MidiClockOut->AddTransportStateChangeToBlock({ StartFrameIndex, 0.0f, EMusicPlayerTransportState::Paused });
+				MidiClockOut->WriteNoAdvance(StartFrameIndex, EndFrameIndex);
+
+				// midi out
+				return EMusicPlayerTransportState::Paused;
+
+			case EMusicPlayerTransportState::Paused:
+				// midi clock out
+				MidiClockOut->WriteNoAdvance(StartFrameIndex, EndFrameIndex);
+
+				// midi out
+				return EMusicPlayerTransportState::Paused;
+
+			case EMusicPlayerTransportState::Stopping:
+				// midi clock out
+				MidiClockOut->AddTransportStateChangeToBlock({ StartFrameIndex, 0.0f, EMusicPlayerTransportState::Prepared });
+				MidiClockOut->WriteNoAdvance(StartFrameIndex, EndFrameIndex);
+
+				// midi out
+				MidiOutPin->AddTransportStateChangeMessage(StartFrameIndex, EMusicPlayerTransportState::Stopping);
+				return EMusicPlayerTransportState::Prepared;
+
+			case EMusicPlayerTransportState::Killing:
+				// midi clock out
+				MidiClockOut->AddTransportStateChangeToBlock({ StartFrameIndex, 0.0f, EMusicPlayerTransportState::Prepared });
+				MidiClockOut->WriteNoAdvance(StartFrameIndex, EndFrameIndex);
+
+				// midi out
+				MidiOutPin->AddTransportStateChangeMessage(StartFrameIndex, EMusicPlayerTransportState::Killing);
+				return EMusicPlayerTransportState::Prepared;
+
+			default:
+				checkNoEntry();
+				return EMusicPlayerTransportState::Invalid;
+			}
+		};
+		ExecuteTransportSpans(TransportInPin, BlockSize, TransportHandler, MidiClockTransportHandler);
+	}
+
+	void FMidiPlayerOperator::BindInputs(FInputVertexInterfaceData& InVertexData)
+	{
+		using namespace CommonPinNames;
+		InVertexData.BindReadVertex(METASOUND_GET_PARAM_NAME(Inputs::MidiFileAsset), MidiAssetInPin);
+		InVertexData.BindReadVertex(METASOUND_GET_PARAM_NAME(Inputs::Transport), TransportInPin);
+		InVertexData.BindReadVertex(METASOUND_GET_PARAM_NAME(Inputs::Loop), LoopInPin);
+		InVertexData.BindReadVertex(METASOUND_GET_PARAM_NAME(Inputs::Speed), SpeedMultInPin);
+		InVertexData.SetValue(METASOUND_GET_PARAM_NAME(Inputs::PrerollBars), PrerollBars);
+		InVertexData.SetValue(METASOUND_GET_PARAM_NAME(MidiPlayerNodePinNames::KillVoicesOnSeek), bKillVoicesOnSeek);
+		InVertexData.SetValue(METASOUND_GET_PARAM_NAME(MidiPlayerNodePinNames::KillVoicesOnMidiChange), bKillVoicesOnMidiChange);
+
+		SetupNewMidiFile(MidiAssetInPin->GetMidiProxy());
+	}
+
+	void FMidiPlayerOperator::BindOutputs(FOutputVertexInterfaceData& InVertexData)
+	{
+		using namespace CommonPinNames;
+		InVertexData.BindReadVertex(METASOUND_GET_PARAM_NAME(Outputs::MidiStream), MidiOutPin);
+		InVertexData.BindReadVertex(METASOUND_GET_PARAM_NAME(Outputs::MidiClock), MidiClockOut);
+	}
+	
+	void FMidiPlayerOperator::Reset(const FResetParams& Params)
+	{
+		BlockSize = Params.OperatorSettings.GetNumFramesPerBlock();
+		CurrentBlockSpanStart = 0;
+
+		MidiOutPin->SetClockSource(MidiClockOut);
+		MidiClockOut->ResetAndStart(0, true);
+	}
+
+	void FMidiPlayerOperator::Reset(bool ForceNoBroadcast /*= false*/)
+	{
+		FMidiPlayCursor::Reset(ForceNoBroadcast);
+	}
+
+	void FMidiPlayerOperator::OnLoop(int32 LoopStartTick, int32 LoopEndTick)
+	{
+		//TRACE_BOOKMARK(TEXT("Midi Looping"));
+		FMidiPlayCursor::OnLoop(LoopStartTick, LoopEndTick);
+	}
+
+	void FMidiPlayerOperator::SeekToTick(int32 Tick)
+	{
+		//TRACE_BOOKMARK(TEXT("Midi Seek To Tick"));
+		FMidiPlayCursor::SeekToTick(Tick);
+
+		if (IsPlaying())
+		{
+			FMidiStreamEvent MidiEvent(this, bKillVoicesOnSeek ? FMidiMsg::CreateAllNotesKill() : FMidiMsg::CreateAllNotesOff());
+			MidiEvent.BlockSampleFrameIndex = MidiClockOut->GetCurrentBlockFrameIndex();
+			MidiEvent.BlockSampleFrameOffset = 0.0f;
+			MidiEvent.AuthoredMidiTick = MidiClockOut->GetCurrentHiResTick();
+			MidiEvent.CurrentMidiTick = MidiClockOut->GetCurrentHiResTick();
+			MidiEvent.MsOffset = 0.0f;
+			MidiEvent.TrackIndex = 0;
+			MidiOutPin->AddMidiEvent(MidiEvent);
+		}
+	}
+
+	void FMidiPlayerOperator::SeekThruTick(int32 Tick)
+	{
+		//TRACE_BOOKMARK(TEXT("Midi Seek Thru Tick"));
+		FMidiPlayCursor::SeekThruTick(Tick);
+
+		if (IsPlaying())
+		{
+			FMidiStreamEvent MidiEvent(this, bKillVoicesOnSeek ? FMidiMsg::CreateAllNotesKill() : FMidiMsg::CreateAllNotesOff());
+			MidiEvent.BlockSampleFrameIndex = MidiClockOut->GetCurrentBlockFrameIndex();
+			MidiEvent.BlockSampleFrameOffset = 0.0f;
+			MidiEvent.AuthoredMidiTick = MidiClockOut->GetCurrentHiResTick();
+			MidiEvent.CurrentMidiTick = MidiClockOut->GetCurrentHiResTick();
+			MidiEvent.MsOffset = 0.0f;
+			MidiEvent.TrackIndex = 0;
+			MidiOutPin->AddMidiEvent(MidiEvent);
+		}
+	}
+
+	void FMidiPlayerOperator::AdvanceThruTick(int32 Tick, bool IsPreRoll)
+	{
+		FMidiPlayCursor::AdvanceThruTick(Tick, IsPreRoll);
+	}
+
+	void FMidiPlayerOperator::OnMidiMessage(int32 TrackIndex, int32 Tick, uint8 Status, uint8 Data1, uint8 Data2, bool IsPreroll)
+	{
+		if (IsPlaying())
+		{
+			FMidiStreamEvent MidiEvent(this, FMidiMsg(Status, Data1, Data2));
+			MidiEvent.BlockSampleFrameIndex = MidiClockOut->GetCurrentBlockFrameIndex();
+			MidiEvent.BlockSampleFrameOffset = 0.0f;
+			MidiEvent.AuthoredMidiTick = Tick;
+			MidiEvent.CurrentMidiTick = Tick;
+			MidiEvent.MsOffset = 0.0f;
+			MidiEvent.TrackIndex = TrackIndex;
+			MidiOutPin->AddMidiEvent(MidiEvent);
+		}
+	}
+
+	void FMidiPlayerOperator::OnTempo(int32 TrackIndex, int32 Tick, int32 tempo, bool isPreroll)
+	{
+		if (IsPlaying())
+		{
+			FMidiStreamEvent MidiEvent(this, FMidiMsg((int32)tempo));
+			MidiEvent.BlockSampleFrameIndex  = MidiClockOut->GetCurrentBlockFrameIndex();
+			MidiEvent.BlockSampleFrameOffset = 0.0f;
+			MidiEvent.AuthoredMidiTick = Tick;
+			MidiEvent.CurrentMidiTick  = Tick;
+			MidiEvent.MsOffset         = 0.0f;
+			MidiEvent.TrackIndex       = TrackIndex;
+			MidiOutPin->AddMidiEvent(MidiEvent);
+		}
+	}
+
+	void FMidiPlayerOperator::OnText(int32 TrackIndex, int32 Tick, int32 TextIndex, const FString& Str, uint8 Type, bool IsPreroll)
+	{
+		if (!IsPreroll && IsPlaying())
+		{
+			FMidiStreamEvent MidiEvent(this, FMidiMsg::CreateText(TextIndex, Type));
+			MidiEvent.BlockSampleFrameIndex = MidiClockOut->GetCurrentBlockFrameIndex();
+			MidiEvent.BlockSampleFrameOffset = 0.0f;
+			MidiEvent.AuthoredMidiTick = Tick;
+			MidiEvent.CurrentMidiTick = Tick;
+			MidiEvent.MsOffset = 0.0f;
+			MidiEvent.TrackIndex = TrackIndex;
+			MidiOutPin->AddMidiEvent(MidiEvent);
+		}
+	}
+
+
+	void FMidiPlayerOperator::OnPreRollNoteOn(int32 TrackIndex, int32 EventTick, int32 InCurrentTick, float InPrerollMs, uint8 InStatus, uint8 Data1, uint8 Data2)
+	{
+		if (IsPlaying())
+		{
+			FMidiStreamEvent MidiEvent(this, FMidiMsg(InStatus, Data1, Data2));
+			MidiEvent.BlockSampleFrameIndex = MidiClockOut->GetCurrentBlockFrameIndex();
+			MidiEvent.BlockSampleFrameOffset = 0.0f;
+			MidiEvent.AuthoredMidiTick = EventTick;
+			MidiEvent.CurrentMidiTick = InCurrentTick;
+			MidiEvent.MsOffset = InPrerollMs;
+			MidiEvent.TrackIndex = TrackIndex;
+			MidiOutPin->AddMidiEvent(MidiEvent);
+		}
+	}
+
+	void FMidiPlayerOperator::SetupNewMidiFile(const FMidiFileProxyPtr& NewMidi)
+	{
+		CurrentMidiFile = NewMidi;
+		MidiOutPin->SetMidiFile(CurrentMidiFile);
+
+		FMidiStreamEvent MidiEvent(this, bKillVoicesOnMidiChange ? FMidiMsg::CreateAllNotesKill() : FMidiMsg::CreateAllNotesOff());
+		MidiEvent.BlockSampleFrameIndex = 0;
+		MidiEvent.BlockSampleFrameOffset = 0.0f;
+		MidiEvent.AuthoredMidiTick = 0;
+		MidiEvent.CurrentMidiTick = 0;
+		MidiEvent.MsOffset = 0.0f;
+		MidiEvent.TrackIndex = 0;
+		MidiOutPin->InsertMidiEvent(MidiEvent);
+
+		if (CurrentMidiFile.IsValid())
+		{			
+			MidiClockOut->AttachToMidiResource(CurrentMidiFile->GetMidiFile(), !IsPlaying(), PrerollBars);
+			MidiOutPin->SetTicksPerQuarterNote(CurrentMidiFile->GetMidiFile()->TicksPerQuarterNote);
+			if (*LoopInPin)
+			{
+				const FSongLengthData& SongLengthData = CurrentMidiFile->GetMidiFile()->SongMaps.GetSongLengthData();
+				int32 LoopStartTick = 0;
+				int32 LoopEndTick = SongLengthData.LengthTicks;
+
+				// Round the content authored ticks to bar boundaries based on the bar map. 
+				const FBarMap& BarMap = CurrentMidiFile->GetMidiFile()->SongMaps.GetBarMap();
+				int32 LoopStartBarIndex = FMath::RoundToInt32(BarMap.TickToFractionalBarIncludingCountIn(LoopStartTick));
+				int32 LoopEndBarIndex = FMath::RoundToInt32(BarMap.TickToFractionalBarIncludingCountIn(LoopEndTick));
+
+				// Make sure there's at least 1 bar of looping
+				if (LoopEndBarIndex <= LoopStartBarIndex)
+				{
+					LoopEndBarIndex = LoopStartBarIndex + 1;
+				}
+				LoopStartTick = BarMap.BarIncludingCountInToTick(LoopStartBarIndex);
+				LoopEndTick = BarMap.BarIncludingCountInToTick(LoopEndBarIndex);
+				
+				MidiClockOut->SetLoop(LoopStartTick, LoopEndTick);
+			}
+			else
+			{
+				MidiClockOut->ClearLoop();
+			}
+		}
+		else
+		{
+			MidiClockOut->AttachToMidiResource(nullptr, !IsPlaying(), 0);
+		}
+	}
+}
+
+#undef LOCTEXT_NAMESPACE // "HarmonixMetaSound"
diff --git a/Engine/Plugins/Runtime/Harmonix/Source/HarmonixMetasound/Private/HarmonixMetasound/Nodes/MidiStreamLoggerNode.cpp b/Engine/Plugins/Runtime/Harmonix/Source/HarmonixMetasound/Private/HarmonixMetasound/Nodes/MidiStreamLoggerNode.cpp
new file mode 100644
--- /dev/null
+++ b/Engine/Plugins/Runtime/Harmonix/Source/HarmonixMetasound/Private/HarmonixMetasound/Nodes/MidiStreamLoggerNode.cpp
@@ -0,0 +1,213 @@
+// Copyright Epic Games, Inc. All Rights Reserved.
+
+#include "MetasoundExecutableOperator.h"
+#include "MetasoundFacade.h"
+#include "MetasoundNodeInterface.h"
+#include "MetasoundParamHelper.h"
+#include "MetasoundSampleCounter.h"
+#include "MetasoundStandardNodesCategories.h"
+#include "MetasoundVertex.h"
+
+#include "HarmonixMetasound/Common.h"
+#include "HarmonixMetasound/DataTypes/MidiStream.h"
+#include "HarmonixMetasound/DataTypes/MusicTransport.h"
+
+DEFINE_LOG_CATEGORY_STATIC(LogMidiStreamLogger, Log, All);
+
+#define LOCTEXT_NAMESPACE "HarmonixMetaSound"
+
+namespace HarmonixMetasound
+{
+	using namespace Metasound;
+
+	class FMidiStreamLoggerOperator : public TExecutableOperator<FMidiStreamLoggerOperator>
+	{
+	public:
+		static const FNodeClassMetadata& GetNodeInfo();
+		static const FVertexInterface& GetVertexInterface();
+		static TUniquePtr<IOperator> CreateOperator(const FBuildOperatorParams& InParams, FBuildResults& OutResults);
+
+		FMidiStreamLoggerOperator(const FBuildOperatorParams& InParams,
+								  const FBoolReadRef& InEnabled,
+								  const FMidiStreamReadRef& InMidiStream);
+
+		virtual void BindInputs(FInputVertexInterfaceData& InVertexData) override;
+		virtual void BindOutputs(FOutputVertexInterfaceData& InVertexData) override;
+		virtual FDataReferenceCollection GetInputs() const override;
+		virtual FDataReferenceCollection GetOutputs() const override;
+
+		void Reset(const FResetParams& ResetParams);
+
+		void Execute();
+
+	private:
+		//** INPUTS
+		FBoolReadRef       EnableInPin;
+		FMidiStreamReadRef MidiStreamInPin;
+
+		//** DATA
+		void DumpTransportEvent(const FMidiTimestampTransportState& TransportEvent);
+		void DumpMidiEvent(const FMidiStreamEvent& MidiEvent);
+	};
+
+	class FMidiStreamLoggerNode : public FNodeFacade
+	{
+	public:
+		FMidiStreamLoggerNode(const FNodeInitData& InInitData)
+			: FNodeFacade(InInitData.InstanceName, InInitData.InstanceID, TFacadeOperatorClass<FMidiStreamLoggerOperator>())
+		{}
+		virtual ~FMidiStreamLoggerNode() = default;
+	};
+
+	METASOUND_REGISTER_NODE(FMidiStreamLoggerNode)
+		
+	const FNodeClassMetadata& FMidiStreamLoggerOperator::GetNodeInfo()
+	{
+		auto InitNodeInfo = []() -> FNodeClassMetadata
+		{
+			FNodeClassMetadata Info;
+			Info.ClassName        = { HarmonixNodeNamespace, TEXT("MidiStreamLogger"), TEXT("")};
+			Info.MajorVersion     = 0;
+			Info.MinorVersion     = 1;
+			Info.DisplayName      = METASOUND_LOCTEXT("MidiStreamLoggerNode_DisplayName", "Midi Stream Logger");
+			Info.Description      = METASOUND_LOCTEXT("MidiStreamLoggerNode_Description", "Receives midi messages and writes them to the log.");
+			Info.Author           = PluginAuthor;
+			Info.PromptIfMissing  = PluginNodeMissingPrompt;
+			Info.DefaultInterface = GetVertexInterface();
+			Info.CategoryHierarchy.Emplace(NodeCategories::Music);
+			return Info;
+		};
+
+		static const FNodeClassMetadata Info = InitNodeInfo();
+
+		return Info;
+	}
+
+	const FVertexInterface& FMidiStreamLoggerOperator::GetVertexInterface()
+	{
+		using namespace CommonPinNames;
+
+		static const FVertexInterface Interface(
+			FInputVertexInterface(
+				TInputDataVertex<bool>(METASOUND_GET_PARAM_NAME_AND_METADATA(Inputs::Enable), true),
+				TInputDataVertex<FMidiStream>(METASOUND_GET_PARAM_NAME_AND_METADATA(Inputs::MidiStream))
+			),
+			FOutputVertexInterface()
+		);
+
+		return Interface;
+	}
+
+	TUniquePtr<IOperator> FMidiStreamLoggerOperator::CreateOperator(const FBuildOperatorParams& InParams, FBuildResults& OutResults)
+	{
+		using namespace CommonPinNames;
+
+		const FMidiStreamLoggerNode& LoggerNode = static_cast<const FMidiStreamLoggerNode&>(InParams.Node);
+
+		const FInputVertexInterfaceData& InputData = InParams.InputData;
+		FBoolReadRef InEnabled = InputData.GetOrCreateDefaultDataReadReference<bool>(METASOUND_GET_PARAM_NAME(Inputs::Enable), InParams.OperatorSettings);
+		FMidiStreamReadRef InMidiStream = InputData.GetOrConstructDataReadReference<FMidiStream>(METASOUND_GET_PARAM_NAME(Inputs::MidiStream), InParams.OperatorSettings);
+
+		return MakeUnique<FMidiStreamLoggerOperator>(InParams, InEnabled, InMidiStream);
+	}
+
+	FMidiStreamLoggerOperator::FMidiStreamLoggerOperator(
+		const FBuildOperatorParams& InParams,
+		const FBoolReadRef& InEnabled,
+		const FMidiStreamReadRef& InMidiStream)
+	: EnableInPin(InEnabled)
+	, MidiStreamInPin(InMidiStream)
+	{
+		Reset(InParams);
+	}
+
+	void FMidiStreamLoggerOperator::BindInputs(FInputVertexInterfaceData& InVertexData)
+	{
+		using namespace CommonPinNames;
+
+		InVertexData.BindReadVertex(METASOUND_GET_PARAM_NAME(Inputs::Enable),     EnableInPin);
+		InVertexData.BindReadVertex(METASOUND_GET_PARAM_NAME(Inputs::MidiStream), MidiStreamInPin);
+	}
+
+	void FMidiStreamLoggerOperator::BindOutputs(FOutputVertexInterfaceData& InVertexData)
+	{
+	}
+
+	FDataReferenceCollection FMidiStreamLoggerOperator::GetInputs() const
+	{
+		// This should never be called. Bind(...) is called instead. This method
+		// exists as a stop-gap until the API can be deprecated and removed.
+		checkNoEntry();
+		return {};
+	}
+
+	FDataReferenceCollection FMidiStreamLoggerOperator::GetOutputs() const
+	{
+		// This should never be called. Bind(...) is called instead. This method
+		// exists as a stop-gap until the API can be deprecated and removed.
+		checkNoEntry();
+		return {};
+	}
+
+	void FMidiStreamLoggerOperator::Reset(const FResetParams& ResetParams)
+	{
+	}
+
+	void FMidiStreamLoggerOperator::Execute()
+	{
+		if (!*EnableInPin)
+		{
+			return;
+		}
+
+		const TArray<FMidiTimestampTransportState>& TransportEvents = MidiStreamInPin->GetTransportChangesInBlock();
+		const TArray<FMidiStreamEvent>& MidiEvents = MidiStreamInPin->GetEventsInBlock();
+		auto TransportEventIterator = TransportEvents.begin();
+		auto MidiEventIterator = MidiEvents.begin();
+		while (TransportEventIterator != TransportEvents.end() || MidiEventIterator != MidiEvents.end())
+		{
+			if (TransportEventIterator != TransportEvents.end() && MidiEventIterator != MidiEvents.end())
+			{
+				if ((*TransportEventIterator).BlockSampleFrameIndex <= (*MidiEventIterator).BlockSampleFrameIndex)
+				{
+					DumpTransportEvent(*TransportEventIterator);
+					++TransportEventIterator;
+				}
+				else
+				{
+					DumpMidiEvent(*MidiEventIterator);
+					++MidiEventIterator;
+				}
+			}
+			else if (TransportEventIterator != TransportEvents.end())
+			{
+				DumpTransportEvent(*TransportEventIterator);
+				++TransportEventIterator;
+			}
+			else
+			{
+				DumpMidiEvent(*MidiEventIterator);
+				++MidiEventIterator;
+			}
+		}
+	}
+
+	void FMidiStreamLoggerOperator::DumpTransportEvent(const FMidiTimestampTransportState& TransportEvent)
+	{
+		UE_LOG(LogMidiStreamLogger, Log, TEXT("[%d (%f)] Transport: %s"), TransportEvent.BlockSampleFrameIndex, TransportEvent.BlockSampleFrameOffset, *FMusicTransportControllable::StateToString(TransportEvent.TransportState));
+	}
+
+	void FMidiStreamLoggerOperator::DumpMidiEvent(const FMidiStreamEvent& MidiEvent)
+	{
+		UE_LOG(LogMidiStreamLogger, Log, TEXT("[%d (%f)] Track: %d, Auth Tick %d, Render Tick %d, Ms Offset %f, Message: %s"),
+			MidiEvent.BlockSampleFrameIndex,
+			MidiEvent.BlockSampleFrameOffset,
+			MidiEvent.TrackIndex,
+			MidiEvent.AuthoredMidiTick,
+			MidiEvent.CurrentMidiTick,
+			MidiEvent.MsOffset,
+			*FMidiMsg::ToString(MidiEvent.MidiMessage));
+	}
+}
+
+#undef LOCTEXT_NAMESPACE // "HarmonixMetaSound"
diff --git a/Engine/Plugins/Runtime/Harmonix/Source/HarmonixMetasound/Private/HarmonixMetasound/Nodes/MidiStreamMergeNode.cpp b/Engine/Plugins/Runtime/Harmonix/Source/HarmonixMetasound/Private/HarmonixMetasound/Nodes/MidiStreamMergeNode.cpp
new file mode 100644
--- /dev/null
+++ b/Engine/Plugins/Runtime/Harmonix/Source/HarmonixMetasound/Private/HarmonixMetasound/Nodes/MidiStreamMergeNode.cpp
@@ -0,0 +1,187 @@
+// Copyright Epic Games, Inc. All Rights Reserved.
+
+#include "MetasoundExecutableOperator.h"
+#include "MetasoundFacade.h"
+#include "MetasoundNodeInterface.h"
+#include "MetasoundParamHelper.h"
+#include "MetasoundSampleCounter.h"
+#include "MetasoundStandardNodesCategories.h"
+#include "MetasoundVertex.h"
+
+#include "HarmonixMetasound/Common.h"
+#include "HarmonixMetasound/DataTypes/MidiStream.h"
+#include "HarmonixMetasound/DataTypes/MusicTransport.h"
+
+DEFINE_LOG_CATEGORY_STATIC(LogMidiStreamMerge, Log, All);
+
+#define LOCTEXT_NAMESPACE "HarmonixMetaSound"
+
+namespace HarmonixMetasound
+{
+	using namespace Metasound;
+
+	class FMidiStreamMergeOperator : public TExecutableOperator<FMidiStreamMergeOperator>
+	{
+	public:
+		static const FNodeClassMetadata& GetNodeInfo();
+		static const FVertexInterface& GetVertexInterface();
+		static TUniquePtr<IOperator> CreateOperator(const FBuildOperatorParams& InParams, FBuildResults& OutResults);
+
+		FMidiStreamMergeOperator(const FBuildOperatorParams& InParams,
+								  const FMidiStreamReadRef& InMidiStreamA,
+								  const FMidiStreamReadRef& InMidiStreamB);
+
+		virtual void BindInputs(FInputVertexInterfaceData& InVertexData) override;
+		virtual void BindOutputs(FOutputVertexInterfaceData& InVertexData) override;
+		virtual FDataReferenceCollection GetInputs() const override;
+		virtual FDataReferenceCollection GetOutputs() const override;
+
+		void Reset(const FResetParams& ResetParams);
+		
+		void Execute();
+
+	private:
+		//** INPUTS
+		FMidiStreamReadRef MidiStreamAInPin;
+		FMidiStreamReadRef MidiStreamBInPin;
+
+		//** OUTPUTS
+		FMidiStreamWriteRef MidiStreamOutPin;
+	};
+
+	class FMidiStreamMergeNode : public FNodeFacade
+	{
+	public:
+		FMidiStreamMergeNode(const FNodeInitData& InInitData)
+			: FNodeFacade(InInitData.InstanceName, InInitData.InstanceID, TFacadeOperatorClass<FMidiStreamMergeOperator>())
+		{}
+		virtual ~FMidiStreamMergeNode() = default;
+	};
+
+	METASOUND_REGISTER_NODE(FMidiStreamMergeNode)
+		
+	const FNodeClassMetadata& FMidiStreamMergeOperator::GetNodeInfo()
+	{
+		auto InitNodeInfo = []() -> FNodeClassMetadata
+		{
+			FNodeClassMetadata Info;
+			Info.ClassName        = { HarmonixNodeNamespace, TEXT("MidiStreamMerge"), TEXT("")};
+			Info.MajorVersion     = 0;
+			Info.MinorVersion     = 1;
+			Info.DisplayName      = METASOUND_LOCTEXT("MidiStreamMergeNode_DisplayName", "Midi Merge");
+			Info.Description      = METASOUND_LOCTEXT("MidiStreamMergeNode_Description", "Combines two midi streams into a single stream.");
+			Info.Author           = PluginAuthor;
+			Info.PromptIfMissing  = PluginNodeMissingPrompt;
+			Info.DefaultInterface = GetVertexInterface();
+			Info.CategoryHierarchy.Emplace(NodeCategories::Music);
+			return Info;
+		};
+
+		static const FNodeClassMetadata Info = InitNodeInfo();
+
+		return Info;
+	}
+
+	namespace MergePinNames
+	{
+		METASOUND_PARAM(InputMidiStreamA, "Midi Stream A", "The first midi stream to merge");
+		METASOUND_PARAM(InputMidiStreamB, "Midi Stream B", "The second midi stream to merge");
+	}
+
+	const FVertexInterface& FMidiStreamMergeOperator::GetVertexInterface()
+	{
+		using namespace CommonPinNames;
+		using namespace MergePinNames;
+
+		static const FVertexInterface Interface(
+			FInputVertexInterface(
+				TInputDataVertex<FMidiStream>(METASOUND_GET_PARAM_NAME_AND_METADATA(InputMidiStreamA)),
+				TInputDataVertex<FMidiStream>(METASOUND_GET_PARAM_NAME_AND_METADATA(InputMidiStreamB))
+			),
+			FOutputVertexInterface(
+				TOutputDataVertex<FMidiStream>(METASOUND_GET_PARAM_NAME_AND_METADATA(Outputs::MidiStream))
+			)
+		);
+
+		return Interface;
+	}
+
+	TUniquePtr<IOperator> FMidiStreamMergeOperator::CreateOperator(const FBuildOperatorParams& InParams, FBuildResults& OutResults)
+	{
+		using namespace MergePinNames;
+
+		const FMidiStreamMergeNode& LoggerNode = static_cast<const FMidiStreamMergeNode&>(InParams.Node);
+
+		const FInputVertexInterfaceData& InputData = InParams.InputData;
+		FMidiStreamReadRef InMidiStreamA = InputData.GetOrConstructDataReadReference<FMidiStream>(METASOUND_GET_PARAM_NAME(InputMidiStreamA), InParams.OperatorSettings);
+		FMidiStreamReadRef InMidiStreamB = InputData.GetOrConstructDataReadReference<FMidiStream>(METASOUND_GET_PARAM_NAME(InputMidiStreamB), InParams.OperatorSettings);
+
+		return MakeUnique<FMidiStreamMergeOperator>(InParams, InMidiStreamA, InMidiStreamB);
+	}
+
+	FMidiStreamMergeOperator::FMidiStreamMergeOperator(const FBuildOperatorParams& InParams,
+											   const FMidiStreamReadRef& InMidiStreamA,
+											   const FMidiStreamReadRef& InMidiStreamB)
+		: MidiStreamAInPin(InMidiStreamA)
+		, MidiStreamBInPin(InMidiStreamB)
+		, MidiStreamOutPin(FMidiStreamWriteRef::CreateNew(InParams.OperatorSettings))
+	{
+		Reset(InParams);
+	}
+
+	void FMidiStreamMergeOperator::BindInputs(FInputVertexInterfaceData& InVertexData)
+	{
+		using namespace MergePinNames;
+
+		InVertexData.BindReadVertex(METASOUND_GET_PARAM_NAME(InputMidiStreamA), MidiStreamAInPin);
+		InVertexData.BindReadVertex(METASOUND_GET_PARAM_NAME(InputMidiStreamB), MidiStreamBInPin);
+	}
+
+	void FMidiStreamMergeOperator::BindOutputs(FOutputVertexInterfaceData& InVertexData)
+	{
+		using namespace CommonPinNames;
+
+		InVertexData.BindReadVertex(METASOUND_GET_PARAM_NAME(Outputs::MidiStream), MidiStreamOutPin);
+	}
+
+	FDataReferenceCollection FMidiStreamMergeOperator::GetInputs() const
+	{
+		// This should never be called. Bind(...) is called instead. This method
+		// exists as a stop-gap until the API can be deprecated and removed.
+		checkNoEntry();
+		return {};
+	}
+
+	FDataReferenceCollection FMidiStreamMergeOperator::GetOutputs() const
+	{
+		// This should never be called. Bind(...) is called instead. This method
+		// exists as a stop-gap until the API can be deprecated and removed.
+		checkNoEntry();
+		return {};
+	}
+
+	void FMidiStreamMergeOperator::Reset(const FResetParams& ResetParams)
+	{
+		MidiStreamOutPin->PrepareBlock();
+		
+		if (MidiStreamAInPin.Get() && MidiStreamAInPin->GetMidiClockSource())
+		{
+			MidiStreamOutPin->SetClockSource(*(MidiStreamAInPin->GetMidiClockSource()));
+		}
+		else if (MidiStreamBInPin.Get() && MidiStreamBInPin->GetMidiClockSource())
+		{
+			MidiStreamOutPin->SetClockSource(*(MidiStreamBInPin->GetMidiClockSource()));
+		}
+	}
+
+	void FMidiStreamMergeOperator::Execute()
+	{
+		MidiStreamOutPin->PrepareBlock();
+
+		TArray<FMidiStreamReadRef> MidiStreams = { MidiStreamAInPin, MidiStreamBInPin };
+
+		MidiStreamOutPin->Copy(MidiStreams);
+	}
+}
+
+#undef LOCTEXT_NAMESPACE // "HarmonixMetaSound"
diff --git a/Engine/Plugins/Runtime/Harmonix/Source/HarmonixMetasound/Private/HarmonixMetasound/Nodes/MidiStreamSelectNode.cpp b/Engine/Plugins/Runtime/Harmonix/Source/HarmonixMetasound/Private/HarmonixMetasound/Nodes/MidiStreamSelectNode.cpp
new file mode 100644
--- /dev/null
+++ b/Engine/Plugins/Runtime/Harmonix/Source/HarmonixMetasound/Private/HarmonixMetasound/Nodes/MidiStreamSelectNode.cpp
@@ -0,0 +1,296 @@
+// Copyright Epic Games, Inc. All Rights Reserved.
+
+#include "MetasoundExecutableOperator.h"
+#include "MetasoundFacade.h"
+#include "MetasoundNodeInterface.h"
+#include "MetasoundParamHelper.h"
+#include "MetasoundSampleCounter.h"
+#include "MetasoundStandardNodesCategories.h"
+#include "MetasoundVertex.h"
+
+#include "HarmonixMetasound/Common.h"
+#include "HarmonixMetasound/DataTypes/MidiStream.h"
+#include "HarmonixMetasound/DataTypes/MusicTransport.h"
+#include "HarmonixMidi/MidiMsg.h"
+
+DEFINE_LOG_CATEGORY_STATIC(LogMidiStreamSelect, Log, All);
+
+#define LOCTEXT_NAMESPACE "HaronixMetaSound"
+
+namespace HarmonixMetasound
+{
+	using namespace Metasound;
+
+	class FMidiStreamSelectOperator : public TExecutableOperator<FMidiStreamSelectOperator>
+	{
+	public:
+		static const FNodeClassMetadata& GetNodeInfo();
+		static const FVertexInterface& GetVertexInterface();
+		static TUniquePtr<IOperator> CreateOperator(const FBuildOperatorParams& InParams, FBuildResults& OutResults);
+
+		FMidiStreamSelectOperator(const FBuildOperatorParams& InParams,
+			const FMidiStreamReadRef& InMidiStreamA,
+			const FMidiStreamReadRef& InMidiStreamB,
+			const FInt32ReadRef& InStreamIndex,
+			const FBoolReadRef& InImmediateNoteOff,
+			const FBoolReadRef& InCopyInactiveNoteOffs);
+
+		virtual void BindInputs(FInputVertexInterfaceData& InVertexData) override;
+		virtual void BindOutputs(FOutputVertexInterfaceData& InVertexData) override;
+		virtual FDataReferenceCollection GetInputs() const override;
+		virtual FDataReferenceCollection GetOutputs() const override;
+
+		void Execute();
+		void Reset(const FResetParams& ResetParams);
+
+	private:
+		//** INPUTS
+		FMidiStreamReadRef MidiStreamAInPin;
+		FMidiStreamReadRef MidiStreamBInPin;
+		FInt32ReadRef	   StreamIndexInPin;
+		FBoolReadRef	   ImmediateNoteOffInPin;
+		FBoolReadRef	   CopyInactiveNoteOffsInPin;
+
+		//** OUTPUTS
+		FMidiStreamWriteRef MidiStreamOutPin;
+
+		//** DATA
+		TSet<FMidiVoiceId> PlayingVoices;
+		int32 CurrentStreamIndex = -1;
+
+		void CopyFromSelectedStream(int32 SelectedStreamIndex);
+		void MergeOrCreateNoteOffsForDeselectedStream(int32 SelectedStreamIndex);
+	};
+
+	class FMidiStreamSelectNode : public FNodeFacade
+	{
+	public:
+		FMidiStreamSelectNode(const FNodeInitData& InInitData)
+			: FNodeFacade(InInitData.InstanceName, InInitData.InstanceID, TFacadeOperatorClass<FMidiStreamSelectOperator>())
+		{}
+		virtual ~FMidiStreamSelectNode() = default;
+	};
+
+	METASOUND_REGISTER_NODE(FMidiStreamSelectNode)
+
+	const FNodeClassMetadata& FMidiStreamSelectOperator::GetNodeInfo()
+	{
+		auto InitNodeInfo = []() -> FNodeClassMetadata
+		{
+			FNodeClassMetadata Info;
+			Info.ClassName = { HarmonixNodeNamespace, TEXT("MidiStreamSelect"), TEXT("") };
+			Info.MajorVersion = 0;
+			Info.MinorVersion = 1;
+			Info.DisplayName = METASOUND_LOCTEXT("MidiStreamSelectNode_DisplayName", "Midi Select");
+			Info.Description = METASOUND_LOCTEXT("MidiStreamSelectNode_Description", "Copies the output of one of the inputs to the output.");
+			Info.Author = PluginAuthor;
+			Info.PromptIfMissing = PluginNodeMissingPrompt;
+			Info.DefaultInterface = GetVertexInterface();
+			Info.CategoryHierarchy.Emplace(NodeCategories::Music);
+			return Info;
+		};
+
+		static const FNodeClassMetadata Info = InitNodeInfo();
+
+		return Info;
+	}
+
+	namespace SelectPinNames
+	{
+		METASOUND_PARAM(InputMidiStreamA, "Midi Stream A", "The first midi stream to select.");
+		METASOUND_PARAM(InputMidiStreamB, "Midi Stream B", "The second midi stream to select.");
+		METASOUND_PARAM(InputStreamIndex, "Index", "The index of the midi stream to copy.");
+		METASOUND_PARAM(InputImmediateNoteOff, "Immediately Note Off", "Should all active midi in the stream turn off when the index changes?");
+		METASOUND_PARAM(InputCopyInactiveNoteOffs, "Copy Inactive Note Offs", "Should note offs for active midi be copied from inactive streams?");
+	}
+
+	const FVertexInterface& FMidiStreamSelectOperator::GetVertexInterface()
+	{
+		using namespace CommonPinNames;
+		using namespace SelectPinNames;
+
+		static const FVertexInterface Interface(
+			FInputVertexInterface(
+				TInputDataVertex<FMidiStream>(METASOUND_GET_PARAM_NAME_AND_METADATA(InputMidiStreamA)),
+				TInputDataVertex<FMidiStream>(METASOUND_GET_PARAM_NAME_AND_METADATA(InputMidiStreamB)),
+				TInputDataVertex<int32>(METASOUND_GET_PARAM_NAME_AND_METADATA(InputStreamIndex)),
+				TInputDataVertex<bool>(METASOUND_GET_PARAM_NAME_AND_METADATA(InputImmediateNoteOff), true),
+				TInputDataVertex<bool>(METASOUND_GET_PARAM_NAME_AND_METADATA(InputCopyInactiveNoteOffs), true)
+			),
+			FOutputVertexInterface(
+				TOutputDataVertex<FMidiStream>(METASOUND_GET_PARAM_NAME_AND_METADATA(Outputs::MidiStream))
+			)
+		);
+
+		return Interface;
+	}
+
+	TUniquePtr<IOperator> FMidiStreamSelectOperator::CreateOperator(const FBuildOperatorParams& InParams, FBuildResults& OutResults)
+	{
+		using namespace SelectPinNames;
+
+		const FMidiStreamSelectNode& LoggerNode = static_cast<const FMidiStreamSelectNode&>(InParams.Node);
+
+		const FInputVertexInterfaceData& InputData = InParams.InputData;
+		FMidiStreamReadRef InMidiStreamA = InputData.GetOrConstructDataReadReference<FMidiStream>(METASOUND_GET_PARAM_NAME(InputMidiStreamA), InParams.OperatorSettings);
+		FMidiStreamReadRef InMidiStreamB = InputData.GetOrConstructDataReadReference<FMidiStream>(METASOUND_GET_PARAM_NAME(InputMidiStreamB), InParams.OperatorSettings);
+		FInt32ReadRef InStreamIndex = InputData.GetOrConstructDataReadReference<int32>(METASOUND_GET_PARAM_NAME(InputStreamIndex));
+		FBoolReadRef InImmediateNoteOff = InputData.GetOrConstructDataReadReference<bool>(METASOUND_GET_PARAM_NAME(InputImmediateNoteOff));
+		FBoolReadRef InCopyInactiveNoteOffs = InputData.GetOrConstructDataReadReference<bool>(METASOUND_GET_PARAM_NAME(InputCopyInactiveNoteOffs));
+
+		return MakeUnique<FMidiStreamSelectOperator>(InParams, InMidiStreamA, InMidiStreamB, InStreamIndex, InImmediateNoteOff, InCopyInactiveNoteOffs);
+	}
+
+	FMidiStreamSelectOperator::FMidiStreamSelectOperator(const FBuildOperatorParams& InParams,
+		const FMidiStreamReadRef& InMidiStreamA,
+		const FMidiStreamReadRef& InMidiStreamB,
+		const FInt32ReadRef& InStreamIndex,
+		const FBoolReadRef& InImmediateNoteOff,
+		const FBoolReadRef& InCopyInactiveNoteOffs)
+		: MidiStreamAInPin(InMidiStreamA)
+		, MidiStreamBInPin(InMidiStreamB)
+		, StreamIndexInPin(InStreamIndex)
+		, ImmediateNoteOffInPin(InImmediateNoteOff)
+		, CopyInactiveNoteOffsInPin(InCopyInactiveNoteOffs)
+		, MidiStreamOutPin(FMidiStreamWriteRef::CreateNew(InParams.OperatorSettings))
+	{
+		Reset(InParams);
+	}
+
+	void FMidiStreamSelectOperator::BindInputs(FInputVertexInterfaceData& InVertexData)
+	{
+		using namespace SelectPinNames;
+
+		InVertexData.BindReadVertex(METASOUND_GET_PARAM_NAME(InputMidiStreamA), MidiStreamAInPin);
+		InVertexData.BindReadVertex(METASOUND_GET_PARAM_NAME(InputMidiStreamB), MidiStreamBInPin);
+		InVertexData.BindReadVertex(METASOUND_GET_PARAM_NAME(InputStreamIndex), StreamIndexInPin);
+		InVertexData.BindReadVertex(METASOUND_GET_PARAM_NAME(InputImmediateNoteOff), ImmediateNoteOffInPin);
+		InVertexData.BindReadVertex(METASOUND_GET_PARAM_NAME(InputCopyInactiveNoteOffs), CopyInactiveNoteOffsInPin);
+	}
+
+	void FMidiStreamSelectOperator::BindOutputs(FOutputVertexInterfaceData& InVertexData)
+	{
+		using namespace CommonPinNames;
+
+		InVertexData.BindReadVertex(METASOUND_GET_PARAM_NAME(Outputs::MidiStream), MidiStreamOutPin);
+	}
+
+	FDataReferenceCollection FMidiStreamSelectOperator::GetInputs() const
+	{
+		// This should never be called. Bind(...) is called instead. This method
+		// exists as a stop-gap until the API can be deprecated and removed.
+		checkNoEntry();
+		return {};
+	}
+
+	FDataReferenceCollection FMidiStreamSelectOperator::GetOutputs() const
+	{
+		// This should never be called. Bind(...) is called instead. This method
+		// exists as a stop-gap until the API can be deprecated and removed.
+		checkNoEntry();
+		return {};
+	}
+
+	void FMidiStreamSelectOperator::Execute()
+	{
+		MidiStreamOutPin->PrepareBlock();
+
+		if (*StreamIndexInPin > 1 || *StreamIndexInPin < 0)
+		{
+			return;
+		}
+
+		int32 NewCurrentStreamIndex = *StreamIndexInPin;
+
+		CopyFromSelectedStream(NewCurrentStreamIndex);
+		MergeOrCreateNoteOffsForDeselectedStream(NewCurrentStreamIndex);
+
+		CurrentStreamIndex = NewCurrentStreamIndex;
+	}
+
+	void FMidiStreamSelectOperator::CopyFromSelectedStream(int32 SelectedStreamIndex)
+	{
+		FMidiStreamReadRef SourceStream = SelectedStreamIndex == 0 ? MidiStreamAInPin : MidiStreamBInPin;
+
+		// Copy, but keep track of 'note-on' messages so we can appropriately send 'note-off' messages later...
+		MidiStreamOutPin->Copy(SourceStream,
+			[&](const FMidiStreamEvent& MidiEvent)
+			{
+				if (MidiEvent.MidiMessage.IsStd())
+				{
+					if (MidiEvent.MidiMessage.IsNoteOn())
+					{
+						PlayingVoices.Add(MidiEvent.GetVoiceId());
+					}
+					else
+					{
+						PlayingVoices.Remove(MidiEvent.GetVoiceId());
+					}
+				}
+				return true;
+			},
+			true);
+	}
+
+	void FMidiStreamSelectOperator::MergeOrCreateNoteOffsForDeselectedStream(int32 SelectedStreamIndex)
+	{
+		FMidiStreamReadRef InActiveStream = SelectedStreamIndex == 0 ? MidiStreamBInPin : MidiStreamAInPin;
+
+		if (*ImmediateNoteOffInPin)
+		{
+			// Immediately send note offs - generate NoteOffs with matching voice IDs and send those through
+			if (CurrentStreamIndex != SelectedStreamIndex && !PlayingVoices.IsEmpty())
+			{
+				for (FMidiVoiceId VoiceId : PlayingVoices)
+				{
+					uint8 MidiCh;
+					uint8 MidiNote;
+					VoiceId.GetChannelAndNote(MidiCh, MidiNote);
+					FMidiStreamEvent MidiEvent(0u, FMidiMsg::CreateNoteOff(MidiCh, MidiNote)); // Do we need a note? VoiceId should handle the note off
+					MidiEvent.BlockSampleFrameIndex = 0;
+					MidiEvent.BlockSampleFrameOffset = 0.0f;
+					MidiEvent.AuthoredMidiTick = 0;
+					MidiEvent.CurrentMidiTick = 0;
+					MidiEvent.MsOffset = 0.0f;
+					MidiEvent.TrackIndex = 1;
+					MidiEvent.SetVoiceId(VoiceId);
+					MidiStreamOutPin->InsertNoteOffEventOrCancelPendingNoteOn(MidiEvent);
+				}
+				PlayingVoices.Empty();
+			}
+		}
+		else if (*CopyInactiveNoteOffsInPin)
+		{
+			// Merge note offs from inactive midi stream
+			MidiStreamOutPin->MergeMidiEvents(InActiveStream, [&](const FMidiStreamEvent& MidiEvent)
+				{
+					if (MidiEvent.MidiMessage.IsStd() && MidiEvent.MidiMessage.IsNoteOff() && PlayingVoices.Contains(MidiEvent.GetVoiceId()))
+					{
+						PlayingVoices.Remove(MidiEvent.GetVoiceId());
+						return true;
+					}
+					return false;
+				});
+		}
+	}
+
+	void FMidiStreamSelectOperator::Reset(const FResetParams& ResetParams)
+	{
+		CurrentStreamIndex = -1;
+		PlayingVoices.Empty();
+
+		MidiStreamOutPin->PrepareBlock();
+		
+		if (MidiStreamAInPin.Get() && MidiStreamAInPin->GetMidiClockSource())
+		{
+			MidiStreamOutPin->SetClockSource(*(MidiStreamAInPin->GetMidiClockSource()));
+		}
+		else if (MidiStreamBInPin.Get() && MidiStreamBInPin->GetMidiClockSource())
+		{
+			// Try the B stream if the A stream isn't hooked up
+			MidiStreamOutPin->SetClockSource(*(MidiStreamBInPin->GetMidiClockSource()));
+		}
+	}
+}
+
+#undef LOCTEXT_NAMESPACE // "HarmonixMetaSound"
\ No newline at end of file
diff --git a/Engine/Plugins/Runtime/Harmonix/Source/HarmonixMetasound/Private/HarmonixMetasound/Nodes/MidiStreamTransposerNode.cpp b/Engine/Plugins/Runtime/Harmonix/Source/HarmonixMetasound/Private/HarmonixMetasound/Nodes/MidiStreamTransposerNode.cpp
new file mode 100644
--- /dev/null
+++ b/Engine/Plugins/Runtime/Harmonix/Source/HarmonixMetasound/Private/HarmonixMetasound/Nodes/MidiStreamTransposerNode.cpp
@@ -0,0 +1,243 @@
+// Copyright Epic Games, Inc. All Rights Reserved.
+
+#include "MetasoundExecutableOperator.h"
+#include "MetasoundFacade.h"
+#include "MetasoundNodeInterface.h"
+#include "MetasoundParamHelper.h"
+#include "MetasoundSampleCounter.h"
+#include "MetasoundStandardNodesCategories.h"
+#include "MetasoundVertex.h"
+
+#include "HarmonixMetasound/Common.h"
+#include "HarmonixMetasound/DataTypes/MidiStream.h"
+#include "HarmonixMetasound/DataTypes/MusicTransport.h"
+
+DEFINE_LOG_CATEGORY_STATIC(LogMidiStreamTransposer, Log, All);
+
+#define LOCTEXT_NAMESPACE "HarmonixMetaSound"
+
+namespace HarmonixMetasound
+{
+	using namespace Metasound;
+
+	class FMidiStreamTransposerOperator : public TExecutableOperator<FMidiStreamTransposerOperator>
+	{
+	public:
+		static const FNodeClassMetadata& GetNodeInfo();
+		static const FVertexInterface& GetVertexInterface();
+		static TUniquePtr<IOperator> CreateOperator(const FBuildOperatorParams& InParams, FBuildResults& OutResults);
+
+		FMidiStreamTransposerOperator(const FBuildOperatorParams& InParams,
+								  const FBoolReadRef&       InEnabled,
+								  const FStringReadRef&     InTrackSelect,
+								  const FStringReadRef&     InChannelSelect,
+								  const FInt32ReadRef&      InTransposition,
+								  const FMidiStreamReadRef& InMidiStream);
+
+		virtual void BindInputs(FInputVertexInterfaceData& InVertexData) override;
+		virtual void BindOutputs(FOutputVertexInterfaceData& InVertexData) override;
+		virtual FDataReferenceCollection GetInputs() const override;
+		virtual FDataReferenceCollection GetOutputs() const override;
+
+		void Reset(const FResetParams& ResetParams);
+
+		void Execute();
+
+	private:
+		//** INPUTS
+		FBoolReadRef       EnableInPin;
+		FStringReadRef     TrackSelectInPin;
+		FStringReadRef     ChannelSelectInPin;
+		FInt32ReadRef      TranspositionInPin;
+		FMidiStreamReadRef MidiStreamInPin;
+
+		//** OUTPUTS
+		FMidiStreamWriteRef MidiStreamOutPin;
+
+		//** DATA
+		FString CurrentTrackFilter;
+		FString CurrentMidiChannelFilter;
+		FMidiStreamEventTrackChannelFilter Filter;
+	};
+
+	class FMidiStreamTransposerNode : public FNodeFacade
+	{
+	public:
+		FMidiStreamTransposerNode(const FNodeInitData& InInitData)
+			: FNodeFacade(InInitData.InstanceName, InInitData.InstanceID, TFacadeOperatorClass<FMidiStreamTransposerOperator>())
+		{}
+		virtual ~FMidiStreamTransposerNode() = default;
+	};
+
+	METASOUND_REGISTER_NODE(FMidiStreamTransposerNode)
+		
+	const FNodeClassMetadata& FMidiStreamTransposerOperator::GetNodeInfo()
+	{
+		auto InitNodeInfo = []() -> FNodeClassMetadata
+		{
+			FNodeClassMetadata Info;
+			Info.ClassName        = { HarmonixNodeNamespace, TEXT("MidiStreamTransposer"), TEXT("")};
+			Info.MajorVersion     = 0;
+			Info.MinorVersion     = 1;
+			Info.DisplayName      = METASOUND_LOCTEXT("MidiStreamTransposerNode_DisplayName", "Midi Stream Transposer");
+			Info.Description      = METASOUND_LOCTEXT("MidiStreamTransposerNode_Description", "Duplicates the incoming midi stream to its output with the note on/off messages transposed by the specified number of semitones.");
+			Info.Author           = PluginAuthor;
+			Info.PromptIfMissing  = PluginNodeMissingPrompt;
+			Info.DefaultInterface = GetVertexInterface();
+			Info.CategoryHierarchy.Emplace(NodeCategories::Music);
+			return Info;
+		};
+
+		static const FNodeClassMetadata Info = InitNodeInfo();
+
+		return Info;
+	}
+
+	const FVertexInterface& FMidiStreamTransposerOperator::GetVertexInterface()
+	{
+		using namespace CommonPinNames;
+
+		static const FVertexInterface Interface(
+			FInputVertexInterface(
+				TInputDataVertex<bool>(METASOUND_GET_PARAM_NAME_AND_METADATA(Inputs::Enable), true),
+				TInputDataVertex<FString>(METASOUND_GET_PARAM_NAME_AND_METADATA(Inputs::MidiTrackIndexFilterSpecifier), FString(TEXT("*"))),
+				TInputDataVertex<FString>(METASOUND_GET_PARAM_NAME_AND_METADATA(Inputs::MidiChannelFilterSpecifier), FString(TEXT("*"))),
+				TInputDataVertex<int32>(METASOUND_GET_PARAM_NAME_AND_METADATA(Inputs::Transposition), 0),
+				TInputDataVertex<FMidiStream>(METASOUND_GET_PARAM_NAME_AND_METADATA(Inputs::MidiStream))
+			),
+			FOutputVertexInterface(
+				TOutputDataVertex<FMidiStream>(METASOUND_GET_PARAM_NAME_AND_METADATA(Outputs::MidiStream))
+			)
+		);
+
+		return Interface;
+	}
+
+	TUniquePtr<IOperator> FMidiStreamTransposerOperator::CreateOperator(const FBuildOperatorParams& InParams, FBuildResults& OutResults)
+	{
+		using namespace CommonPinNames;
+
+		const FMidiStreamTransposerNode& LoggerNode = static_cast<const FMidiStreamTransposerNode&>(InParams.Node);
+
+		const FInputVertexInterfaceData& InputData = InParams.InputData;
+		FBoolReadRef InEnabled = InputData.GetOrCreateDefaultDataReadReference<bool>(METASOUND_GET_PARAM_NAME(Inputs::Enable), InParams.OperatorSettings);
+		FStringReadRef InTrackFilter = InputData.GetOrCreateDefaultDataReadReference<FString>(METASOUND_GET_PARAM_NAME(Inputs::MidiTrackIndexFilterSpecifier), InParams.OperatorSettings);
+		FStringReadRef InMidiChannelFilter = InputData.GetOrCreateDefaultDataReadReference<FString>(METASOUND_GET_PARAM_NAME(Inputs::MidiChannelFilterSpecifier), InParams.OperatorSettings);
+		FInt32ReadRef InTransposition = InputData.GetOrCreateDefaultDataReadReference<int32>(METASOUND_GET_PARAM_NAME(Inputs::Transposition), InParams.OperatorSettings);
+		FMidiStreamReadRef InMidiStream = InputData.GetOrConstructDataReadReference<FMidiStream>(METASOUND_GET_PARAM_NAME(Inputs::MidiStream), InParams.OperatorSettings);
+
+		return MakeUnique<FMidiStreamTransposerOperator>(InParams, 
+			InEnabled,
+			InTrackFilter,
+			InMidiChannelFilter,
+			InTransposition,
+			InMidiStream);
+	}
+
+	FMidiStreamTransposerOperator::FMidiStreamTransposerOperator(const FBuildOperatorParams& InParams,
+											   const FBoolReadRef& InEnabled,
+											   const FStringReadRef& InTrackSelect,
+											   const FStringReadRef& InChannelSelect,
+											   const FInt32ReadRef& InTransposition,
+											   const FMidiStreamReadRef& InMidiStream)
+		: EnableInPin(InEnabled)
+		, TrackSelectInPin(InTrackSelect)
+		, ChannelSelectInPin(InChannelSelect)
+		, TranspositionInPin(InTransposition)
+		, MidiStreamInPin(InMidiStream)
+		, MidiStreamOutPin(FMidiStreamWriteRef::CreateNew(InParams.OperatorSettings))
+	{
+		Reset(InParams);
+	}
+
+	void FMidiStreamTransposerOperator::BindInputs(FInputVertexInterfaceData& InVertexData)
+	{
+		using namespace CommonPinNames;
+
+		InVertexData.BindReadVertex(METASOUND_GET_PARAM_NAME(Inputs::Enable), EnableInPin);
+		InVertexData.BindReadVertex(METASOUND_GET_PARAM_NAME(Inputs::MidiTrackIndexFilterSpecifier), TrackSelectInPin);
+		InVertexData.BindReadVertex(METASOUND_GET_PARAM_NAME(Inputs::MidiChannelFilterSpecifier), ChannelSelectInPin);
+		InVertexData.BindReadVertex(METASOUND_GET_PARAM_NAME(Inputs::Transposition), TranspositionInPin);
+		InVertexData.BindReadVertex(METASOUND_GET_PARAM_NAME(Inputs::MidiStream), MidiStreamInPin);
+	}
+
+	void FMidiStreamTransposerOperator::BindOutputs(FOutputVertexInterfaceData& InVertexData)
+	{
+		using namespace CommonPinNames;
+
+		InVertexData.BindReadVertex(METASOUND_GET_PARAM_NAME(Outputs::MidiStream), MidiStreamOutPin);
+	}
+
+	FDataReferenceCollection FMidiStreamTransposerOperator::GetInputs() const
+	{
+		checkNoEntry();
+		return {};
+	}
+
+	FDataReferenceCollection FMidiStreamTransposerOperator::GetOutputs() const
+	{
+		checkNoEntry();
+		return {};
+	}
+
+	void FMidiStreamTransposerOperator::Reset(const FResetParams& ResetParams)
+	{
+		MidiStreamOutPin->PrepareBlock();
+		
+		if (const FMidiClockReadRef* Clock = MidiStreamInPin->GetMidiClockSource())
+		{
+			MidiStreamOutPin->SetClockSource(*Clock);
+		}
+		CurrentTrackFilter = *TrackSelectInPin;
+		CurrentMidiChannelFilter = *ChannelSelectInPin;
+		
+		FString ParseTrackFilterErrorMessage;
+		if (!Filter.SetTrackFilterFromString(CurrentTrackFilter, ParseTrackFilterErrorMessage))
+		{
+			UE_LOG(LogMidiStreamTransposer, Warning, TEXT("Error parsing Track Filter string: %s"), *ParseTrackFilterErrorMessage);
+		}
+
+		FString ParseMidiChannelFilterErrorMessage;
+		if (!Filter.SetMidiChannelFilterFromString(CurrentMidiChannelFilter, ParseMidiChannelFilterErrorMessage))
+		{
+			UE_LOG(LogMidiStreamTransposer, Warning, TEXT("Error parsing Channel Filter string: %s"), *ParseMidiChannelFilterErrorMessage);
+		}
+	}
+
+	void FMidiStreamTransposerOperator::Execute()
+	{
+		if (CurrentTrackFilter != *TrackSelectInPin)
+		{
+			CurrentTrackFilter = *TrackSelectInPin;
+			FString ParseTrackFilterErrorMessage;
+			if (!Filter.SetTrackFilterFromString(CurrentTrackFilter, ParseTrackFilterErrorMessage))
+			{
+				UE_LOG(LogMidiStreamTransposer, Warning, TEXT("Error parsing Track Filter string: %s"), *ParseTrackFilterErrorMessage);
+			}
+		}
+
+		if (CurrentMidiChannelFilter != *ChannelSelectInPin)
+		{
+			CurrentMidiChannelFilter = *ChannelSelectInPin;
+			FString ParseMidiChannelFilterErrorMessage;
+			if (!Filter.SetMidiChannelFilterFromString(CurrentMidiChannelFilter, ParseMidiChannelFilterErrorMessage))
+			{
+				UE_LOG(LogMidiStreamTransposer, Warning, TEXT("Error parsing Channel Filter string: %s"), *ParseMidiChannelFilterErrorMessage);
+			}
+		}
+
+		auto Transformer = [&](const FMidiStreamEvent& InEvent) -> FMidiStreamEvent
+			{
+				FMidiStreamEvent NewEvent = InEvent;
+				if (Filter(NewEvent) && (NewEvent.MidiMessage.IsNoteOn() || NewEvent.MidiMessage.IsNoteOff()))
+				{
+					const int32 NewNote = FMath::Clamp(*TranspositionInPin + NewEvent.MidiMessage.Data1, 0, 127);
+					NewEvent.MidiMessage.Data1 = NewNote;
+				}
+				return NewEvent;
+			};
+		MidiStreamOutPin->Copy(MidiStreamInPin, [](const FMidiStreamEvent&){return true;}, Transformer);
+	}
+}
+
+#undef LOCTEXT_NAMESPACE // "HarmonixMetaSound"
