[Harmonix] Move Harmonix plugin into Engine/Plugins/Runtime to make it available to licensees in UE 5.4

diff --git a/Engine/Plugins/Runtime/Harmonix/Source/HarmonixMetasound/Private/HarmonixMetasound/Nodes/MidiTextTriggerNode.cpp b/Engine/Plugins/Runtime/Harmonix/Source/HarmonixMetasound/Private/HarmonixMetasound/Nodes/MidiTextTriggerNode.cpp
new file mode 100644
--- /dev/null
+++ b/Engine/Plugins/Runtime/Harmonix/Source/HarmonixMetasound/Private/HarmonixMetasound/Nodes/MidiTextTriggerNode.cpp
@@ -0,0 +1,200 @@
+// Copyright Epic Games, Inc. All Rights Reserved.
+
+#include "MetasoundExecutableOperator.h"
+#include "MetasoundFacade.h"
+#include "MetasoundNodeInterface.h"
+#include "MetasoundParamHelper.h"
+#include "MetasoundSampleCounter.h"
+#include "MetasoundStandardNodesCategories.h"
+#include "MetasoundVertex.h"
+#include "MetasoundTrigger.h"
+
+#include "HarmonixMetasound/Common.h"
+#include "HarmonixMetasound/DataTypes/MidiStream.h"
+#include "HarmonixMetasound/DataTypes/MusicTransport.h"
+
+DEFINE_LOG_CATEGORY_STATIC(LogMidiTextTrigger, Log, All);
+
+#define LOCTEXT_NAMESPACE "HarmonixMetaSound"
+
+namespace HarmonixMetasound
+{
+	using namespace Metasound;
+
+	class FMidiTextTriggerOperator : public TExecutableOperator<FMidiTextTriggerOperator>
+	{
+	public:
+		static const FNodeClassMetadata& GetNodeInfo();
+		static const FVertexInterface& GetVertexInterface();
+		static TUniquePtr<IOperator> CreateOperator(const FBuildOperatorParams& InParams, FBuildResults& OutResults);
+
+		FMidiTextTriggerOperator(const FBuildOperatorParams& InParams,
+								 const FBoolReadRef&       InEnabled,
+								 const FMidiStreamReadRef& InMidiStream,
+								 const FInt32ReadRef&      InTrackNumber,
+								 const FStringReadRef&     InText);
+
+		virtual void BindInputs(FInputVertexInterfaceData& InVertexData) override;
+		virtual void BindOutputs(FOutputVertexInterfaceData& InVertexData) override;
+
+		void Reset(const FResetParams& ResetParams);
+		
+		void Execute();
+
+	private:
+		//** INPUTS
+		FBoolReadRef       EnableInPin;
+		FMidiStreamReadRef MidiStreamInPin;
+		FInt32ReadRef      TrackNumberInPin;
+		FStringReadRef     TextInPin;
+
+		//** OUTPUTS
+		FTriggerWriteRef TriggerOutPin;
+	};
+
+	class FMidiTextTriggerNode : public FNodeFacade
+	{
+	public:
+		FMidiTextTriggerNode(const FNodeInitData& InInitData)
+			: FNodeFacade(InInitData.InstanceName, InInitData.InstanceID, TFacadeOperatorClass<FMidiTextTriggerOperator>())
+		{}
+		virtual ~FMidiTextTriggerNode() = default;
+	};
+
+	METASOUND_REGISTER_NODE(FMidiTextTriggerNode)
+
+	const FNodeClassMetadata& FMidiTextTriggerOperator::GetNodeInfo()
+	{
+		auto InitNodeInfo = []() -> FNodeClassMetadata
+		{
+			FNodeClassMetadata Info;
+			Info.ClassName        = { HarmonixNodeNamespace, TEXT("MidiTextTrigger"), TEXT("")};
+			Info.MajorVersion     = 0;
+			Info.MinorVersion     = 1;
+			Info.DisplayName      = METASOUND_LOCTEXT("MidiTextTriggerNode_DisplayName", "Midi Text Trigger");
+			Info.Description      = METASOUND_LOCTEXT("MidiTextTriggerNode_Description", "Receives a midi stream, filters for the desired messages, and outputs triggers.");
+			Info.Author           = PluginAuthor;
+			Info.PromptIfMissing  = PluginNodeMissingPrompt;
+			Info.DefaultInterface = GetVertexInterface();
+			Info.CategoryHierarchy.Emplace(NodeCategories::Music);
+			return Info;
+		};
+
+		static const FNodeClassMetadata Info = InitNodeInfo();
+
+		return Info;
+	}
+
+	namespace MidiTextTriggerPinNames
+	{
+		METASOUND_PARAM(TextInput, "Text", "String of characters to look for.")
+		METASOUND_PARAM(TriggerOutput, "Trigger Out", "A trigger when the text is encountered.")
+	}
+
+	const FVertexInterface& FMidiTextTriggerOperator::GetVertexInterface()
+	{
+		using namespace MidiTextTriggerPinNames;
+		using namespace CommonPinNames;
+
+		static const FVertexInterface Interface(
+			FInputVertexInterface(
+				TInputDataVertex<bool>(METASOUND_GET_PARAM_NAME_AND_METADATA(Inputs::Enable), true),
+				TInputDataVertex<FMidiStream>(METASOUND_GET_PARAM_NAME_AND_METADATA(Inputs::MidiStream)),
+				TInputDataVertex<int32>(METASOUND_GET_PARAM_NAME_AND_METADATA(Inputs::MidiTrackNumber),1),
+				TInputDataVertex<FString>(METASOUND_GET_PARAM_NAME_AND_METADATA(TextInput))
+			),
+			FOutputVertexInterface(
+				TOutputDataVertex<FTrigger>(METASOUND_GET_PARAM_NAME_AND_METADATA(TriggerOutput))
+				)
+			);
+
+		return Interface;
+	}
+
+	TUniquePtr<IOperator> FMidiTextTriggerOperator::CreateOperator(const FBuildOperatorParams& InParams, FBuildResults& OutResults)
+	{
+		using namespace MidiTextTriggerPinNames;
+		using namespace CommonPinNames;
+
+		const FMidiTextTriggerNode& LoggerNode = static_cast<const FMidiTextTriggerNode&>(InParams.Node);
+
+		const FInputVertexInterfaceData& InputData = InParams.InputData;
+		FBoolReadRef InEnabled          = InputData.GetOrCreateDefaultDataReadReference<bool>(METASOUND_GET_PARAM_NAME(Inputs::Enable), InParams.OperatorSettings);
+		FMidiStreamReadRef InMidiStream = InputData.GetOrConstructDataReadReference<FMidiStream>(METASOUND_GET_PARAM_NAME(Inputs::MidiStream), InParams.OperatorSettings);
+		FInt32ReadRef InTrackNumber     = InputData.GetOrCreateDefaultDataReadReference<int32>(METASOUND_GET_PARAM_NAME(Inputs::MidiTrackNumber), InParams.OperatorSettings);
+		FStringReadRef InText           = InputData.GetOrCreateDefaultDataReadReference<FString>(METASOUND_GET_PARAM_NAME(TextInput), InParams.OperatorSettings);
+
+		return MakeUnique<FMidiTextTriggerOperator>(InParams, InEnabled, InMidiStream, InTrackNumber, InText);
+	}
+
+	FMidiTextTriggerOperator::FMidiTextTriggerOperator(const FBuildOperatorParams& InParams,
+													   const FBoolReadRef&       InEnabled,
+													   const FMidiStreamReadRef& InMidiStream,
+													   const FInt32ReadRef&      InTrackNumber,
+													   const FStringReadRef&     InText)
+		: EnableInPin(InEnabled)
+		, MidiStreamInPin(InMidiStream)
+		, TrackNumberInPin(InTrackNumber)
+		, TextInPin(InText)
+		, TriggerOutPin(FTriggerWriteRef::CreateNew(InParams.OperatorSettings))
+	{
+		Reset(InParams);
+	}
+
+	void FMidiTextTriggerOperator::BindInputs(FInputVertexInterfaceData& InVertexData)
+	{
+		using namespace MidiTextTriggerPinNames;
+		using namespace CommonPinNames;
+		InVertexData.BindReadVertex(METASOUND_GET_PARAM_NAME(Inputs::Enable), EnableInPin);
+		InVertexData.BindReadVertex(METASOUND_GET_PARAM_NAME(Inputs::MidiStream), MidiStreamInPin);
+		InVertexData.BindReadVertex(METASOUND_GET_PARAM_NAME(Inputs::MidiTrackNumber), TrackNumberInPin);
+		InVertexData.BindReadVertex(METASOUND_GET_PARAM_NAME(TextInput), TextInPin);
+	}
+
+	void FMidiTextTriggerOperator::BindOutputs(FOutputVertexInterfaceData& InVertexData)
+	{
+		using namespace MidiTextTriggerPinNames;
+		using namespace CommonPinNames;
+		InVertexData.BindReadVertex(METASOUND_GET_PARAM_NAME(TriggerOutput), TriggerOutPin);
+	}
+
+	void FMidiTextTriggerOperator::Reset(const FResetParams& ResetParams)
+	{
+		TriggerOutPin->Reset();
+	}
+
+	void FMidiTextTriggerOperator::Execute()
+	{
+		TriggerOutPin->AdvanceBlock();
+
+		if (!*EnableInPin)
+		{
+			return;
+		}
+
+		int32 TrackIndex = *TrackNumberInPin;
+		if (TrackIndex >= 0)
+		{
+			const TArray<FMidiStreamEvent>& MidiEvents = MidiStreamInPin->GetEventsInBlock();
+			for (auto& Event : MidiEvents)
+			{
+				bool IsText = Event.MidiMessage.IsText();
+
+				if (Event.TrackIndex != TrackIndex || !IsText)
+				{
+					continue;
+				}
+
+				const FString* EventString = MidiStreamInPin->GetMidiTrackText(TrackIndex, Event.MidiMessage.GetTextIndex());
+				if (EventString && EventString->Compare(*TextInPin) == 0)
+				{
+					TriggerOutPin->TriggerFrame(Event.BlockSampleFrameIndex);
+				}
+			}
+		}
+	
+	}
+}
+
+#undef LOCTEXT_NAMESPACE // "HarmonixMetaSound"
+
diff --git a/Engine/Plugins/Runtime/Harmonix/Source/HarmonixMetasound/Private/HarmonixMetasound/Nodes/MidiVoiceManagerNode.cpp b/Engine/Plugins/Runtime/Harmonix/Source/HarmonixMetasound/Private/HarmonixMetasound/Nodes/MidiVoiceManagerNode.cpp
new file mode 100644
--- /dev/null
+++ b/Engine/Plugins/Runtime/Harmonix/Source/HarmonixMetasound/Private/HarmonixMetasound/Nodes/MidiVoiceManagerNode.cpp
@@ -0,0 +1,249 @@
+// Copyright Epic Games, Inc. All Rights Reserved.
+
+#include "MetasoundExecutableOperator.h"
+#include "MetasoundFacade.h"
+#include "MetasoundNodeInterface.h"
+#include "MetasoundParamHelper.h"
+#include "MetasoundSampleCounter.h"
+#include "MetasoundStandardNodesCategories.h"
+#include "MetasoundVertex.h"
+#include "MetasoundTrigger.h"
+
+#include "HarmonixMetasound/Common.h"
+#include "HarmonixMetasound/DataTypes/MidiStream.h"
+#include "HarmonixMetasound/DataTypes/MusicTransport.h"
+#include "HarmonixMidi/MidiMsg.h"
+
+DEFINE_LOG_CATEGORY_STATIC(LogMidiVoiceManager, Log, All);
+
+#define LOCTEXT_NAMESPACE "HarmonixMetaSound"
+
+namespace HarmonixMetasound
+{
+	using namespace Metasound;
+
+	class FMidiVoiceManagerOperator : public TExecutableOperator<FMidiVoiceManagerOperator>
+	{
+	public:
+		static const FNodeClassMetadata& GetNodeInfo();
+		static const FVertexInterface& GetVertexInterface();
+		static TUniquePtr<IOperator> CreateOperator(const FBuildOperatorParams& InParams, FBuildResults& OutResults);
+
+		FMidiVoiceManagerOperator(
+			const FBuildOperatorParams& InParams,
+			FMidiStreamReadRef&& InMidiStream,
+			int32 InNumVoices);
+
+		virtual void BindInputs(FInputVertexInterfaceData& InVertexData) override;
+		virtual void BindOutputs(FOutputVertexInterfaceData& InVertexData) override;
+		virtual FDataReferenceCollection GetInputs() const override;
+		virtual FDataReferenceCollection GetOutputs() const override;
+
+		void Reset(const FResetParams& ResetParams);
+		
+		void Execute();
+
+	private:
+		int32 GetNumVoices() const;
+
+		static constexpr int32 MaxVoices = 8;
+
+		//** INPUTS
+		FMidiStreamReadRef MidiStreamInPin;
+		int32 NumVoices;
+
+		//** OUTPUTS
+		TArray<FMidiStreamWriteRef> MidiOutPin;
+
+		//** DATA
+		int32 LastVoiceIndex = -1;
+		FMidiVoiceId VoiceIds[MaxVoices];
+	};
+
+	class FMidiVoiceManagerNode : public FNodeFacade
+	{
+	public:
+		FMidiVoiceManagerNode(const FNodeInitData& InInitData)
+			: FNodeFacade(InInitData.InstanceName, InInitData.InstanceID, TFacadeOperatorClass<FMidiVoiceManagerOperator>())
+		{}
+		virtual ~FMidiVoiceManagerNode() = default;
+	};
+
+	METASOUND_REGISTER_NODE(FMidiVoiceManagerNode)
+
+	const FNodeClassMetadata& FMidiVoiceManagerOperator::GetNodeInfo()
+	{
+		auto InitNodeInfo = []() -> FNodeClassMetadata
+		{
+			FNodeClassMetadata Info;
+			Info.ClassName        = { HarmonixNodeNamespace, TEXT("MidiVoiceManager"), TEXT("")};
+			Info.MajorVersion     = 0;
+			Info.MinorVersion     = 1;
+			Info.DisplayName      = METASOUND_LOCTEXT("MidiVoiceManagerNode_DisplayName", "Midi Voice Manager");
+			Info.Description      = METASOUND_LOCTEXT("MidiVoiceManagerNode_Description", "Receives a midi stream, and splits it into multiple streams, for voice polyphony.");
+			Info.Author           = PluginAuthor;
+			Info.PromptIfMissing  = PluginNodeMissingPrompt;
+			Info.DefaultInterface = GetVertexInterface();
+			Info.CategoryHierarchy.Emplace(NodeCategories::Music);
+			return Info;
+		};
+
+		static const FNodeClassMetadata Info = InitNodeInfo();
+
+		return Info;
+	}
+
+	namespace MidiVoiceManagerPinNames
+	{
+		METASOUND_PARAM(NumVoices, "Num Voices", "Number of voices to cycle through.")
+		METASOUND_PARAM(MidiStreamVoice, "Voice {0}", "Midi Stream for voice {0}.");
+	}
+
+	const FVertexInterface& FMidiVoiceManagerOperator::GetVertexInterface()
+	{
+		using namespace CommonPinNames;
+		using namespace MidiVoiceManagerPinNames;
+
+		auto CreateVertexInterface = []() -> FVertexInterface
+		{
+			FInputVertexInterface InputInterface;
+			InputInterface.Add(TInputDataVertex<FMidiStream>(METASOUND_GET_PARAM_NAME_AND_METADATA(Inputs::MidiStream)));
+			InputInterface.Add(TInputConstructorVertex<int32>(METASOUND_GET_PARAM_NAME_AND_METADATA(NumVoices), MaxVoices));
+
+			FOutputVertexInterface OutputInterface;
+			for (uint32 i = 0; i < MaxVoices; ++i)
+			{
+				OutputInterface.Add(TOutputDataVertex<FMidiStream>(METASOUND_GET_PARAM_NAME_WITH_INDEX_AND_METADATA(MidiStreamVoice, i)));
+			}
+			return FVertexInterface(InputInterface, OutputInterface);
+		};
+
+		static const FVertexInterface Interface = CreateVertexInterface();
+
+		return Interface;
+	}
+
+	TUniquePtr<IOperator> FMidiVoiceManagerOperator::CreateOperator(const FBuildOperatorParams& InParams, FBuildResults& OutResults)
+	{
+		using namespace CommonPinNames;
+		using namespace MidiVoiceManagerPinNames;
+
+		const FMidiVoiceManagerNode& LoggerNode = static_cast<const FMidiVoiceManagerNode&>(InParams.Node);
+
+		const FInputVertexInterfaceData& InputData = InParams.InputData;
+		FMidiStreamReadRef InMidiStream = InputData.GetOrConstructDataReadReference<FMidiStream>(
+			METASOUND_GET_PARAM_NAME(Inputs::MidiStream),
+			InParams.OperatorSettings);
+		const TDataReadReference<int32> InNumVoices = InputData.GetOrCreateDefaultDataReadReference<int32>(
+			METASOUND_GET_PARAM_NAME(NumVoices),
+			InParams.OperatorSettings);
+
+		return MakeUnique<FMidiVoiceManagerOperator>(InParams, MoveTemp(InMidiStream), *InNumVoices);
+	}
+
+	FMidiVoiceManagerOperator::FMidiVoiceManagerOperator(
+		const FBuildOperatorParams& InParams,
+		FMidiStreamReadRef&& InMidiStream,
+		const int32 InNumVoices)
+	: MidiStreamInPin(MoveTemp(InMidiStream))
+	, NumVoices(InNumVoices)
+	{
+		for (uint32 i = 0; i < MaxVoices; ++i)
+		{
+			MidiOutPin.Add(FMidiStreamWriteRef::CreateNew(InParams.OperatorSettings));
+		}
+
+		Reset(InParams);
+	}
+
+	void FMidiVoiceManagerOperator::BindInputs(FInputVertexInterfaceData& InVertexData)
+	{
+		using namespace CommonPinNames;
+		using namespace MidiVoiceManagerPinNames;
+
+		InVertexData.BindReadVertex(METASOUND_GET_PARAM_NAME(Inputs::MidiStream), MidiStreamInPin);
+		InVertexData.SetValue(METASOUND_GET_PARAM_NAME(NumVoices), NumVoices);
+	}
+
+	void FMidiVoiceManagerOperator::BindOutputs(FOutputVertexInterfaceData& InVertexData)
+	{
+		using namespace MidiVoiceManagerPinNames;
+
+		for (uint32 i = 0; i < MaxVoices; ++i)
+		{
+			InVertexData.BindReadVertex(METASOUND_GET_PARAM_NAME_WITH_INDEX(MidiStreamVoice, i), MidiOutPin[i]);
+		}
+	}
+	
+	FDataReferenceCollection FMidiVoiceManagerOperator::GetInputs() const
+	{
+		// This should never be called. Bind(...) is called instead. This method
+		// exists as a stop-gap until the API can be deprecated and removed.
+		checkNoEntry();
+		return {};
+	}
+
+	FDataReferenceCollection FMidiVoiceManagerOperator::GetOutputs() const
+	{
+		// This should never be called. Bind(...) is called instead. This method
+		// exists as a stop-gap until the API can be deprecated and removed.
+		checkNoEntry();
+		return {};
+	}
+
+	void FMidiVoiceManagerOperator::Reset(const FResetParams& ResetParams)
+	{
+		// Reset the MIDI streams
+		for (FMidiStreamWriteRef& MidiOut : MidiOutPin)
+		{
+			MidiOut->PrepareBlock();
+		}
+
+		// Reset the voices
+		for (int32 i = 0; i < MaxVoices; ++i)
+		{
+			VoiceIds[i] = FMidiVoiceId::None();
+		}
+
+		LastVoiceIndex = -1;
+	}
+
+	int32 FMidiVoiceManagerOperator::GetNumVoices() const
+	{
+		return FMath::Clamp(NumVoices, 1, MaxVoices);
+	}
+
+	void FMidiVoiceManagerOperator::Execute()
+	{
+		for (int32 i = 0; i < GetNumVoices(); ++i)
+		{
+			MidiOutPin[i]->PrepareBlock();
+		}
+
+		const TArray<FMidiStreamEvent>& MidiEvents = MidiStreamInPin->GetEventsInBlock();
+		for (auto& Event : MidiEvents)
+		{
+			if (Event.MidiMessage.IsNoteOn()) 
+			{
+				int32 voiceIndex = (LastVoiceIndex + 1) % GetNumVoices();
+				MidiOutPin[voiceIndex]->AddMidiEvent(Event);
+				VoiceIds[voiceIndex] = Event.GetVoiceId();
+				LastVoiceIndex = voiceIndex;
+			}
+
+			if (Event.MidiMessage.IsNoteOff())
+			{
+				for (int32 voiceIndex = 0; voiceIndex < GetNumVoices(); ++voiceIndex)
+				{
+					if (VoiceIds[voiceIndex] == Event.GetVoiceId())
+					{
+						MidiOutPin[voiceIndex]->AddMidiEvent(Event);
+					}
+				}
+			}
+		}
+	}
+}
+
+#undef LOCTEXT_NAMESPACE // "HarmonixMetaSound"
+
diff --git a/Engine/Plugins/Runtime/Harmonix/Source/HarmonixMetasound/Private/HarmonixMetasound/Nodes/StepSequencePlayerNode.cpp b/Engine/Plugins/Runtime/Harmonix/Source/HarmonixMetasound/Private/HarmonixMetasound/Nodes/StepSequencePlayerNode.cpp
new file mode 100644
--- /dev/null
+++ b/Engine/Plugins/Runtime/Harmonix/Source/HarmonixMetasound/Private/HarmonixMetasound/Nodes/StepSequencePlayerNode.cpp
@@ -0,0 +1,917 @@
+// Copyright Epic Games, Inc. All Rights Reserved.
+
+#include "MetasoundExecutableOperator.h"
+#include "MetasoundFacade.h"
+#include "MetasoundNodeInterface.h"
+#include "MetasoundParamHelper.h"
+#include "MetasoundSampleCounter.h"
+#include "MetasoundStandardNodesCategories.h"
+#include "MetasoundVertex.h"
+
+#include "IAudioProxyInitializer.h"
+#include "Harmonix/AudioRenderableProxy.h"
+
+#include "HarmonixMetasound/Common.h"
+#include "HarmonixMetasound/DataTypes/MidiStream.h"
+#include "HarmonixMetasound/DataTypes/MidiClock.h"
+#include "HarmonixMetasound/DataTypes/MidiStepSequence.h"
+#include "HarmonixMetasound/DataTypes/MusicTransport.h"
+
+#include "HarmonixMidi/MidiPlayCursor.h"
+#include "HarmonixMidi/MidiVoiceId.h"
+
+#define LOCTEXT_NAMESPACE "HarmonixMetaSound"
+
+DEFINE_LOG_CATEGORY_STATIC(LogStepSequencePlayer, Log, All);
+
+namespace HarmonixMetasound
+{
+	using namespace Metasound;
+	using namespace Harmonix;
+
+	class FStepSequencePlayerOperator : public TExecutableOperator<FStepSequencePlayerOperator>, public FMusicTransportControllable, public FMidiPlayCursor, public FMidiVoiceGeneratorBase
+	{
+	public:
+		static const FNodeClassMetadata& GetNodeInfo();
+		static const FVertexInterface&   GetVertexInterface();
+		static TUniquePtr<IOperator>     CreateOperator(const FBuildOperatorParams& InParams, FBuildResults& OutResults);
+
+		FStepSequencePlayerOperator(const FBuildOperatorParams& InParams, 
+							        const FMidiStepSequenceAssetReadRef& InSequenceAsset,
+							        const FMusicTransportEventStreamReadRef& InTransport,
+							        const FMidiClockReadRef& InMidiClockSource,
+							        const FFloatReadRef& InSpeedMultiplier,
+									const FFloatReadRef& InVelocityMultiplier,
+									const FFloatReadRef& InMaxColumns,
+									const FFloatReadRef& InAdditionalOctaves,
+									const FFloatReadRef& InStepSizeQuarterNotes,
+									const FFloatReadRef& InActivePage,
+									const FBoolReadRef& InAutoPage,
+									const FBoolReadRef& InAutoPagePlaysBlankPages,
+									const FBoolReadRef& InLoop,
+									const FBoolReadRef& InEnabled);
+
+		virtual void BindInputs(FInputVertexInterfaceData& InVertexData) override;
+		virtual void BindOutputs(FOutputVertexInterfaceData& InVertexData) override;
+		virtual FDataReferenceCollection GetInputs() const override;
+		virtual FDataReferenceCollection GetOutputs() const override;
+		void Reset(const FResetParams& ResetParams);
+		void Execute();
+
+	private:
+		void Init();
+		
+		//** INPUTS
+		FMidiStepSequenceAssetReadRef SequenceAssetInPin;
+		FMusicTransportEventStreamReadRef TransportInPin;
+		FFloatReadRef SpeedMultInPin;
+		FFloatReadRef VelocityMultInPin;
+		FMidiClockReadRef MidiClockInPin;
+		FFloatReadRef MaxColumnsInPin;
+		FFloatReadRef AdditionalOctavesInPin;
+		FFloatReadRef StepSizeQuarterNotesInPin;
+		FFloatReadRef ActivePageInPin;
+		FBoolReadRef AutoPageInPin;
+		FBoolReadRef AutoPagePlaysBlankPagesInPin;
+		FBoolReadRef LoopInPin;
+		FBoolReadRef EnabledInPin;
+
+		//** OUTPUTS
+		FMidiStreamWriteRef MidiOutPin;
+
+		//** DATA
+		TSharedAudioRenderableDataPtr<FStepSequenceTable, TRefCountedAudioRenderableWithQueuedChanges<FStepSequenceTable>> SequenceTable;
+		FSampleCount BlockSize			= 0;
+		int32 CurrentBlockSpanStart		= 0;
+		int32 CurrentPageIndex			= -1;
+		int32 CurrentCellIndex			= -1;
+		int32 ProcessedThruTick			= -1;
+		int32 SequenceStartTick			= -1;
+		int32 CurrentStepSkipIndex		= 0;
+		bool  bNeedsMidiOutClear		= true;
+		bool  bAutoPage					= false;
+		bool  bPreviousAutoPage			= false;
+		bool  bAutoPagePlaysBlankPages	= false;
+		bool  bLoop						= true;
+		bool  bEnabled					= true;
+		bool  bPlaying					= false;
+		bool  bNeedsRebase				= false;
+		TArray<FMidiVoiceId> CurrentCellNotes;
+
+		void CheckForUpdatedSequenceTable();
+		void ResizeCellStatesForTable();
+		void AllNotesOff(int32 AtFrameIndex, int32 AbsMidiTick, bool ResetCellIndex);
+		int32 CalculateAutoPageIndex(int32 CurTick, int32 TableTickLength, bool bRound) const;
+		int32 CalculatePagesProgressed(const int FromTick, const int ToTick, int32 TableTickLength, const bool bRound) const;
+		void CalculatePageProperties(const FStepSequencePage& Page, int32 MaxColumns, float StepSizeQuarterNotes, int32& OutColumns, int32& OutTicksPerCell, int32& OutTableTickLength) const;
+		void RebaseSequenceStartTickForLoop(const int32 CurTick, const int32 TableTickLength);
+		void EnsureCurrentPageIndexIsValid();
+
+	protected:
+		//** BEGIN FMidiPlayCursor
+		virtual void SeekToTick(int32 Tick) override;
+		virtual void SeekThruTick(int32 Tick) override;
+		virtual void AdvanceThruTick(int32 Tick, bool IsPreRoll) override;
+		// We have to override this to disambiguate the FMidiPlayCursor method and the MS operator Reset above
+		virtual void Reset(const bool ForceNoBroadcast) override { FMidiPlayCursor::Reset(ForceNoBroadcast); }
+		//** END FMidiPlayCursor
+	};
+
+	class FStepSequencePlayerNode : public FNodeFacade
+	{
+	public:
+		FStepSequencePlayerNode(const FNodeInitData& InInitData)
+			: FNodeFacade(InInitData.InstanceName, InInitData.InstanceID, TFacadeOperatorClass<FStepSequencePlayerOperator>())
+		{}
+		virtual ~FStepSequencePlayerNode() = default;
+	};
+
+	METASOUND_REGISTER_NODE(FStepSequencePlayerNode)
+
+	const FNodeClassMetadata& FStepSequencePlayerOperator::GetNodeInfo()
+	{
+		auto InitNodeInfo = []() -> FNodeClassMetadata
+		{
+			FNodeClassMetadata Info;
+			Info.ClassName        = { HarmonixNodeNamespace, TEXT("StepSequencePlayer"), TEXT("") };
+			Info.MajorVersion     = 0;
+			Info.MinorVersion     = 1;
+			Info.DisplayName      = METASOUND_LOCTEXT("StepSequencePlayerNode_DisplayName", "Step Sequence Player");
+			Info.Description      = METASOUND_LOCTEXT("StepSequencePlayerNode_Description", "Plays a Step Sequence Asset.");
+			Info.Author           = PluginAuthor;
+			Info.PromptIfMissing  = PluginNodeMissingPrompt;
+			Info.DefaultInterface = GetVertexInterface();
+			Info.CategoryHierarchy.Emplace(NodeCategories::Music);
+
+			return Info;
+		};
+
+		static const FNodeClassMetadata Info = InitNodeInfo();
+
+		return Info;
+	}
+
+	namespace StepSequencePlayerPinNames
+	{
+		METASOUND_PARAM(InputSequenceAsset, "Step Sequence Asset", "Step sequence to play.");
+		METASOUND_PARAM(InputVelocityMultiplier, "Velocity Multiplier", "Multiplies the current note velocity by this number");
+		METASOUND_PARAM(InputMaxColumns, "Max Columns", "The Maximinum Number of cells to play per step sequence row.");
+		METASOUND_PARAM(InputAdditionalOctaves, "Additional Octaves", "The number of octaves to add to the authored step sequence note.");
+		METASOUND_PARAM(InputStepSizeQuarterNotes, "Step Size Quarter Notes", "The size, in quarter notes, of each step");
+		METASOUND_PARAM(InputActivePage, "Active Page", "The page of the step sequence to play (1 indexed)");
+		METASOUND_PARAM(InputAutoPage, "Auto Page", "Whether to calculate the page of the step sequence based on current position");
+		METASOUND_PARAM(InputAutoPagePlaysBlankPages, "Auto Page Plays Blank Pages", "If autopaging, should blank pages be played?");
+		METASOUND_PARAM(InputLoop, "Loop", "If not looping, the sequence will be played as a one-shot starting on the next stepsizebeats interval (autopaging if it's enabled)");
+		METASOUND_PARAM(InputEnabled, "Enabled", "Whether the sequencer is enabled. If loop is off, enabling will trigger to start on next stepsizebeats interval");
+	}
+
+	const FVertexInterface& FStepSequencePlayerOperator::GetVertexInterface()
+	{
+		using namespace StepSequencePlayerPinNames;
+		using namespace CommonPinNames;
+
+		static const FVertexInterface Interface(
+			FInputVertexInterface(
+				TInputDataVertex<FMidiStepSequenceAsset>(METASOUND_GET_PARAM_NAME_AND_METADATA(InputSequenceAsset)),
+				TInputDataVertex<FMusicTransportEventStream>(METASOUND_GET_PARAM_NAME_AND_METADATA(Inputs::Transport)),
+				TInputDataVertex<FMidiClock>(METASOUND_GET_PARAM_NAME_AND_METADATA(Inputs::MidiClock)),
+				TInputDataVertex<float>(METASOUND_GET_PARAM_NAME_AND_METADATA(Inputs::Speed), 1.0f),
+				TInputDataVertex<float>(METASOUND_GET_PARAM_NAME_AND_METADATA(InputVelocityMultiplier), 1.0f),
+				TInputDataVertex<float>(METASOUND_GET_PARAM_NAME_AND_METADATA(InputMaxColumns), 64.0f),
+				TInputDataVertex<float>(METASOUND_GET_PARAM_NAME_AND_METADATA(InputAdditionalOctaves), 0),
+				TInputDataVertex<float>(METASOUND_GET_PARAM_NAME_AND_METADATA(InputStepSizeQuarterNotes), 0.25f),
+				TInputDataVertex<float>(METASOUND_GET_PARAM_NAME_AND_METADATA(InputActivePage), 0),
+				TInputDataVertex<bool>(METASOUND_GET_PARAM_NAME_AND_METADATA(InputAutoPage), false),
+				TInputDataVertex<bool>(METASOUND_GET_PARAM_NAME_AND_METADATA(InputAutoPagePlaysBlankPages), true),
+				TInputDataVertex<bool>(METASOUND_GET_PARAM_NAME_AND_METADATA(InputLoop), true),
+				TInputDataVertex<bool>(METASOUND_GET_PARAM_NAME_AND_METADATA(InputEnabled), true)
+			),
+			FOutputVertexInterface(
+				TOutputDataVertex<FMidiStream>(METASOUND_GET_PARAM_NAME_AND_METADATA(Outputs::MidiStream))
+			)
+		);
+		return Interface;
+	}
+
+	TUniquePtr<IOperator> FStepSequencePlayerOperator::CreateOperator(const FBuildOperatorParams& InParams, FBuildResults& OutResults)
+	{
+		using namespace StepSequencePlayerPinNames;
+		using namespace CommonPinNames;
+
+		const FStepSequencePlayerNode& PlayerNode = static_cast<const FStepSequencePlayerNode&>(InParams.Node);
+
+		const FInputVertexInterfaceData& InputData = InParams.InputData;
+
+		FMidiStepSequenceAssetReadRef InSequenceAsset = InputData.GetOrConstructDataReadReference<FMidiStepSequenceAsset>(METASOUND_GET_PARAM_NAME(InputSequenceAsset));
+		FMusicTransportEventStreamReadRef InTransport = InputData.GetOrConstructDataReadReference<FMusicTransportEventStream>(METASOUND_GET_PARAM_NAME(Inputs::Transport), InParams.OperatorSettings);
+		FFloatReadRef InSpeed = InputData.GetOrCreateDefaultDataReadReference<float>(METASOUND_GET_PARAM_NAME(Inputs::Speed), InParams.OperatorSettings);
+		FFloatReadRef InVelocityMult = InputData.GetOrCreateDefaultDataReadReference<float>(METASOUND_GET_PARAM_NAME(InputVelocityMultiplier), InParams.OperatorSettings);
+		FMidiClockReadRef InMidiClock = InputData.GetOrConstructDataReadReference<FMidiClock>(METASOUND_GET_PARAM_NAME(Inputs::MidiClock), InParams.OperatorSettings);
+		FFloatReadRef InMaxColumns = InputData.GetOrCreateDefaultDataReadReference<float>(METASOUND_GET_PARAM_NAME(InputMaxColumns), InParams.OperatorSettings);
+		FFloatReadRef InAdditionalOctaves = InputData.GetOrCreateDefaultDataReadReference<float>(METASOUND_GET_PARAM_NAME(InputAdditionalOctaves), InParams.OperatorSettings);
+		FFloatReadRef InStepSize = InputData.GetOrCreateDefaultDataReadReference<float>(METASOUND_GET_PARAM_NAME(InputStepSizeQuarterNotes), InParams.OperatorSettings);
+		FFloatReadRef InActivePage = InputData.GetOrCreateDefaultDataReadReference<float>(METASOUND_GET_PARAM_NAME(InputActivePage), InParams.OperatorSettings);
+		FBoolReadRef InAutoPage = InputData.GetOrCreateDefaultDataReadReference<bool>(METASOUND_GET_PARAM_NAME(InputAutoPage), InParams.OperatorSettings);
+		FBoolReadRef InAutoPagePlaysBlankPages = InputData.GetOrCreateDefaultDataReadReference<bool>(METASOUND_GET_PARAM_NAME(InputAutoPagePlaysBlankPages), InParams.OperatorSettings);
+		FBoolReadRef InLoop = InputData.GetOrCreateDefaultDataReadReference<bool>(METASOUND_GET_PARAM_NAME(InputLoop), InParams.OperatorSettings);
+		FBoolReadRef InEnabled = InputData.GetOrCreateDefaultDataReadReference<bool>(METASOUND_GET_PARAM_NAME(InputEnabled), InParams.OperatorSettings);
+		return MakeUnique<FStepSequencePlayerOperator>(InParams, InSequenceAsset, InTransport, InMidiClock, InSpeed, InVelocityMult, InMaxColumns, InAdditionalOctaves, InStepSize, InActivePage, InAutoPage, InAutoPagePlaysBlankPages, InLoop, InEnabled);
+	}
+
+	FStepSequencePlayerOperator::FStepSequencePlayerOperator(const FBuildOperatorParams& InParams,
+															 const FMidiStepSequenceAssetReadRef& InSequenceAsset,
+															 const FMusicTransportEventStreamReadRef& InTransport,
+															 const FMidiClockReadRef& InMidiClockSource,
+															 const FFloatReadRef& InSpeedMultiplier,
+															 const FFloatReadRef& InVelocityMultiplier,
+															 const FFloatReadRef& InMaxColumns,
+															 const FFloatReadRef& InAdditionalOctaves,
+															 const FFloatReadRef& InStepSizeQuarterNotes,
+															 const FFloatReadRef& InActivePage,
+															 const FBoolReadRef& InAutoPage,
+															 const FBoolReadRef& InAutoPagePlaysBlankPages,
+															 const FBoolReadRef& InLoop,
+															 const FBoolReadRef& InEnabled)
+		: FMusicTransportControllable(EMusicPlayerTransportState::Prepared)
+		, SequenceAssetInPin(InSequenceAsset)
+		, TransportInPin(InTransport)
+		, SpeedMultInPin(InSpeedMultiplier)
+		, VelocityMultInPin(InVelocityMultiplier)
+		, MidiClockInPin(InMidiClockSource)
+		, MaxColumnsInPin(InMaxColumns)
+		, AdditionalOctavesInPin(InAdditionalOctaves)
+		, StepSizeQuarterNotesInPin(InStepSizeQuarterNotes)
+		, ActivePageInPin(InActivePage)
+		, AutoPageInPin(InAutoPage)
+		, AutoPagePlaysBlankPagesInPin(InAutoPagePlaysBlankPages)
+		, LoopInPin(InLoop)
+		, EnabledInPin(InEnabled)
+		, MidiOutPin(FMidiStreamWriteRef::CreateNew(InParams.OperatorSettings))
+	{
+		Reset(InParams);
+		Init();
+	}
+
+	void FStepSequencePlayerOperator::BindInputs(FInputVertexInterfaceData& InVertexData)
+	{
+		using namespace StepSequencePlayerPinNames;
+		using namespace CommonPinNames;
+		InVertexData.BindReadVertex(METASOUND_GET_PARAM_NAME(InputSequenceAsset), SequenceAssetInPin);
+		InVertexData.BindReadVertex(METASOUND_GET_PARAM_NAME(Inputs::Transport), TransportInPin);
+		InVertexData.BindReadVertex(METASOUND_GET_PARAM_NAME(Inputs::MidiClock), MidiClockInPin);
+		InVertexData.BindReadVertex(METASOUND_GET_PARAM_NAME(Inputs::Speed), SpeedMultInPin);
+		InVertexData.BindReadVertex(METASOUND_GET_PARAM_NAME(InputVelocityMultiplier), VelocityMultInPin);
+		InVertexData.BindReadVertex(METASOUND_GET_PARAM_NAME(InputMaxColumns), MaxColumnsInPin);
+		InVertexData.BindReadVertex(METASOUND_GET_PARAM_NAME(InputAdditionalOctaves), AdditionalOctavesInPin);
+		InVertexData.BindReadVertex(METASOUND_GET_PARAM_NAME(InputStepSizeQuarterNotes), StepSizeQuarterNotesInPin);
+		InVertexData.BindReadVertex(METASOUND_GET_PARAM_NAME(InputActivePage), ActivePageInPin);
+		InVertexData.BindReadVertex(METASOUND_GET_PARAM_NAME(InputAutoPage), AutoPageInPin);
+		InVertexData.BindReadVertex(METASOUND_GET_PARAM_NAME(InputAutoPagePlaysBlankPages), AutoPagePlaysBlankPagesInPin);
+		InVertexData.BindReadVertex(METASOUND_GET_PARAM_NAME(InputLoop), LoopInPin);
+		InVertexData.BindReadVertex(METASOUND_GET_PARAM_NAME(InputEnabled), EnabledInPin);
+
+		Init();
+	}
+
+	void FStepSequencePlayerOperator::BindOutputs(FOutputVertexInterfaceData& InVertexData)
+	{
+		using namespace CommonPinNames;
+		InVertexData.BindReadVertex(METASOUND_GET_PARAM_NAME(Outputs::MidiStream), MidiOutPin);
+	}
+
+	FDataReferenceCollection FStepSequencePlayerOperator::GetInputs() const
+	{
+		// This should never be called. Bind(...) is called instead. This method
+		// exists as a stop-gap until the API can be deprecated and removed.
+		checkNoEntry();
+		return {};
+	}
+
+	FDataReferenceCollection FStepSequencePlayerOperator::GetOutputs() const
+	{
+		// This should never be called. Bind(...) is called instead. This method
+		// exists as a stop-gap until the API can be deprecated and removed.
+		checkNoEntry();
+		return {};
+	}
+
+	void FStepSequencePlayerOperator::Reset(const FResetParams& ResetParams)
+	{
+		BlockSize = ResetParams.OperatorSettings.GetNumFramesPerBlock();
+		CurrentBlockSpanStart = 0;
+		CurrentPageIndex = -1;
+		CurrentCellIndex = -1;
+		ProcessedThruTick = -1;
+		SequenceStartTick = -1;
+		bNeedsMidiOutClear = true;
+		bAutoPage = false;
+		bPreviousAutoPage = false;
+		bAutoPagePlaysBlankPages = false;
+		bLoop = true;
+		bEnabled = true;
+		bPlaying = false;
+		bNeedsRebase = false;
+		CurrentCellNotes.Reset();
+	}
+
+	void FStepSequencePlayerOperator::Execute()
+	{
+		// If no midi cursor callbacks happened we still have to 
+		// remove old midi events from last block...
+		if (bNeedsMidiOutClear)
+		{
+			MidiOutPin->PrepareBlock();
+			bNeedsMidiOutClear = false;
+		}
+
+		// if we have no sequence table there is nothing to do. 
+		// Make sure the notes are all off and return.
+		if (!SequenceTable)
+		{
+			if (CurrentCellNotes.Num() > 0)
+			{
+				AllNotesOff(0, MidiClockInPin->GetCurrentMidiTick(), true);
+				CurrentCellNotes.Empty();
+			}
+			return;
+		}
+
+		CheckForUpdatedSequenceTable();
+
+		// We need to cache this to avoid avoid a crash if the value in the 
+		// sequence table asset changes while we are in the middle of rendering. 
+		CurrentStepSkipIndex = SequenceTable->StepSkipIndex;
+
+		TransportSpanProcessor TransportHandler = [&](int32 StartFrameIndex, int32 EndFrameIndex, EMusicPlayerTransportState CurrentState)
+			{
+				CurrentBlockSpanStart = StartFrameIndex;
+				switch (CurrentState)
+				{
+				case EMusicPlayerTransportState::Invalid:
+				case EMusicPlayerTransportState::Preparing:
+					MidiOutPin->AddTransportStateChangeMessage(StartFrameIndex, EMusicPlayerTransportState::Prepared);
+					return EMusicPlayerTransportState::Prepared;
+				
+				case EMusicPlayerTransportState::Prepared:
+					MidiOutPin->AddTransportStateChangeMessage(StartFrameIndex, CurrentState);
+					return CurrentState;
+				
+				case EMusicPlayerTransportState::Starting:
+					MidiOutPin->AddTransportStateChangeMessage(StartFrameIndex, EMusicPlayerTransportState::Playing);
+					return EMusicPlayerTransportState::Playing;
+
+				case EMusicPlayerTransportState::Playing:
+					return EMusicPlayerTransportState::Playing;
+
+				case EMusicPlayerTransportState::Seeking:
+					return GetTransportState();
+
+				case EMusicPlayerTransportState::Continuing:
+					MidiOutPin->AddTransportStateChangeMessage(StartFrameIndex, EMusicPlayerTransportState::Continuing);
+					return EMusicPlayerTransportState::Playing;
+
+				case EMusicPlayerTransportState::Pausing:
+					MidiOutPin->AddTransportStateChangeMessage(StartFrameIndex, EMusicPlayerTransportState::Pausing);
+					return EMusicPlayerTransportState::Paused;
+
+				case EMusicPlayerTransportState::Paused:
+					return EMusicPlayerTransportState::Paused;
+
+				case EMusicPlayerTransportState::Stopping:
+					AllNotesOff(StartFrameIndex, MidiClockInPin->GetCurrentMidiTick(), true);
+					MidiOutPin->AddTransportStateChangeMessage(StartFrameIndex, EMusicPlayerTransportState::Stopping);
+					return EMusicPlayerTransportState::Prepared;
+
+				case EMusicPlayerTransportState::Killing:
+					AllNotesOff(StartFrameIndex, MidiClockInPin->GetCurrentMidiTick(), true);
+					MidiOutPin->AddTransportStateChangeMessage(StartFrameIndex, EMusicPlayerTransportState::Killing);
+					return EMusicPlayerTransportState::Prepared;
+
+				default:
+					checkNoEntry();
+					return EMusicPlayerTransportState::Invalid;
+				}
+			};
+		ExecuteTransportSpans(TransportInPin, BlockSize, TransportHandler);
+
+		bNeedsMidiOutClear = true;
+	}
+
+	void FStepSequencePlayerOperator::Init()
+	{
+		MidiOutPin->PrepareBlock();
+		MidiOutPin->SetClockSource(MidiClockInPin);
+		MidiClockInPin->RegisterHiResPlayCursor(this);
+
+		SequenceTable = SequenceAssetInPin->GetRenderable();
+		if (SequenceTable)
+		{
+			ResizeCellStatesForTable();
+			UE_LOG(LogStepSequencePlayer, Verbose, TEXT("Got a Sequence: %d pages, %d rows, %d columns"), 
+				SequenceTable->Pages.Num(), 
+				SequenceTable->Pages.IsEmpty() ? 0 : SequenceTable->Pages[0].Rows.Num(), 
+				SequenceTable->Pages.IsEmpty() || SequenceTable->Pages[0].Rows.IsEmpty() ? 0 : SequenceTable->Pages[0].Rows[0].Cells.Num());
+		}
+		else
+		{
+			UE_LOG(LogStepSequencePlayer, Verbose, TEXT("No Sequence Provided!"));
+		}
+
+		FTransportInitFn InitFn = [this](EMusicPlayerTransportState CurrentState)
+		{
+			CurrentBlockSpanStart = 0;
+			switch (CurrentState)
+			{
+			case EMusicPlayerTransportState::Invalid:
+			case EMusicPlayerTransportState::Preparing:
+				MidiOutPin->AddTransportStateChangeMessage(0, EMusicPlayerTransportState::Prepared);
+				return EMusicPlayerTransportState::Prepared;
+				
+			case EMusicPlayerTransportState::Prepared:
+				MidiOutPin->AddTransportStateChangeMessage(0, CurrentState);
+				return CurrentState;
+				
+			case EMusicPlayerTransportState::Starting:
+				MidiOutPin->AddTransportStateChangeMessage(0, EMusicPlayerTransportState::Playing);
+				return EMusicPlayerTransportState::Playing;
+
+			case EMusicPlayerTransportState::Playing:
+				return EMusicPlayerTransportState::Playing;
+
+			case EMusicPlayerTransportState::Seeking:
+				return GetTransportState();
+
+			case EMusicPlayerTransportState::Continuing:
+				MidiOutPin->AddTransportStateChangeMessage(0, EMusicPlayerTransportState::Continuing);
+				return EMusicPlayerTransportState::Playing;
+
+			case EMusicPlayerTransportState::Pausing:
+				MidiOutPin->AddTransportStateChangeMessage(0, EMusicPlayerTransportState::Pausing);
+				return EMusicPlayerTransportState::Paused;
+
+			case EMusicPlayerTransportState::Paused:
+				return EMusicPlayerTransportState::Paused;
+
+			case EMusicPlayerTransportState::Stopping:
+				AllNotesOff(0, MidiClockInPin->GetCurrentMidiTick(), true);
+				MidiOutPin->AddTransportStateChangeMessage(0, EMusicPlayerTransportState::Stopping);
+				return EMusicPlayerTransportState::Prepared;
+
+			case EMusicPlayerTransportState::Killing:
+				AllNotesOff(0, MidiClockInPin->GetCurrentMidiTick(), true);
+				MidiOutPin->AddTransportStateChangeMessage(0, EMusicPlayerTransportState::Killing);
+				return EMusicPlayerTransportState::Prepared;
+
+			default:
+				checkNoEntry();
+				return EMusicPlayerTransportState::Invalid;
+			}
+		};
+		FMusicTransportControllable::Init(*TransportInPin, MoveTemp(InitFn));
+	}
+
+	void FStepSequencePlayerOperator::SeekToTick(int32 Tick)
+	{
+		SeekThruTick(Tick-1);
+	}
+
+	void FStepSequencePlayerOperator::SeekThruTick(int32 Tick)
+	{
+		if (bNeedsMidiOutClear)
+		{
+			MidiOutPin->PrepareBlock();
+			bNeedsMidiOutClear = false;
+		}
+
+		int32 BlockFrame = (*MidiClockInPin).GetCurrentBlockFrameIndex();
+		ProcessedThruTick = FMath::Max(-1, Tick);
+		bNeedsRebase = true;
+		AllNotesOff(BlockFrame, Tick, true);
+	}
+
+	int32 FStepSequencePlayerOperator::CalculatePagesProgressed(const int FromTick, const int ToTick, int32 TableTickLength, const bool bRound) const
+	{
+		const float PagesProgressed = (static_cast<float>(ToTick - FromTick) / static_cast<float>(TableTickLength));
+		return bRound ? FMath::RoundToInt32(PagesProgressed) : FMath::FloorToInt32(PagesProgressed);
+	}
+
+	int32 FStepSequencePlayerOperator::CalculateAutoPageIndex(int32 CurTick, int32 TableTickLength, bool bRound) const
+	{
+		const int32 PagesProgressed = CalculatePagesProgressed(SequenceStartTick, CurTick, TableTickLength, bRound);
+		int32 TargetPageIndex = SequenceTable->CalculateAutoPageIndex(PagesProgressed, bAutoPagePlaysBlankPages, bLoop);
+
+		if (TargetPageIndex != INDEX_NONE && !SequenceTable->Pages.IsValidIndex(TargetPageIndex))
+		{
+			TargetPageIndex = 0;
+		}
+
+		// Page Index is 1-based. If INDEX_NONE is returned from the above, this will be 0 and therefore signals to be nonplaying
+		++TargetPageIndex;
+		return TargetPageIndex;
+	}
+
+	void FStepSequencePlayerOperator::CalculatePageProperties(const FStepSequencePage& Page, int32 MaxColumns, float StepSizeQuarterNotes, int32& OutColumns, int32& OutTicksPerCell, int32& OutTableTickLength) const
+	{		
+		OutColumns = FMath::Min(Page.Rows[0].Cells.Num(), (MaxColumns < 1 ? 1 : MaxColumns));
+
+		int32 LengthColumns = OutColumns;
+		
+		// Lower the total number of columns to reflect the number of steps skipped
+		// This value is used only for calculating the tick length
+		if (CurrentStepSkipIndex > 0)
+		{
+			LengthColumns -= LengthColumns / CurrentStepSkipIndex;
+			LengthColumns = FMath::Max(LengthColumns, 1);
+		}
+
+		// Min float possible is 64th note triplets.
+		OutTicksPerCell = FMath::Max(StepSizeQuarterNotes, 0.0416666f) * MidiConstants::kTicksPerQuarterNote;
+		OutTableTickLength = OutTicksPerCell * LengthColumns;
+	}
+
+	void FStepSequencePlayerOperator::RebaseSequenceStartTickForLoop(const int32 CurTick, const int32 TableTickLength)
+	{
+		if (CurTick == INDEX_NONE)
+		{
+			SequenceStartTick = 0;
+			CurrentPageIndex = -1;
+			bNeedsRebase = false;
+			return;
+		}
+
+		const int32 PagesProgressed = CalculatePagesProgressed(0, CurTick, TableTickLength, false);
+		const int32 TickInPage = TableTickLength == 0 ? 0 : CurTick % TableTickLength;
+		const int32 NumValidPages = bAutoPage ? SequenceTable->CalculateNumValidPages(bAutoPagePlaysBlankPages) : 1;
+		const int32 PagesSinceStart = NumValidPages == 0 ? 0 : PagesProgressed % NumValidPages;
+
+		SequenceStartTick = CurTick - (PagesSinceStart * TableTickLength) - TickInPage;
+		CurrentPageIndex = CalculateAutoPageIndex(CurTick, TableTickLength, false);
+		bNeedsRebase = false;
+	}
+
+	void FStepSequencePlayerOperator::EnsureCurrentPageIndexIsValid()
+	{
+		if (bAutoPage)
+		{
+			// Ensure we don't process erroneous -1, conversion to 1-based first page
+			CurrentPageIndex = FMath::Abs(CurrentPageIndex);
+			if (CurrentPageIndex == 0)
+			{
+				// The loop ended in a previous tick, bPlaying should be false here.
+				// Reset the current index, and we'll wait to start playing again.
+				CurrentPageIndex = SequenceTable->GetFirstValidPage(bAutoPagePlaysBlankPages) + 1;
+			}
+		}
+		else
+		{
+			CurrentPageIndex = FMath::Clamp((int32)*ActivePageInPin, 1, SequenceTable->Pages.Num());
+		}
+	}
+
+	void FStepSequencePlayerOperator::AdvanceThruTick(int32 Tick, bool IsPreRoll)
+	{
+		if (!SequenceTable || SequenceTable->Pages.IsEmpty())
+		{
+			return;
+		}
+
+		if (bNeedsMidiOutClear)
+		{
+			MidiOutPin->PrepareBlock();
+			bNeedsMidiOutClear = false;
+		}
+
+		// Read input pins
+		const int32 CurrentMaxColumns = *MaxColumnsInPin;
+		const int32 AdditionalOctaveNotes = (int32)*AdditionalOctavesInPin * 12;
+		const float CurrentStepSizeQuarterNotes = *StepSizeQuarterNotesInPin;
+		const float CurrentVelocityMultiplierValue = *VelocityMultInPin;
+		const int32 BlockFrameIndex = MidiClockInPin->GetCurrentBlockFrameIndex();
+
+		if (CurrentMaxColumns < 1)
+		{
+			// Nothing to play
+			return;
+		}
+
+		if (!bAutoPage && *AutoPageInPin)
+		{
+			bNeedsRebase = true;
+		}
+
+		bAutoPage = *AutoPageInPin;
+		bAutoPagePlaysBlankPages = *AutoPagePlaysBlankPagesInPin;
+
+		EnsureCurrentPageIndexIsValid();
+		if (!SequenceTable->Pages.IsValidIndex(CurrentPageIndex - 1))
+		{
+			// WE SHOULD NEVER HIT THIS
+			return;
+		}
+		FStepSequencePage* CurrentPage = &(SequenceTable->Pages[CurrentPageIndex - 1]);
+
+		// Do an initial calc with the page we think we're on
+		int32 Columns;
+		int32 TicksPerCell;
+		int32 TableTickLength;
+		CalculatePageProperties(*CurrentPage, CurrentMaxColumns, CurrentStepSizeQuarterNotes, Columns, TicksPerCell, TableTickLength);
+		
+		bool bNewPlaying = bEnabled;
+
+		if (*EnabledInPin != bEnabled || (*LoopInPin && !bLoop))
+		{
+			// If our enable state changes or we turn loop on, control playback
+			bEnabled = *EnabledInPin;
+			bNewPlaying = bEnabled;
+		}
+		else if (*EnabledInPin && *LoopInPin)
+		{
+			// Force playing to true if enabled and looping
+			bNewPlaying = true;
+		}
+
+		if (bPlaying)
+		{
+			if (!bNewPlaying)
+			{
+				// We're not playing anymore, make sure we're all off
+				AllNotesOff(BlockFrameIndex, Tick, true);
+			}
+			else if (*LoopInPin != bLoop)
+			{
+				// If we're already playing and change the loop setting, need to re-sync to the global timeline
+				bNeedsRebase = true;
+			}
+		}
+		else if (bNewPlaying)
+		{
+			// We weren't playing before, but now we are
+			if (*LoopInPin)
+			{
+				bNeedsRebase = true;
+			}
+			else
+			{
+				// Begin the oneshot on the next beat subdivision
+				SequenceStartTick = ProcessedThruTick + TicksPerCell - (ProcessedThruTick % TicksPerCell);
+				if (bAutoPage)
+				{
+					// Nonlooping autopaging sequences start on the first valid page.
+					CurrentPageIndex = SequenceTable->GetFirstValidPage(bAutoPagePlaysBlankPages) + 1;
+				}
+				bNeedsRebase = false;
+			}
+		}
+
+		bPlaying = bNewPlaying;
+		bLoop = *LoopInPin;
+
+		// TODO When IsPreRoll we should probably be doing some of the handling of note ons/offs differently (more like FMidiPlayCursor::AdvanceThruTick), and therefore potentially our timeline sync as well
+		if (bNeedsRebase)
+		{
+			if (!bLoop)
+			{
+				// Temp: Don't rebase if loop is off, so we don't lose our start tick on graph rebuilds.
+				// NOTE: This introduces a bug where turning loop off in the middle of a sequence will cause a visual desync, but with loop in the Customize menu that's not a major concern.
+				bNeedsRebase = false;
+			}
+			else
+			{
+				RebaseSequenceStartTickForLoop(ProcessedThruTick, TableTickLength);
+				// Our page index may have changed, let's set it again
+				EnsureCurrentPageIndexIsValid();
+				if (!SequenceTable->Pages.IsValidIndex(CurrentPageIndex - 1))
+				{
+					// WE SHOULD NEVER HIT THIS
+					return;
+				}
+				CurrentPage = &(SequenceTable->Pages[CurrentPageIndex - 1]);
+				// Uncomment if there is ever a time where pages can have different tick lengths
+				// CalculatePageProperties(*CurrentPage, CurrentMaxColumns, CurrentStepSizeQuarterNotes, Columns, TicksPerCell, TableTickLength);
+			}
+		}
+
+		ProcessedThruTick = FMath::Max(-1, ProcessedThruTick);
+		
+		while (ProcessedThruTick < Tick)
+		{
+			ProcessedThruTick++;
+
+			if (!bPlaying || ProcessedThruTick < SequenceStartTick)
+			{
+				// Loop has ended or hasn't started yet
+				continue;
+			}
+
+			int32 EffectiveTickForLoopPosition = ProcessedThruTick;
+
+			if (!bLoop)
+			{
+				// If loop is off, we want to position ourselves relative to when the sequence was triggered
+				EffectiveTickForLoopPosition -= SequenceStartTick;
+
+				if ((bAutoPage && CurrentPageIndex == 0) || (!bAutoPage && EffectiveTickForLoopPosition >= TableTickLength))
+				{
+					bPlaying = false;
+					continue;
+				}
+			}
+
+			const int32 TickInTable = (TableTickLength == 0) ? 0 : EffectiveTickForLoopPosition % TableTickLength;
+			int32 CellInRow = TickInTable / TicksPerCell;
+			const int32 TickInCell = TickInTable - (CellInRow * TicksPerCell);
+			int32 SkippedIndex = CurrentStepSkipIndex - 1;
+
+			if (CurrentStepSkipIndex >= 2)
+			{
+				// Skip past the unused cells
+				CellInRow += CellInRow / SkippedIndex;
+			}
+
+			if (TickInCell >= TicksPerCell - 1)
+			{
+				int32 NextCellInRow = (CellInRow + 1) % Columns;
+
+				if (CurrentStepSkipIndex >= 2 && NextCellInRow % SkippedIndex == 0)
+				{
+					// If the next cell is unused, instead check the cell past this
+					NextCellInRow = (NextCellInRow + 1) % Columns;
+				}
+
+				if (bAutoPage && (NextCellInRow == 0 || bAutoPage != bPreviousAutoPage))
+				{
+					CurrentPageIndex = CalculateAutoPageIndex(ProcessedThruTick, TableTickLength, true);
+
+					if (CurrentPageIndex == 0)
+					{
+						// Loop has ended
+						bPlaying = false;
+						continue;
+					}
+					else if (!SequenceTable->Pages.IsValidIndex(CurrentPageIndex - 1))
+					{
+						// WE SHOULD NEVER HIT THIS
+						return;
+					}
+
+					CurrentPage = &(SequenceTable->Pages[CurrentPageIndex - 1]);
+					CalculatePageProperties(*CurrentPage, CurrentMaxColumns, CurrentStepSizeQuarterNotes, Columns, TicksPerCell, TableTickLength);
+				}	
+				bPreviousAutoPage = bAutoPage;
+
+				for (int32 i = 0; i < CurrentPage->Rows.Num(); ++i)
+				{
+					if (CurrentCellNotes[i])
+					{
+						uint8 MidiCh;
+						uint8 MidiNote;
+						CurrentCellNotes[i].GetChannelAndNote(MidiCh, MidiNote);
+
+						// If: 
+						// 1. The note that would play is not enabled (it is not a new note);
+						// 2. A note with this pitch would start up this tick;
+						// 3. The note is marked as a continuation note
+						// Then keep this note playing through the next cell
+						if (MidiNote == (SequenceTable->Notes[i].NoteNumber + AdditionalOctaveNotes) && !CurrentPage->Rows[i].Cells[NextCellInRow].bEnabled && CurrentPage->Rows[i].Cells[NextCellInRow].bContinuation)
+						{
+							continue;
+						}
+
+						// note off!
+						FMidiStreamEvent MidiEvent(this, FMidiMsg::CreateNoteOff(MidiCh, SequenceTable->Notes[i].NoteNumber + AdditionalOctaveNotes));
+						MidiEvent.BlockSampleFrameIndex = BlockFrameIndex;
+						MidiEvent.BlockSampleFrameOffset = 0.0f;
+						MidiEvent.AuthoredMidiTick = ProcessedThruTick;
+						MidiEvent.CurrentMidiTick = ProcessedThruTick;
+						MidiEvent.MsOffset = 0.0f;
+						MidiEvent.TrackIndex = 1;
+						MidiEvent.SetVoiceId(CurrentCellNotes[i]);
+						MidiOutPin->AddNoteOffEventOrCancelPendingNoteOn(MidiEvent);
+						UE_LOG(LogStepSequencePlayer, Verbose, TEXT("0x%x Note-Off %d at %d"), (uint32)(size_t)this, SequenceTable->Notes[i].NoteNumber + AdditionalOctaveNotes, BlockFrameIndex);
+						CurrentCellNotes[i] = FMidiVoiceId::None();
+					}
+				}
+			}
+
+			if (CellInRow != CurrentCellIndex)
+			{	
+				for (int32 i = 0; i < CurrentPage->Rows.Num(); ++i)
+				{
+					if (CurrentPage->Rows[i].bRowEnabled && CurrentPage->Rows[i].Cells[CellInRow].bEnabled)
+					{
+						if (!CurrentCellNotes[i])
+						{
+							// note on!
+							FMidiStreamEvent MidiEvent(this, FMidiMsg::CreateNoteOn(0, SequenceTable->Notes[i].NoteNumber + AdditionalOctaveNotes, SequenceTable->Notes[i].Velocity));
+							float NoteOnVelocity = FMath::Clamp(static_cast<float>(SequenceTable->Notes[i].Velocity) * CurrentVelocityMultiplierValue, 0.0f, 127.0f);
+							MidiEvent.MidiMessage.SetNoteOnVelocity(static_cast<uint8>(NoteOnVelocity));
+							MidiEvent.BlockSampleFrameIndex = BlockFrameIndex;
+							MidiEvent.BlockSampleFrameOffset = 0.0f;
+							MidiEvent.AuthoredMidiTick = ProcessedThruTick;
+							MidiEvent.CurrentMidiTick = ProcessedThruTick;
+							MidiEvent.MsOffset = 0.0f;
+							MidiEvent.TrackIndex = 1;
+							MidiOutPin->AddMidiEvent(MidiEvent);
+							UE_LOG(LogStepSequencePlayer, Verbose, TEXT("0x%x Note-On %d at %d"), (uint32)(size_t)this, SequenceTable->Notes[i].NoteNumber + AdditionalOctaveNotes, BlockFrameIndex);
+							CurrentCellNotes[i] = MidiEvent.GetVoiceId();
+						}
+					}
+				}
+				CurrentCellIndex = CellInRow;
+			}
+		}
+
+		if (bNewPlaying && !bPlaying)
+		{
+			// Loop ended above, make sure we turn everything off
+			AllNotesOff(BlockFrameIndex, Tick, true);
+		}
+	}
+
+	void FStepSequencePlayerOperator::CheckForUpdatedSequenceTable()
+	{
+		if (SequenceTable)
+		{
+			TRefCountedAudioRenderableWithQueuedChanges<FStepSequenceTable>* Table = SequenceTable;
+			if (Table->HasUpdate())
+			{
+				SequenceTable = Table->GetUpdate();
+				ResizeCellStatesForTable();
+				UE_LOG(LogStepSequencePlayer, Verbose, TEXT("Got a NEW Sequence: %d pages, %d rows, %d columns"),
+					SequenceTable->Pages.Num(),
+					SequenceTable->Pages.IsEmpty() ? 0 : SequenceTable->Pages[0].Rows.Num(),
+					SequenceTable->Pages.IsEmpty() || SequenceTable->Pages[0].Rows.IsEmpty() ? 0 : SequenceTable->Pages[0].Rows[0].Cells.Num());
+			}
+		}
+	}
+
+	void FStepSequencePlayerOperator::ResizeCellStatesForTable()
+	{
+		if (!SequenceTable || SequenceTable->Notes.Num() == 0)
+		{
+			return;
+		}
+
+		if (SequenceTable->Notes.Num() < CurrentCellNotes.Num())
+		{
+			// We may have existing notes that need to be stopped.
+			for (int32 i = SequenceTable->Notes.Num(); i < CurrentCellNotes.Num(); ++i)
+			{
+				if (CurrentCellNotes[i])
+				{
+					uint8 MidiCh;
+					uint8 MidiNote;
+					CurrentCellNotes[i].GetChannelAndNote(MidiCh, MidiNote);
+					FMidiStreamEvent MidiEvent(this, FMidiMsg::CreateNoteOff(MidiCh, MidiNote));
+					MidiEvent.BlockSampleFrameIndex  = CurrentBlockSpanStart;
+					MidiEvent.BlockSampleFrameOffset = 0.0f;
+					MidiEvent.AuthoredMidiTick       = 0;
+					MidiEvent.CurrentMidiTick        = 0;
+					MidiEvent.MsOffset               = 0.0f;
+					MidiEvent.TrackIndex             = 1;
+					MidiEvent.SetVoiceId(CurrentCellNotes[i]);
+					UE_LOG(LogStepSequencePlayer, Verbose, TEXT("0x%x Note-Off %d (during resize)"), (uint32)(size_t)this, MidiNote);
+					MidiOutPin->AddMidiEvent(MidiEvent);
+				}
+			}
+			CurrentCellNotes.SetNum(SequenceTable->Notes.Num());
+		}
+		else
+		{
+			while (CurrentCellNotes.Num() < SequenceTable->Notes.Num())
+			{
+				CurrentCellNotes.Add(FMidiVoiceId());
+			}
+		}
+	}
+
+	void FStepSequencePlayerOperator::AllNotesOff(int32 AtFrameIndex, int32 AbsMidiTick, bool ResetCellIndex)
+	{
+		for (int i = 0; i < CurrentCellNotes.Num(); ++i)
+		{
+			if (CurrentCellNotes[i])
+			{
+				uint8 MidiCh;
+				uint8 MidiNote;
+				CurrentCellNotes[i].GetChannelAndNote(MidiCh, MidiNote);
+				FMidiStreamEvent MidiEvent(this, FMidiMsg::CreateNoteOff(MidiCh, MidiNote));
+				MidiEvent.BlockSampleFrameIndex  = AtFrameIndex;
+				MidiEvent.BlockSampleFrameOffset = 0.0f;
+				MidiEvent.AuthoredMidiTick       = AbsMidiTick;
+				MidiEvent.CurrentMidiTick        = AbsMidiTick;
+				MidiEvent.MsOffset               = 0.0f;
+				MidiEvent.TrackIndex             = 1;
+				MidiEvent.SetVoiceId(CurrentCellNotes[i]);
+				MidiOutPin->AddMidiEvent(MidiEvent);
+				UE_LOG(LogStepSequencePlayer, Verbose, TEXT("0x%x Note-Off %d (during all notes off)"), (uint32)(size_t)this, MidiNote);
+				CurrentCellNotes[i] = FMidiVoiceId::None();
+			}
+		}
+		if (ResetCellIndex)
+		{
+			CurrentCellIndex = -1;
+		}
+	}
+}
+
+#undef LOCTEXT_NAMESPACE // "HarmonixMetaSound"
diff --git a/Engine/Plugins/Runtime/Harmonix/Source/HarmonixMetasound/Private/HarmonixMetasound/Nodes/TempoMapMetronomeNode.cpp b/Engine/Plugins/Runtime/Harmonix/Source/HarmonixMetasound/Private/HarmonixMetasound/Nodes/TempoMapMetronomeNode.cpp
new file mode 100644
--- /dev/null
+++ b/Engine/Plugins/Runtime/Harmonix/Source/HarmonixMetasound/Private/HarmonixMetasound/Nodes/TempoMapMetronomeNode.cpp
@@ -0,0 +1,260 @@
+// Copyright Epic Games, Inc. All Rights Reserved.
+
+#include "MetasoundExecutableOperator.h"
+#include "MetasoundFacade.h"
+#include "MetasoundNodeInterface.h"
+#include "MetasoundParamHelper.h"
+#include "MetasoundSampleCounter.h"
+#include "MetasoundStandardNodesCategories.h"
+#include "MetasoundVertex.h"
+
+#include "HarmonixMetasound/Common.h"
+#include "HarmonixMetasound/DataTypes/MidiAsset.h"
+#include "HarmonixMetasound/DataTypes/MidiClock.h"
+#include "HarmonixMetasound/DataTypes/MusicTransport.h"
+#include "HarmonixMetasound/DataTypes/MusicTimestamp.h"
+#include <algorithm>
+
+#define LOCTEXT_NAMESPACE "HarmonixMetaSound"
+
+DEFINE_LOG_CATEGORY_STATIC(TempoMapMetronomeOperator, Log, All);
+
+namespace HarmonixMetasound
+{
+	using namespace Metasound;
+
+	class FTempoMapMetronomeOperator : public TExecutableOperator<FTempoMapMetronomeOperator>, public FMusicTransportControllable
+	{
+	public:
+		static const FNodeClassMetadata& GetNodeInfo();
+		static const FVertexInterface& GetVertexInterface();
+		static TUniquePtr<IOperator> CreateOperator(const FBuildOperatorParams& InParams, FBuildResults& OutResults);
+
+		FTempoMapMetronomeOperator(const FOperatorSettings& InSettings,
+                                   const FMidiAssetReadRef& InMidiAsset,
+						           const FMusicTransportEventStreamReadRef& InTransport,
+		                           const FFloatReadRef& InSpeedMultiplier,
+								   const FInt32ReadRef& InSeekPrerollBars);
+
+		virtual void BindInputs(FInputVertexInterfaceData& InVertexData) override;
+		virtual void BindOutputs(FOutputVertexInterfaceData& InVertexData) override;
+		virtual FDataReferenceCollection GetInputs() const override;
+		virtual FDataReferenceCollection GetOutputs() const override;
+
+		void Reset(const FResetParams& Params);
+		void Execute();
+
+	private:
+		//** INPUTS
+		FMidiAssetReadRef MidiAssetInPin;
+		FMusicTransportEventStreamReadRef TransportInPin;
+		FFloatReadRef SpeedMultInPin;
+		FInt32ReadRef SeekPrerollBarsInPin;
+
+		//** OUTPUTS
+		FMidiClockWriteRef MidiClockOutPin;
+
+		//** DATA
+		FMidiClockEventCursor MidiClockEventCursor;
+		FSampleCount BlockSize;
+		float        SampleRate;
+	};
+
+	class FTempoMapMetronomeNode : public FNodeFacade
+	{
+	public:
+		FTempoMapMetronomeNode(const FNodeInitData& InInitData)
+			: FNodeFacade(InInitData.InstanceName, InInitData.InstanceID, TFacadeOperatorClass<FTempoMapMetronomeOperator>())
+		{}
+		virtual ~FTempoMapMetronomeNode() = default;
+	};
+
+	METASOUND_REGISTER_NODE(FTempoMapMetronomeNode)
+
+	const FNodeClassMetadata& FTempoMapMetronomeOperator::GetNodeInfo()
+	{
+		auto InitNodeInfo = []() -> FNodeClassMetadata
+		{
+			FNodeClassMetadata Info;
+			Info.ClassName        = { HarmonixNodeNamespace, TEXT("TempoMapMetronome"), TEXT("") };
+			Info.MajorVersion     = 0;
+			Info.MinorVersion     = 1;
+			Info.DisplayName      = METASOUND_LOCTEXT("TempoMapMetronomeNode_DisplayName", "Tempo Map Midi Clock Generator");
+			Info.Description      = METASOUND_LOCTEXT("TempoMapMetronomeNode_Description", "Provides a midi clock using the tempo map in the specified Midi file.");
+			Info.Author           = PluginAuthor;
+			Info.PromptIfMissing  = PluginNodeMissingPrompt;
+			Info.DefaultInterface = GetVertexInterface();
+			Info.CategoryHierarchy.Emplace(NodeCategories::Music);
+
+			return Info;
+		};
+
+		static const FNodeClassMetadata Info = InitNodeInfo();
+
+		return Info;
+	}
+
+	const FVertexInterface& FTempoMapMetronomeOperator::GetVertexInterface()
+	{
+		using namespace CommonPinNames;
+
+		static const FVertexInterface Interface(
+			FInputVertexInterface(
+				TInputDataVertex<FMidiAsset>(METASOUND_GET_PARAM_NAME_AND_METADATA(Inputs::MidiFileAsset)),
+				TInputDataVertex<FMusicTransportEventStream>(METASOUND_GET_PARAM_NAME_AND_METADATA(Inputs::Transport)),
+				TInputDataVertex<float>(METASOUND_GET_PARAM_NAME_AND_METADATA(Inputs::Speed), 1.0f),
+				TInputDataVertex<int32>(METASOUND_GET_PARAM_NAME_AND_METADATA(Inputs::PrerollBars), 8)
+			),
+			FOutputVertexInterface(
+				TOutputDataVertex<FMidiClock>(METASOUND_GET_PARAM_NAME_AND_METADATA(Outputs::MidiClock))
+			)
+		);
+
+		return Interface;
+	}
+
+	TUniquePtr<IOperator> FTempoMapMetronomeOperator::CreateOperator(const FBuildOperatorParams& InParams, FBuildResults& OutResults)
+	{
+		using namespace CommonPinNames;
+
+		const FTempoMapMetronomeNode& TempoClockNode = static_cast<const FTempoMapMetronomeNode&>(InParams.Node);
+
+		const FInputVertexInterfaceData& InputData = InParams.InputData;
+		FMidiAssetReadRef InMidiAsset = InputData.GetOrConstructDataReadReference<FMidiAsset>(METASOUND_GET_PARAM_NAME(Inputs::MidiFileAsset));
+		FMusicTransportEventStreamReadRef InTransport = InputData.GetOrConstructDataReadReference<FMusicTransportEventStream>(METASOUND_GET_PARAM_NAME(Inputs::Transport), InParams.OperatorSettings);
+		FFloatReadRef InSpeed = InputData.GetOrCreateDefaultDataReadReference<float>(METASOUND_GET_PARAM_NAME(Inputs::Speed), InParams.OperatorSettings);
+		FInt32ReadRef InPrerollBars = InputData.GetOrCreateDefaultDataReadReference<int32>(METASOUND_GET_PARAM_NAME(Inputs::PrerollBars), InParams.OperatorSettings);
+
+		return MakeUnique<FTempoMapMetronomeOperator>(InParams.OperatorSettings, InMidiAsset, InTransport, InSpeed, InPrerollBars);
+	}
+
+	FTempoMapMetronomeOperator::FTempoMapMetronomeOperator(const FOperatorSettings& InSettings, 
+                                                           const FMidiAssetReadRef& InMidiAsset,
+	                                                       const FMusicTransportEventStreamReadRef& InTransport,
+	                                                       const FFloatReadRef& InSpeedMultiplier,
+														   const FInt32ReadRef& InPrerollBars)
+		: FMusicTransportControllable(EMusicPlayerTransportState::Prepared) 
+		, MidiAssetInPin(InMidiAsset)
+		, TransportInPin(InTransport)
+		, SpeedMultInPin(InSpeedMultiplier)
+		, SeekPrerollBarsInPin(InPrerollBars)
+		, MidiClockOutPin(FMidiClockWriteRef::CreateNew(InSettings))
+		, MidiClockEventCursor(MidiClockOutPin)
+		, BlockSize(InSettings.GetNumFramesPerBlock())
+		, SampleRate(InSettings.GetSampleRate())
+	{
+		const FMidiFileProxyPtr& MidiProxy = MidiAssetInPin->GetMidiProxy();
+		MidiClockOutPin->AttachToMidiResource(MidiProxy ? MidiProxy->GetMidiFile() : nullptr);
+	}
+
+	void FTempoMapMetronomeOperator::BindInputs(FInputVertexInterfaceData& InVertexData)
+	{
+		using namespace CommonPinNames;
+		InVertexData.BindReadVertex(METASOUND_GET_PARAM_NAME(Inputs::MidiFileAsset), MidiAssetInPin);
+		InVertexData.BindReadVertex(METASOUND_GET_PARAM_NAME(Inputs::Transport), TransportInPin);
+		InVertexData.BindReadVertex(METASOUND_GET_PARAM_NAME(Inputs::Speed), SpeedMultInPin);
+		InVertexData.BindReadVertex(METASOUND_GET_PARAM_NAME(Inputs::PrerollBars), SeekPrerollBarsInPin);
+	}
+
+	void FTempoMapMetronomeOperator::BindOutputs(FOutputVertexInterfaceData& InVertexData)
+	{
+		using namespace CommonPinNames;
+		InVertexData.BindReadVertex(METASOUND_GET_PARAM_NAME(Outputs::MidiClock), MidiClockOutPin);
+	}
+
+	FDataReferenceCollection FTempoMapMetronomeOperator::GetInputs() const
+	{
+		// This should never be called. Bind(...) is called instead. This method
+		// exists as a stop-gap until the API can be deprecated and removed.
+		checkNoEntry();
+		return {};
+	}
+
+	FDataReferenceCollection FTempoMapMetronomeOperator::GetOutputs() const
+	{
+		// This should never be called. Bind(...) is called instead. This method
+		// exists as a stop-gap until the API can be deprecated and removed.
+		checkNoEntry();
+		return {};
+	}
+
+	void FTempoMapMetronomeOperator::Reset(const FResetParams& Params)
+	{
+		BlockSize = Params.OperatorSettings.GetNumFramesPerBlock();
+		SampleRate = Params.OperatorSettings.GetSampleRate();
+
+		const FMidiFileProxyPtr& MidiProxy = MidiAssetInPin->GetMidiProxy();
+		MidiClockOutPin->AttachToMidiResource(MidiProxy ? MidiProxy->GetMidiFile() : nullptr);
+		MidiClockOutPin->ResetAndStart(0, true);
+	}
+
+	void FTempoMapMetronomeOperator::Execute()
+	{
+		MidiClockOutPin->PrepareBlock();
+		MidiClockOutPin->AddSpeedChangeToBlock({0, 0.0f, *SpeedMultInPin});
+
+		TransportSpanProcessor TransportHandler = [this](int32 StartFrameIndex, int32 EndFrameIndex, EMusicPlayerTransportState CurrentState)
+		{
+			switch (CurrentState)
+			{
+			case EMusicPlayerTransportState::Invalid:
+			case EMusicPlayerTransportState::Preparing:
+				MidiClockOutPin->AddTransportStateChangeToBlock({ StartFrameIndex, 0.0f, EMusicPlayerTransportState::Prepared });
+				MidiClockOutPin->WriteNoAdvance(StartFrameIndex, EndFrameIndex);
+				return EMusicPlayerTransportState::Prepared;
+
+			case EMusicPlayerTransportState::Prepared:
+				MidiClockOutPin->WriteNoAdvance(StartFrameIndex, EndFrameIndex);
+				return CurrentState;
+
+			case EMusicPlayerTransportState::Starting:
+				// Play from the beginning if we haven't received a seek call while we were stopped...
+				MidiClockOutPin->ResetAndStart(StartFrameIndex, !ReceivedSeekWhileStopped());
+				MidiClockOutPin->WriteAdvance(StartFrameIndex, EndFrameIndex, *SpeedMultInPin);
+				return EMusicPlayerTransportState::Playing;
+
+			case EMusicPlayerTransportState::Playing:
+				MidiClockOutPin->WriteAdvance(StartFrameIndex, EndFrameIndex, *SpeedMultInPin);
+				return EMusicPlayerTransportState::Playing;
+
+			case EMusicPlayerTransportState::Seeking:
+				MidiClockOutPin->SeekTo(StartFrameIndex, TransportInPin->GetNextSeekDestination(), *SeekPrerollBarsInPin);
+				// Here we will return that we want to be in the same state we were in before this request to 
+				// seek since we can seek "instantaneously"...
+				return GetTransportState();
+
+			case EMusicPlayerTransportState::Continuing:
+				MidiClockOutPin->AddTransportStateChangeToBlock({ StartFrameIndex, 0.0f, EMusicPlayerTransportState::Playing });
+				MidiClockOutPin->WriteAdvance(StartFrameIndex, EndFrameIndex, *SpeedMultInPin);
+				return EMusicPlayerTransportState::Playing;
+
+			case EMusicPlayerTransportState::Pausing:
+				MidiClockOutPin->AddTransportStateChangeToBlock({ StartFrameIndex, 0.0f, EMusicPlayerTransportState::Paused });
+				MidiClockOutPin->WriteNoAdvance(StartFrameIndex, EndFrameIndex);
+				return EMusicPlayerTransportState::Paused;
+
+			case EMusicPlayerTransportState::Paused:
+				MidiClockOutPin->WriteNoAdvance(StartFrameIndex, EndFrameIndex);
+				return EMusicPlayerTransportState::Paused;
+
+			case EMusicPlayerTransportState::Stopping:
+				MidiClockOutPin->AddTransportStateChangeToBlock({ StartFrameIndex, 0.0f, EMusicPlayerTransportState::Prepared });
+				MidiClockOutPin->WriteNoAdvance(StartFrameIndex, EndFrameIndex);
+				return EMusicPlayerTransportState::Prepared;
+
+			case EMusicPlayerTransportState::Killing:
+				MidiClockOutPin->AddTransportStateChangeToBlock({ StartFrameIndex, 0.0f, EMusicPlayerTransportState::Prepared });
+				MidiClockOutPin->WriteNoAdvance(StartFrameIndex, EndFrameIndex);
+				return EMusicPlayerTransportState::Prepared;
+
+			default:
+				checkNoEntry();
+				return EMusicPlayerTransportState::Invalid;
+			}
+		};
+		ExecuteTransportSpans(TransportInPin, BlockSize, TransportHandler);
+	}
+
+}
+
+#undef LOCTEXT_NAMESPACE // "HarmonixMetaSound"
diff --git a/Engine/Plugins/Runtime/Harmonix/Source/HarmonixMetasound/Private/HarmonixMetasound/Nodes/TransportWavePlayerControllerNode.cpp b/Engine/Plugins/Runtime/Harmonix/Source/HarmonixMetasound/Private/HarmonixMetasound/Nodes/TransportWavePlayerControllerNode.cpp
new file mode 100644
--- /dev/null
+++ b/Engine/Plugins/Runtime/Harmonix/Source/HarmonixMetasound/Private/HarmonixMetasound/Nodes/TransportWavePlayerControllerNode.cpp
@@ -0,0 +1,254 @@
+// Copyright Epic Games, Inc. All Rights Reserved.
+
+#include "MetasoundExecutableOperator.h"
+#include "MetasoundFacade.h"
+#include "MetasoundNodeInterface.h"
+#include "MetasoundParamHelper.h"
+#include "MetasoundSampleCounter.h"
+#include "MetasoundStandardNodesCategories.h"
+#include "MetasoundVertex.h"
+#include "HarmonixMetasound/Common.h"
+#include "HarmonixMetasound/DataTypes/MidiClock.h"
+#include "HarmonixMetasound/DataTypes/MusicTransport.h"
+
+#define LOCTEXT_NAMESPACE "HarmonixMetaSound"
+
+namespace HarmonixMetasound
+{
+	using namespace Metasound;
+
+	namespace TransportWavePlayerControllerVertexNames
+	{
+		METASOUND_PARAM(OutputStartTime, "Start Time", "Time into the wave asset to start (seek) the wave asset.")
+	}
+
+	class FTransportWavePlayerControllerOperator : public TExecutableOperator<FTransportWavePlayerControllerOperator>, public FMusicTransportControllable
+	{
+	public:
+		FTransportWavePlayerControllerOperator(const FOperatorSettings& InSettings,
+									const FMusicTransportEventStreamReadRef& InTransport,
+									const FMidiClockReadRef& InMidiClock) :
+			FMusicTransportControllable(EMusicPlayerTransportState::Prepared),
+			TransportInPin(InTransport),
+			MidiClockInPin(InMidiClock),
+			PlayOutPin(FTriggerWriteRef::CreateNew(InSettings)),
+			StopOutPin(FTriggerWriteRef::CreateNew(InSettings)),
+			StartTimeOutPin(FTimeWriteRef::CreateNew()),
+			BlockSizeFrames(InSettings.GetNumFramesPerBlock()),
+			SampleRate(InSettings.GetSampleRate())
+		{
+		}
+
+		static const FVertexInterface& GetVertexInterface()
+		{
+			using namespace CommonPinNames;
+			using namespace TransportWavePlayerControllerVertexNames;
+
+			static const FVertexInterface Interface(
+				FInputVertexInterface(
+					TInputDataVertex<FMusicTransportEventStream>(METASOUND_GET_PARAM_NAME_AND_METADATA(Inputs::Transport)),
+					TInputDataVertex<FMidiClock>(METASOUND_GET_PARAM_NAME_AND_METADATA(Inputs::MidiClock))
+				),
+				FOutputVertexInterface(
+					TOutputDataVertex<FTrigger>(METASOUND_GET_PARAM_NAME_AND_METADATA(Outputs::TransportPlay)),
+					TOutputDataVertex<FTrigger>(METASOUND_GET_PARAM_NAME_AND_METADATA(Outputs::TransportStop)),
+					TOutputDataVertex<FTime>(METASOUND_GET_PARAM_NAME_AND_METADATA(OutputStartTime))
+				)
+			);
+
+			return Interface;
+		}
+
+		static const FNodeClassMetadata& GetNodeInfo()
+		{
+			auto InitNodeInfo = []() -> FNodeClassMetadata
+			{
+				FNodeClassMetadata Info;
+				Info.ClassName = { HarmonixNodeNamespace, TEXT("TransportWavePlayerController"), TEXT("") };
+				Info.MajorVersion = 0;
+				Info.MinorVersion = 1;
+				Info.DisplayName = METASOUND_LOCTEXT("TransportWavePlayerControllerNode_DisplayName", "Music Transport Wave Player Controller");
+				Info.Description = METASOUND_LOCTEXT("TransportWavePlayerControllerNode_Description", "An interface between a music transport and a wave player");
+				Info.Author = PluginAuthor;
+				Info.PromptIfMissing = PluginNodeMissingPrompt;
+				Info.DefaultInterface = GetVertexInterface();
+				Info.CategoryHierarchy.Emplace(NodeCategories::Music);
+				return Info;
+			};
+
+			static const FNodeClassMetadata Info = InitNodeInfo();
+
+			return Info;
+		}
+
+		virtual void BindInputs(FInputVertexInterfaceData& InVertexData) override
+		{
+			using namespace CommonPinNames;
+
+			InVertexData.BindReadVertex(METASOUND_GET_PARAM_NAME(Inputs::Transport), TransportInPin);
+			InVertexData.BindReadVertex(METASOUND_GET_PARAM_NAME(Inputs::MidiClock), MidiClockInPin);
+		}
+
+		virtual void BindOutputs(FOutputVertexInterfaceData& InVertexData) override
+		{
+			using namespace CommonPinNames;
+			using namespace TransportWavePlayerControllerVertexNames;
+
+			InVertexData.BindReadVertex(METASOUND_GET_PARAM_NAME(Outputs::TransportPlay), PlayOutPin);
+			InVertexData.BindReadVertex(METASOUND_GET_PARAM_NAME(Outputs::TransportStop), StopOutPin);
+			InVertexData.BindReadVertex(METASOUND_GET_PARAM_NAME(OutputStartTime), StartTimeOutPin);
+		}
+		
+		virtual FDataReferenceCollection GetInputs() const override
+		{
+			// This should never be called. Bind(...) is called instead. This method
+			// exists as a stop-gap until the API can be deprecated and removed.
+			checkNoEntry();
+			return {};
+		}
+
+		virtual FDataReferenceCollection GetOutputs() const override
+		{
+			// This should never be called. Bind(...) is called instead. This method
+			// exists as a stop-gap until the API can be deprecated and removed.
+			checkNoEntry();
+			return {};
+		}
+
+		static TUniquePtr<IOperator> CreateOperator(const FBuildOperatorParams& InParams, FBuildResults& OutResults)
+		{
+			using namespace CommonPinNames;
+
+			const FInputVertexInterfaceData& InputData = InParams.InputData;
+			FMusicTransportEventStreamReadRef InTransport = InputData.GetOrConstructDataReadReference<FMusicTransportEventStream>(METASOUND_GET_PARAM_NAME(Inputs::Transport), InParams.OperatorSettings);
+			FMidiClockReadRef InMidiClock = InputData.GetOrConstructDataReadReference<FMidiClock>(METASOUND_GET_PARAM_NAME(Inputs::MidiClock), InParams.OperatorSettings);
+			
+			return MakeUnique<FTransportWavePlayerControllerOperator>(InParams.OperatorSettings, InTransport, InMidiClock);
+		}
+
+		void Execute()
+		{
+			// advance the outputs
+			PlayOutPin->AdvanceBlock();
+			StopOutPin->AdvanceBlock();
+
+			TransportSpanProcessor TransportHandler = [this](int32 StartFrameIndex, int32 EndFrameIndex, EMusicPlayerTransportState CurrentState)
+			{
+				switch (CurrentState)
+				{
+				case EMusicPlayerTransportState::Invalid:
+				case EMusicPlayerTransportState::Preparing:
+					return EMusicPlayerTransportState::Prepared;
+
+				case EMusicPlayerTransportState::Prepared:
+					return EMusicPlayerTransportState::Prepared;
+
+				case EMusicPlayerTransportState::Starting:
+					if (!ReceivedSeekWhileStopped())
+					{
+						// Play from the beginning if we haven't received a seek call while we were stopped...
+						*StartTimeOutPin = FTime();
+					}
+					PlayOutPin->TriggerFrame(StartFrameIndex);
+					bPlaying = true;
+					return EMusicPlayerTransportState::Playing;
+
+				case EMusicPlayerTransportState::Playing:
+					return EMusicPlayerTransportState::Playing;
+
+				case EMusicPlayerTransportState::Seeking:
+					if (ReceivedSeekWhileStopped())
+					{
+						// Assumes the MidiClock is stopped for the remainder of the block.
+						*StartTimeOutPin = FTime(MidiClockInPin->GetCurrentHiResMs() * 0.001f);
+					}
+					else
+					{
+						StopOutPin->TriggerFrame(StartFrameIndex);
+						int32 PlayFrameIndex = FMath::Min(StartFrameIndex + 1, EndFrameIndex);
+
+						// Assumes the MidiClock is playing for the remainder of the block.
+						*StartTimeOutPin = FTime(MidiClockInPin->GetCurrentHiResMs() * 0.001f - (BlockSizeFrames - PlayFrameIndex) / SampleRate);
+						PlayOutPin->TriggerFrame(PlayFrameIndex);
+					}
+					// Here we will return that we want to be in the same state we were in before this request to 
+					// seek since we can seek "instantaneously"...
+					return GetTransportState();
+
+				case EMusicPlayerTransportState::Continuing:
+					// Assumes the StartTimeOutPin won't change for the remainder of the block.
+					PlayOutPin->TriggerFrame(StartFrameIndex);
+					bPlaying = true;
+					return EMusicPlayerTransportState::Playing;
+
+				case EMusicPlayerTransportState::Pausing:
+					bPlaying = false;
+					StopOutPin->TriggerFrame(StartFrameIndex);
+
+					// Assumes the MidiClock is paused for the remainder of the block.
+					*StartTimeOutPin = FTime(MidiClockInPin->GetCurrentHiResMs() * 0.001f);
+					return EMusicPlayerTransportState::Paused;
+
+				case EMusicPlayerTransportState::Paused:
+					return EMusicPlayerTransportState::Paused;
+
+				case EMusicPlayerTransportState::Stopping:
+				case EMusicPlayerTransportState::Killing:
+					if (bPlaying)
+					{
+						bPlaying = false;
+						StopOutPin->TriggerFrame(StartFrameIndex);
+					}
+					*StartTimeOutPin = FTime();
+					return EMusicPlayerTransportState::Prepared;
+
+				default:
+					checkNoEntry();
+					return EMusicPlayerTransportState::Invalid;
+				}
+			};
+			ExecuteTransportSpans(TransportInPin, BlockSizeFrames, TransportHandler);
+		}
+
+		void Reset(const IOperator::FResetParams& InParams)
+		{
+			PlayOutPin->Reset();
+			StopOutPin->Reset();
+			*StartTimeOutPin = FTime();
+
+			BlockSizeFrames = InParams.OperatorSettings.GetNumFramesPerBlock();
+			SampleRate = InParams.OperatorSettings.GetSampleRate();
+
+			bPlaying = false;
+		}
+
+	private:
+
+		// Inputs
+		FMusicTransportEventStreamReadRef TransportInPin;
+		FMidiClockReadRef MidiClockInPin;
+
+		// Outputs
+		FTriggerWriteRef PlayOutPin;
+		FTriggerWriteRef StopOutPin;
+		FTimeWriteRef StartTimeOutPin;
+
+		int32 BlockSizeFrames;
+		float SampleRate;
+
+		bool bPlaying = false;
+	};
+
+	class FTransportWavePlayerControllerNode : public FNodeFacade
+	{
+	public:
+		FTransportWavePlayerControllerNode(const FNodeInitData& InitData)
+			: FNodeFacade(InitData.InstanceName, InitData.InstanceID, TFacadeOperatorClass<FTransportWavePlayerControllerOperator>())
+		{
+		}
+	};
+
+	METASOUND_REGISTER_NODE(FTransportWavePlayerControllerNode);
+}
+
+#undef LOCTEXT_NAMESPACE
diff --git a/Engine/Plugins/Runtime/Harmonix/Source/HarmonixMetasound/Private/HarmonixMetasound/Subsystems/MidiClockUpdateSubsystem.cpp b/Engine/Plugins/Runtime/Harmonix/Source/HarmonixMetasound/Private/HarmonixMetasound/Subsystems/MidiClockUpdateSubsystem.cpp
new file mode 100644
--- /dev/null
+++ b/Engine/Plugins/Runtime/Harmonix/Source/HarmonixMetasound/Private/HarmonixMetasound/Subsystems/MidiClockUpdateSubsystem.cpp
@@ -0,0 +1,45 @@
+// Copyright Epic Games, Inc. All Rights Reserved.
+
+#include "HarmonixMetasound/Subsystems/MidiClockUpdateSubsystem.h"
+
+#include "HarmonixMetasound/DataTypes/MidiClock.h"
+#include "Misc/ScopeLock.h"
+
+bool UMidiClockUpdateSubsystem::IsTickable() const
+{
+	FScopeLock Lock{ &TrackedClocksMutex };
+	
+	return TrackedClocks.Num() > 0;
+}
+
+void UMidiClockUpdateSubsystem::Tick(float DeltaTime)
+{
+	FScopeLock Lock{ &TrackedClocksMutex };
+
+	for (HarmonixMetasound::FMidiClock* Clock : TrackedClocks)
+	{
+		Clock->UpdateLowResCursors();
+	}
+}
+
+TStatId UMidiClockUpdateSubsystem::GetStatId() const
+{
+	RETURN_QUICK_DECLARE_CYCLE_STAT(UMidiClockUpdateSubsystem, STATGROUP_Tickables);
+}
+
+void UMidiClockUpdateSubsystem::TrackClock(HarmonixMetasound::FMidiClock* Clock)
+{
+	check(nullptr != Clock);
+	
+	FScopeLock Lock{ &TrackedClocksMutex };
+
+	TrackedClocks.AddUnique(Clock);
+}
+
+void UMidiClockUpdateSubsystem::StopTrackingClock(HarmonixMetasound::FMidiClock* Clock)
+{
+	FScopeLock Lock{ &TrackedClocksMutex };
+
+	TrackedClocks.Remove(Clock);
+}
+
diff --git a/Engine/Plugins/Runtime/Harmonix/Source/HarmonixMetasound/Public/HarmonixMetasound/Analysis/FFTAnalyzerResultVertexAnalyzer.h b/Engine/Plugins/Runtime/Harmonix/Source/HarmonixMetasound/Public/HarmonixMetasound/Analysis/FFTAnalyzerResultVertexAnalyzer.h
new file mode 100644
--- /dev/null
+++ b/Engine/Plugins/Runtime/Harmonix/Source/HarmonixMetasound/Public/HarmonixMetasound/Analysis/FFTAnalyzerResultVertexAnalyzer.h
@@ -0,0 +1,37 @@
+// Copyright Epic Games, Inc. All Rights Reserved.
+#pragma once
+
+#include "Analysis/MetasoundFrontendAnalyzerFactory.h"
+#include "Analysis/MetasoundFrontendVertexAnalyzer.h"
+
+#include "HarmonixMetasound/DataTypes/FFTAnalyzerResult.h"
+#include "MetasoundDataReference.h"
+
+namespace HarmonixMetasound::Analysis
+{
+	class HARMONIXMETASOUND_API FFFTAnalyzerResultVertexAnalyzer final : public Metasound::Frontend::FVertexAnalyzerBase
+	{
+	public:
+		static const FName& GetAnalyzerName();
+		static const FName& GetDataType();
+
+		struct HARMONIXMETASOUND_API FOutputs
+		{
+			static const Metasound::Frontend::FAnalyzerOutput& GetValue();
+		};
+
+		class HARMONIXMETASOUND_API FFactory final : public Metasound::Frontend::TVertexAnalyzerFactory<FFFTAnalyzerResultVertexAnalyzer>
+		{
+		public:
+			virtual const TArray<Metasound::Frontend::FAnalyzerOutput>& GetAnalyzerOutputs() const override;
+		};
+
+		explicit FFFTAnalyzerResultVertexAnalyzer(const Metasound::Frontend::FCreateAnalyzerParams& InParams);
+		virtual ~FFFTAnalyzerResultVertexAnalyzer() override = default;
+
+		virtual void Execute() override;
+
+	private:
+		FHarmonixFFTAnalyzerResultsWriteRef LastAnalyzerResult;
+	};
+}
diff --git a/Engine/Plugins/Runtime/Harmonix/Source/HarmonixMetasound/Public/HarmonixMetasound/Analysis/MidiClockVertexAnalyzer.h b/Engine/Plugins/Runtime/Harmonix/Source/HarmonixMetasound/Public/HarmonixMetasound/Analysis/MidiClockVertexAnalyzer.h
new file mode 100644
--- /dev/null
+++ b/Engine/Plugins/Runtime/Harmonix/Source/HarmonixMetasound/Public/HarmonixMetasound/Analysis/MidiClockVertexAnalyzer.h
@@ -0,0 +1,49 @@
+// Copyright Epic Games, Inc. All Rights Reserved.
+
+#pragma once
+
+#include "Analysis/MetasoundFrontendAnalyzerFactory.h"
+#include "Analysis/MetasoundFrontendVertexAnalyzer.h"
+
+#include "HarmonixMetasound/DataTypes/MusicTimestamp.h"
+#include "HarmonixMetasound/DataTypes/TimeSignature.h"
+
+namespace HarmonixMetasound::Analysis
+{
+	class HARMONIXMETASOUND_API FMidiClockVertexAnalyzer final : public Metasound::Frontend::FVertexAnalyzerBase
+	{
+	public:
+		struct HARMONIXMETASOUND_API FOutputs
+		{
+			/**
+			 * @brief Get the default output for this analyzer
+			 * @return The default output
+			 */
+			static const Metasound::Frontend::FAnalyzerOutput& GetValue();
+			
+			static const Metasound::Frontend::FAnalyzerOutput Timestamp;
+			static const Metasound::Frontend::FAnalyzerOutput Tempo;
+			static const Metasound::Frontend::FAnalyzerOutput TimeSignature;
+			static const Metasound::Frontend::FAnalyzerOutput Speed;
+		};
+
+		class HARMONIXMETASOUND_API FFactory final : public Metasound::Frontend::TVertexAnalyzerFactory<FMidiClockVertexAnalyzer>
+		{
+		public:
+			virtual const TArray<Metasound::Frontend::FAnalyzerOutput>& GetAnalyzerOutputs() const override;
+		};
+
+		static const FName& GetAnalyzerName();
+		static const FName& GetDataType();
+
+		explicit FMidiClockVertexAnalyzer(const Metasound::Frontend::FCreateAnalyzerParams& InParams);
+
+		virtual void Execute() override;
+
+	private:
+		FMusicTimestampWriteRef Timestamp;
+		Metasound::FFloatWriteRef Tempo;
+		FTimeSignatureWriteRef TimeSignature;
+		Metasound::FFloatWriteRef Speed;
+	};
+}
diff --git a/Engine/Plugins/Runtime/Harmonix/Source/HarmonixMetasound/Public/HarmonixMetasound/Analysis/MidiStreamVertexAnalyzer.h b/Engine/Plugins/Runtime/Harmonix/Source/HarmonixMetasound/Public/HarmonixMetasound/Analysis/MidiStreamVertexAnalyzer.h
new file mode 100644
--- /dev/null
+++ b/Engine/Plugins/Runtime/Harmonix/Source/HarmonixMetasound/Public/HarmonixMetasound/Analysis/MidiStreamVertexAnalyzer.h
@@ -0,0 +1,37 @@
+// Copyright Epic Games, Inc. All Rights Reserved.
+#pragma once
+
+#include "MetasoundDataReference.h"
+#include "Analysis/MetasoundFrontendAnalyzerFactory.h"
+#include "Analysis/MetasoundFrontendVertexAnalyzer.h"
+
+#include "HarmonixMetasound/DataTypes/MidiEventInfo.h"
+
+namespace HarmonixMetasound::Analysis
+{
+	class HARMONIXMETASOUND_API FMidiStreamVertexAnalyzer final : public Metasound::Frontend::FVertexAnalyzerBase
+	{
+	public:
+		static const FName& GetAnalyzerName();
+		static const FName& GetDataType();
+
+		struct HARMONIXMETASOUND_API FOutputs
+		{
+			static const Metasound::Frontend::FAnalyzerOutput& GetValue();
+		};
+
+		class HARMONIXMETASOUND_API FFactory final : public Metasound::Frontend::TVertexAnalyzerFactory<FMidiStreamVertexAnalyzer>
+		{
+		public:
+			virtual const TArray<Metasound::Frontend::FAnalyzerOutput>& GetAnalyzerOutputs() const override;
+		};
+
+		explicit FMidiStreamVertexAnalyzer(const Metasound::Frontend::FCreateAnalyzerParams& InParams);
+		virtual ~FMidiStreamVertexAnalyzer() override = default;
+
+		virtual void Execute() override;
+
+	private:
+		FMidiEventInfoWriteRef LastMidiEvent;
+	};
+}
diff --git a/Engine/Plugins/Runtime/Harmonix/Source/HarmonixMetasound/Public/HarmonixMetasound/Analysis/PeakTamer.h b/Engine/Plugins/Runtime/Harmonix/Source/HarmonixMetasound/Public/HarmonixMetasound/Analysis/PeakTamer.h
new file mode 100644
--- /dev/null
+++ b/Engine/Plugins/Runtime/Harmonix/Source/HarmonixMetasound/Public/HarmonixMetasound/Analysis/PeakTamer.h
@@ -0,0 +1,84 @@
+// Copyright Epic Games, Inc. All Rights Reserved.
+
+#pragma once
+
+#include "PeakTamer.generated.h"
+
+USTRUCT(BlueprintType)
+struct FHarmonixPeakTamerSettings
+{
+	GENERATED_BODY()
+
+	UPROPERTY(BlueprintReadWrite, Category="AudioAnalysis")
+	float PeakAttackTimeSeconds = 0.01f;
+
+	UPROPERTY(BlueprintReadWrite, Category="AudioAnalysis")
+	float PeakReleaseTimeSeconds = 2.0f;
+
+	UPROPERTY(BlueprintReadWrite, Category="AudioAnalysis")
+	bool bEnableValueSmoothing = true;
+
+	UPROPERTY(BlueprintReadWrite, Category="AudioAnalysis")
+	float ValueAttackTimeSeconds = 0.01f;
+	
+	UPROPERTY(BlueprintReadWrite, Category="AudioAnalysis")
+	float ValueReleaseTimeSeconds = 0.01f;
+};
+
+namespace Harmonix::AudioReactivity
+{
+	namespace PeakTamerPrivate
+	{
+		// Smooth a value given a delta time, a smoothing time, and the last value.
+		// This is pretty much a one-pole lowpass filter.
+		HARMONIXMETASOUND_API float SmoothValue(float X0, float Y1, float DeltaTime, float SmoothTime);
+	}
+	
+	/**
+	 * Takes a raw peak value from an audio analyzer and outputs a smoothed, compressed value in range [0, 1]
+	 */
+	class HARMONIXMETASOUND_API FPeakTamer final
+	{
+	public:
+		void Configure(const FHarmonixPeakTamerSettings& InSettings);
+		
+		void Update(float InputValue, float DeltaTimeSeconds);
+
+		float GetPeak() const { return Peak; }
+		
+		float GetValue() const { return Value; }
+
+	private:
+		FHarmonixPeakTamerSettings Settings{};
+		float Peak{ 0.0f };
+		float Value{ 0.0f };
+	};
+}
+
+/**
+ * Takes a raw peak value from an audio analyzer and outputs a smoothed, compressed value in range [0, 1]
+ */
+UCLASS(BlueprintType)
+class HARMONIXMETASOUND_API UHarmonixPeakTamer final : public UObject
+{
+	GENERATED_BODY()
+	
+public:
+	UFUNCTION(BlueprintCallable, Category="AudioAnalysis")
+	static UHarmonixPeakTamer* CreateHarmonixPeakTamer();
+	
+	UFUNCTION(BlueprintCallable, Category="AudioAnalysis")
+	void Configure(const FHarmonixPeakTamerSettings& Settings);
+
+	UFUNCTION(BlueprintCallable, Category="AudioAnalysis")
+	void Update(float InputValue, float DeltaTimeSeconds);
+
+	UFUNCTION(BlueprintCallable, Category="AudioAnalysis")
+	float GetPeak() const;
+	
+	UFUNCTION(BlueprintCallable, Category="AudioAnalysis")
+	float GetValue() const;
+
+private:
+	Harmonix::AudioReactivity::FPeakTamer PeakTamer;
+};
diff --git a/Engine/Plugins/Runtime/Harmonix/Source/HarmonixMetasound/Public/HarmonixMetasound/Common.h b/Engine/Plugins/Runtime/Harmonix/Source/HarmonixMetasound/Public/HarmonixMetasound/Common.h
new file mode 100644
--- /dev/null
+++ b/Engine/Plugins/Runtime/Harmonix/Source/HarmonixMetasound/Public/HarmonixMetasound/Common.h
@@ -0,0 +1,137 @@
+// Copyright Epic Games, Inc. All Rights Reserved.
+
+#pragma once
+
+#include "MetasoundParamHelper.h"
+#include "Internationalization/Internationalization.h"
+#include "Internationalization/Text.h"
+#include "UObject/NameTypes.h"
+#include "UObject/ObjectMacros.h"
+
+// TODO: move these into MetasoundParamHelper.h
+#define METASOUND_GET_PARAM_METADATA_ADVANCED(NAME) FDataVertexMetadata { NAME##Tooltip, NAME##DisplayName, true }
+#define METASOUND_GET_PARAM_NAME_AND_METADATA_ADVANCED(NAME) METASOUND_GET_PARAM_NAME(NAME), METASOUND_GET_PARAM_METADATA_ADVANCED(NAME)
+
+namespace HarmonixMetasound
+{
+	HARMONIXMETASOUND_API extern const FName HarmonixNodeNamespace;
+
+	namespace MetasoundNodeCategories
+	{
+		HARMONIXMETASOUND_API extern const FText Harmonix;
+		HARMONIXMETASOUND_API extern const FText Modulation;
+	}
+
+#define EXTERN_METASOUND_PARAM(NAME)      \
+	HARMONIXMETASOUND_API extern const TCHAR* NAME##Name;       \
+	HARMONIXMETASOUND_API extern const FText NAME##Tooltip;     \
+	HARMONIXMETASOUND_API extern const FText NAME##DisplayName; 
+
+#if WITH_EDITOR
+#define LOC_DEFINE_REGION
+#define EXTERN_METASOUND_PARAM_D(NAME, NAME_TEXT, TOOLTIP_TEXT)              \
+	const TCHAR* NAME##Name = TEXT(NAME_TEXT);                               \
+	const FText NAME##Tooltip = LOCTEXT(#NAME "Tooltip", TOOLTIP_TEXT);      \
+	const FText NAME##DisplayName = LOCTEXT(#NAME "DisplayName", NAME_TEXT); 
+#undef LOC_DEFINE_REGION
+#else 
+#define EXTERN_METASOUND_PARAM_D(NAME, NAME_TEXT, TOOLTIP_TEXT) \
+	const TCHAR* NAME##Name = TEXT(NAME_TEXT);                  \
+	const FText NAME##Tooltip = FText::GetEmpty();              \
+	const FText NAME##DisplayName = FText::GetEmpty();
+#endif // WITH_EDITOR
+
+// Alias an existing param definition. Helpful for re-using definitions while providing a public node API.
+#define EXTERN_METASOUND_PARAM_ALIAS(ALIAS_NAME, TARGET_NAME) \
+	const TCHAR* ALIAS_NAME##Name = TARGET_NAME##Name; \
+	const FText ALIAS_NAME##Tooltip = TARGET_NAME##Tooltip; \
+	const FText ALIAS_NAME##DisplayName = TARGET_NAME##DisplayName; 
+	
+	namespace CommonPinNames
+	{
+		namespace Inputs
+		{
+			EXTERN_METASOUND_PARAM(Enable)
+			EXTERN_METASOUND_PARAM(MidiStream)
+			EXTERN_METASOUND_PARAM(MidiChannelNumber)
+			EXTERN_METASOUND_PARAM(MidiChannelFilterSpecifier)
+			EXTERN_METASOUND_PARAM(MidiTrackNumber)
+			EXTERN_METASOUND_PARAM(MidiTrackIndexFilterSpecifier)
+			EXTERN_METASOUND_PARAM(MinMidiNote)
+			EXTERN_METASOUND_PARAM(MaxMidiNote)
+			EXTERN_METASOUND_PARAM(MinMidiVelocity)
+			EXTERN_METASOUND_PARAM(MaxMidiVelocity)
+			EXTERN_METASOUND_PARAM(MidiFileAsset)
+			EXTERN_METASOUND_PARAM(Transport)
+			EXTERN_METASOUND_PARAM(MidiClock)
+			EXTERN_METASOUND_PARAM(Tempo)
+			EXTERN_METASOUND_PARAM(TimeSigNumerator)
+			EXTERN_METASOUND_PARAM(TimeSigDenominator)
+			EXTERN_METASOUND_PARAM(Speed)
+			EXTERN_METASOUND_PARAM(Loop)
+			EXTERN_METASOUND_PARAM(LoopLengthBars)
+			EXTERN_METASOUND_PARAM(TransportPrepare)
+			EXTERN_METASOUND_PARAM(TransportPlay)
+			EXTERN_METASOUND_PARAM(TransportPause)
+			EXTERN_METASOUND_PARAM(TransportContinue)
+			EXTERN_METASOUND_PARAM(TransportStop)
+			EXTERN_METASOUND_PARAM(TransportKill)
+			EXTERN_METASOUND_PARAM(TriggerSeek)
+			EXTERN_METASOUND_PARAM(SeekDestination)
+			EXTERN_METASOUND_PARAM(PrerollBars)
+			EXTERN_METASOUND_PARAM(GridSizeUnits)
+			EXTERN_METASOUND_PARAM(GridSizeMult)
+			EXTERN_METASOUND_PARAM(OffsetUnits)
+			EXTERN_METASOUND_PARAM(OffsetMult)
+			EXTERN_METASOUND_PARAM(Bar)
+			EXTERN_METASOUND_PARAM(FloatBeat)
+			EXTERN_METASOUND_PARAM(SynthPatch)
+			EXTERN_METASOUND_PARAM(Transposition)
+			EXTERN_METASOUND_PARAM(ClockSpeedToPitch)
+			EXTERN_METASOUND_PARAM(ClockSpeedToFrequency)
+			EXTERN_METASOUND_PARAM(LFOSyncType);
+			EXTERN_METASOUND_PARAM(LFOFrequency);
+			EXTERN_METASOUND_PARAM(LFOInvert);
+			EXTERN_METASOUND_PARAM(LFOShape);
+			EXTERN_METASOUND_PARAM(AudioMono);
+		}
+		namespace Outputs
+		{
+			EXTERN_METASOUND_PARAM(Transport)
+			EXTERN_METASOUND_PARAM(MidiClock)
+			EXTERN_METASOUND_PARAM(MidiFileAsset)
+			EXTERN_METASOUND_PARAM(MidiStream)
+			EXTERN_METASOUND_PARAM(NoteOn)
+			EXTERN_METASOUND_PARAM(NoteOff)
+			EXTERN_METASOUND_PARAM(MidiNoteNumber)
+			EXTERN_METASOUND_PARAM(Frequency)
+			EXTERN_METASOUND_PARAM(MidiVelocity)
+			EXTERN_METASOUND_PARAM(NormalizedVelocity)
+			EXTERN_METASOUND_PARAM(TransportPrepare)
+			EXTERN_METASOUND_PARAM(TransportPlay)
+			EXTERN_METASOUND_PARAM(TransportPause)
+			EXTERN_METASOUND_PARAM(TransportContinue)
+			EXTERN_METASOUND_PARAM(TransportStop)
+			EXTERN_METASOUND_PARAM(TransportKill)
+			EXTERN_METASOUND_PARAM(MusicTimestamp)
+			EXTERN_METASOUND_PARAM(SeekTarget)
+			EXTERN_METASOUND_PARAM(AudioMono)
+			EXTERN_METASOUND_PARAM(AudioLeft)
+			EXTERN_METASOUND_PARAM(AudioRight)
+			EXTERN_METASOUND_PARAM(AudioCenter)
+			EXTERN_METASOUND_PARAM(AudioLFE)
+			EXTERN_METASOUND_PARAM(AudioLeftBack)
+			EXTERN_METASOUND_PARAM(AudioRightBack)
+			EXTERN_METASOUND_PARAM(AudioLeftSide)
+			EXTERN_METASOUND_PARAM(AudioRightSide)
+			EXTERN_METASOUND_PARAM(MusicTimespanBar)
+			EXTERN_METASOUND_PARAM(MusicTimespanBeat)
+			EXTERN_METASOUND_PARAM(Tempo)
+			EXTERN_METASOUND_PARAM(Speed)
+			EXTERN_METASOUND_PARAM(TimeSigNumerator)
+			EXTERN_METASOUND_PARAM(TimeSigDenominator)
+			EXTERN_METASOUND_PARAM(SecsPerQuarter)
+			EXTERN_METASOUND_PARAM(SecsPerBeat)
+		}
+	}
+}
diff --git a/Engine/Plugins/Runtime/Harmonix/Source/HarmonixMetasound/Public/HarmonixMetasound/Components/MetasoundOfflinePlayerComponent.h b/Engine/Plugins/Runtime/Harmonix/Source/HarmonixMetasound/Public/HarmonixMetasound/Components/MetasoundOfflinePlayerComponent.h
new file mode 100644
--- /dev/null
+++ b/Engine/Plugins/Runtime/Harmonix/Source/HarmonixMetasound/Public/HarmonixMetasound/Components/MetasoundOfflinePlayerComponent.h
@@ -0,0 +1,51 @@
+// Copyright Epic Games, Inc. All Rights Reserved.
+#pragma once
+
+#include "Components/ActorComponent.h"
+#include "MetasoundSource.h"
+
+#include "MetasoundOfflinePlayerComponent.generated.h"
+
+namespace Metasound
+{
+	class FMetasoundGeneratorHandle;
+}
+class UAudioComponent;
+class UMetasoundGeneratorHandle;
+
+UCLASS()
+class HARMONIXMETASOUND_API UMetasoundOfflinePlayerComponent : public UActorComponent
+{
+	GENERATED_BODY()
+
+public:	
+	UMetasoundOfflinePlayerComponent();
+
+	// To deprecate, only required to keep engine-only integrated backends happy
+	UMetasoundGeneratorHandle* CreateGeneratorBasedOnAudioComponent(UAudioComponent* AudioComponent, int32 InSampleRate = 12000, int32 InBlockSize = 256);
+
+	TSharedPtr<Metasound::FMetasoundGeneratorHandle> CreateSharedGeneratorBasedOnAudioComponent(UAudioComponent* AudioComponent, int32 InSampleRate = 12000, int32 InBlockSize = 256);
+	void ReleaseGenerator();
+
+	virtual void BeginPlay() override;
+	virtual void BeginDestroy() override;
+
+	TSharedPtr<Metasound::FMetasoundGenerator> GetGenerator();
+
+public:	
+	virtual void TickComponent(float DeltaTime, ELevelTick TickType, FActorComponentTickFunction* ThisTickFunction) override;
+
+private:
+	ISoundGeneratorPtr Generator;
+
+	UPROPERTY(Transient)
+	TObjectPtr<UMetaSoundSource> MetasoundSource;
+
+	int32 SampleRate;
+	int32 BlockSize;
+	int32 GeneratorBlockSize;
+
+	double StartSecond = 0.0;
+	int64  RenderedSamples = 0;
+	TArray<float> ScratchBuffer;
+};
