[Harmonix] Move Harmonix plugin into Engine/Plugins/Runtime to make it available to licensees in UE 5.4

diff --git a/Engine/Plugins/Runtime/Harmonix/Source/HarmonixMetasound/Public/HarmonixMetasound/Components/MusicClockComponent.h b/Engine/Plugins/Runtime/Harmonix/Source/HarmonixMetasound/Public/HarmonixMetasound/Components/MusicClockComponent.h
new file mode 100644
--- /dev/null
+++ b/Engine/Plugins/Runtime/Harmonix/Source/HarmonixMetasound/Public/HarmonixMetasound/Components/MusicClockComponent.h
@@ -0,0 +1,382 @@
+// Copyright Epic Games, Inc. All Rights Reserved.
+#pragma once
+
+#include "Components/ActorComponent.h"
+#include "Harmonix/MusicalTimebase.h"
+#include "HarmonixMidi/MidiSongPos.h"
+#include "Delegates/DelegateCombinations.h"
+#include "HarmonixMidi/SmoothedMidiPlayCursor.h"
+
+#include "Templates/UniquePtr.h"
+
+#include "MusicClockComponent.generated.h"
+
+DECLARE_LOG_CATEGORY_EXTERN(LogMusicClock, Log, All)
+
+namespace Metasound { class FMetasoundGenerator; }
+struct FMusicalTimeSpan;
+class UAudioComponent;
+struct FMusicClockDriverBase;
+
+UENUM(BlueprintType)
+enum class EMusicClockState : uint8
+{
+	Stopped,
+	Paused,
+	Running,
+};
+
+UENUM(BlueprintType)
+enum class EMusicClockDriveMethod : uint8
+{
+	WallClock,
+	MetaSound,
+};
+
+DECLARE_DYNAMIC_MULTICAST_DELEGATE_TwoParams(FBeatEvent, int, BeatNumber, int, BeatInBar);
+DECLARE_DYNAMIC_MULTICAST_DELEGATE_OneParam(FBarEvent, int, BarNumber);
+DECLARE_DYNAMIC_MULTICAST_DELEGATE_ThreeParams(FSectionEvent, const FString&, SectionName, float, SectionStartMs, float, SectionLengthMs);
+DECLARE_DYNAMIC_MULTICAST_DELEGATE_OneParam(FPlayStateEvent, EMusicClockState, State);
+
+DECLARE_DYNAMIC_MULTICAST_DELEGATE(FMusicClockConnected);
+DECLARE_DYNAMIC_MULTICAST_DELEGATE(FMusicClockDisconnected);
+
+UCLASS(ClassGroup = (MetaSoundMusic), PrioritizeCategories = "MusicClock", meta = (BlueprintSpawnableComponent, DisplayName = "Music Clock", ScriptName = MusicClockComponent))
+class HARMONIXMETASOUND_API UMusicClockComponent : public UActorComponent
+{
+	GENERATED_BODY()
+
+public:
+	// Sets default values for this component's properties
+	UMusicClockComponent();
+
+	void BeginPlay() override;
+	void EndPlay(const EEndPlayReason::Type EndPlayReason) override;
+
+	UPROPERTY(EditDefaultsOnly, Category = "MusicClock")
+	EMusicClockDriveMethod DriveMethod = EMusicClockDriveMethod::MetaSound;
+
+	UPROPERTY(EditDefaultsOnly, Category = "MusicClock", meta = (EditCondition = "DriveMethod == EMusicClockDriveMethod::MetaSound"))
+	FName MetasoundOutputName = "Midi Clock";
+
+	UPROPERTY(EditAnywhere, BlueprintReadWrite, Category = "MusicClock", meta = (EditCondition = "DriveMethod == EMusicClockDriveMethod::MetaSound"))
+	TObjectPtr<UAudioComponent> MetasoundsAudioComponent;
+
+	UPROPERTY(EditDefaultsOnly, Category = "MusicClock", meta = (EditCondition = "DriveMethod == EMusicClockDriveMethod::WallClock"))
+	TObjectPtr<UMidiFile> TempoMap;
+
+	UPROPERTY(EditAnywhere, BlueprintReadWrite, Category = "MusicClock|Defaults")
+	float Tempo = 120.0f;
+
+	UPROPERTY(EditAnywhere, BlueprintReadWrite, Category = "MusicClock|Defaults")
+	int TimeSignatureNum = 4;
+
+	UPROPERTY(EditAnywhere, BlueprintReadWrite, Category = "MusicClock|Defaults")
+	int TimeSignatureDenom = 4;
+
+	UPROPERTY(Transient, VisibleAnywhere, BlueprintReadOnly, Category = "MusicClock")
+	float CurrentBeatDurationSec = 0.5f;
+
+	UPROPERTY(Transient, VisibleAnywhere, BlueprintReadOnly, Category = "MusicClock")
+	float CurrentBarDurationSec = 2.0f;
+
+	UPROPERTY(Transient, VisibleAnywhere, BlueprintReadOnly, Category = "MusicClock")
+	float CurrentClockAdvanceRate = 1.0f;
+
+	UPROPERTY(EditAnywhere, BlueprintReadWrite, Category = "MusicClock")
+	bool RunPastMusicEnd = false;
+
+	UFUNCTION(BlueprintCallable, Category = "MusicClock")
+	void Start();
+
+	UFUNCTION(BlueprintCallable, Category = "MusicClock")
+	void Pause();
+
+	UFUNCTION(BlueprintCallable, Category = "MusicClock")
+	void Continue();
+
+	UFUNCTION(BlueprintCallable, Category = "MusicClock")
+	void Stop();
+
+	UFUNCTION(BlueprintPure, Category = "MusicClock")
+	EMusicClockState GetState() const { return State; }
+
+	// Getters for all of the fields in mCurrentSongPos...
+	
+	// Time from the beginning of the authored music content.
+	// NOTE: INCLUDES time for count-in and pickup bars.
+	// Note: Not const as it might cause the clock to update from its source.
+	UFUNCTION(BlueprintPure, Category = "MusicClock")
+	float GetSecondsIncludingCountIn(ECalibratedMusicTimebase Timebase = ECalibratedMusicTimebase::VideoRenderTime);
+
+	// Time from Bar 1 Beat 1. The classic "start of the song".
+	// NOTE: DOES NOT INCLUDE time for count-in and pickup bars.
+	// Note: Not const as it might cause the clock to update from its source.
+	UFUNCTION(BlueprintPure, Category = "MusicClock")
+	float GetSecondsFromBarOne(ECalibratedMusicTimebase Timebase = ECalibratedMusicTimebase::VideoRenderTime);
+
+	// Returns the fractional total bars from the beginning of the authored music content.
+	// NOTE: INCLUDES time for count-in and pickup bars.
+	// Note: Not const as it might cause the clock to update from its source.
+	UFUNCTION(BlueprintPure, Category = "MusicClock")
+	float GetBarsIncludingCountIn(ECalibratedMusicTimebase Timebase = ECalibratedMusicTimebase::VideoRenderTime);
+
+	// Returns the fractional total beats from the beginning of the authored music content.
+	// NOTE: INCLUDES time for count-in and pickup bars.
+	// Note: Not const as it might cause the clock to update from its source.
+	UFUNCTION(BlueprintPure, Category = "MusicClock")
+	float GetBeatsIncludingCountIn(ECalibratedMusicTimebase Timebase = ECalibratedMusicTimebase::VideoRenderTime); 
+
+	// Returns the "classic" musical timestamp in the form Bar (int) & Beat (float). In this form...
+	//    - Bar 1, Beat 1.0 is the "beginning of the song" AFTER count-in/pickups
+	//    - Bar 0, Beat 1.0 would be one bar BEFORE the "beginning of the song"... eg. a bar of count-in or pickup.
+	//    - While Bar can be positive or negative, Beat is always >= 1.0 and is read as "beat in the bar". Again... '1' based!
+	// Note: Not const as it might cause the clock to update from its source.
+	UFUNCTION(BlueprintPure, Category = "MusicClock")
+	FMusicTimestamp GetCurrentTimestamp(ECalibratedMusicTimebase Timebase = ECalibratedMusicTimebase::VideoRenderTime);
+
+	/** Returns the name of the section that we're currently in (intro, chorus, outro) */
+	// Note: Not const as it might cause the clock to update from its source.
+	UFUNCTION(BlueprintPure, Category = "MusicClock")
+	FString GetCurrentSectionName(ECalibratedMusicTimebase Timebase = ECalibratedMusicTimebase::VideoRenderTime);
+
+	/** Returns the index of the current section for the provided time base. [0, Num-1] */
+	// Note: Not const as it might cause the clock to update from its source.
+	UFUNCTION(BlueprintPure, Category = "MusicClock")
+	int32 GetCurrentSectionIndex(ECalibratedMusicTimebase Timebase = ECalibratedMusicTimebase::VideoRenderTime);
+
+	UFUNCTION(BlueprintPure, Category = "MusicClock")
+	const TArray<FSongSection>& GetSongSections() const;
+
+	/** Returns the start time of the current section in milliseconds */
+	// Note: Not const as it might cause the clock to update from its source.
+	UFUNCTION(BlueprintPure, Category = "MusicClock")
+	float GetCurrentSectionStartMs(ECalibratedMusicTimebase Timebase = ECalibratedMusicTimebase::VideoRenderTime);
+
+	/** Returns the length of the current section in milliseconds */
+	// Note: Not const as it might cause the clock to update from its source.
+	UFUNCTION(BlueprintPure, Category = "MusicClock")
+	float GetCurrentSectionLengthMs(ECalibratedMusicTimebase Timebase = ECalibratedMusicTimebase::VideoRenderTime);
+
+	/** Gets a value expressed in beats between 0-1 that indicates how much progress we made towards the current beat to the next one */
+	// Note: Not const as it might cause the clock to update from its source.
+	UFUNCTION(BlueprintPure, Category = "MusicClock")
+	float GetDistanceFromCurrentBeat(ECalibratedMusicTimebase Timebase = ECalibratedMusicTimebase::VideoRenderTime);
+
+	/** Gets a value expressed in beats between 0-1 that indicates how close we are to the next beat. */
+	// Note: Not const as it might cause the clock to update from its source.
+	UFUNCTION(BlueprintPure, Category = "MusicClock")
+	float GetDistanceToNextBeat(ECalibratedMusicTimebase Timebase = ECalibratedMusicTimebase::VideoRenderTime);
+
+	/** Gets a value expressed in beats between 0-1 that indicates how close we are to the closest beat (current beat or next beat). */
+	// Note: Not const as it might cause the clock to update from its source.
+	UFUNCTION(BlueprintPure, Category = "MusicClock")
+	float GetDistanceToClosestBeat(ECalibratedMusicTimebase Timebase = ECalibratedMusicTimebase::VideoRenderTime);
+
+	/** Gets a value expressed in bars between 0-1 that indicates how much progress we made towards the current bar to the next one. */
+	// Note: Not const as it might cause the clock to update from its source.
+	UFUNCTION(BlueprintPure, Category = "MusicClock")
+	float GetDistanceFromCurrentBar(ECalibratedMusicTimebase Timebase = ECalibratedMusicTimebase::VideoRenderTime);
+
+	/** Gets a value expressed in bars between 0-1 that indicates how close we are to the next bar. */
+	// Note: Not const as it might cause the clock to update from its source.
+	UFUNCTION(BlueprintPure, Category = "MusicClock")
+	float GetDistanceToNextBar(ECalibratedMusicTimebase Timebase = ECalibratedMusicTimebase::VideoRenderTime);
+
+	/** Gets value expressed in bars between 0-1 that indicates how close we are to the closest bar (current bar or next bar). */
+	// Note: Not const as it might cause the clock to update from its source.
+	UFUNCTION(BlueprintPure, Category = "MusicClock")
+	float GetDistanceToClosestBar(ECalibratedMusicTimebase Timebase = ECalibratedMusicTimebase::VideoRenderTime);
+
+	// Note: Not const as it might cause the clock to update from its source.
+	UFUNCTION(BlueprintPure, Category = "MusicClock")
+	float GetDeltaBar(ECalibratedMusicTimebase Timebase = ECalibratedMusicTimebase::VideoRenderTime);
+
+	// Note: Not const as it might cause the clock to update from its source.
+	UFUNCTION(BlueprintPure, Category = "MusicClock")
+	float GetDeltaBeat(ECalibratedMusicTimebase Timebase = ECalibratedMusicTimebase::VideoRenderTime);
+
+	// Note: Not const as it might cause the clock to update from its source.
+	UFUNCTION(BlueprintPure, Category = "MusicClock")
+	const FMidiSongPos& GetSongPos(ECalibratedMusicTimebase Timebase = ECalibratedMusicTimebase::VideoRenderTime);
+
+	UFUNCTION(BlueprintPure, Category = "Count In")
+	float GetCountInSeconds() const;
+
+	UFUNCTION(BlueprintPure, Category = "Tick")
+	float TickToMs(float Tick) const;
+
+	UFUNCTION(BlueprintPure, Category = "Beat")
+	float BeatToMs(float Beat) const;
+
+	UFUNCTION(BlueprintPure, Category = "Beat")
+	float GetMsPerBeatAtMs(float Ms) const;
+
+	UFUNCTION(BlueprintPure, Category = "Beat")
+	float GetNumBeatsInBarAtMs(float Ms) const;
+
+	UFUNCTION(BlueprintPure, Category = "Beat")
+	float GetBeatInBarAtMs(float Ms) const;
+	
+	UFUNCTION(BlueprintPure, Category = "Bar")
+	float BarToMs(float Bar) const;
+
+	UFUNCTION(BlueprintPure, Category = "Bar")
+	float GetMsPerBarAtMs(float Ms) const;
+
+	UFUNCTION(BlueprintPure, Category = "Section")
+	FString GetSectionNameAtMs(float Ms) const;
+
+	UFUNCTION(BlueprintPure, Category = "Section")
+	float GetSectionLengthMsAtMs(float Ms) const;
+
+	UFUNCTION(BlueprintPure, Category = "Section")
+	float GetSectionStartMsAtMs(float Ms) const;
+
+	UFUNCTION(BlueprintPure, Category = "Section")
+	float GetSectionEndMsAtMs(float Ms) const;
+
+	UFUNCTION(BlueprintPure, Category = "Section")
+	float GetNumSections() const;
+
+	UFUNCTION(BlueprintPure, Category = "Song Data")
+	float GetSongLengthMs() const;
+
+	UFUNCTION(BlueprintPure, Category = "Song Data")
+	float GetSongLengthBeats() const;
+
+	UFUNCTION(BlueprintPure, Category = "Song Data")
+	int32 GetSongLengthBars() const;
+
+	/** Returns the remaining time until the end of the MIDI in milliseconds based on the timestamp corresponding to the passed Timebase */
+	// Note: Not const as it might cause the clock to update from its source.
+	UFUNCTION(BlueprintPure, Category = "MusicClock")
+	float GetSongRemainingTime(ECalibratedMusicTimebase Timebase = ECalibratedMusicTimebase::VideoRenderTime);
+
+	const FSongMaps& GetSongMaps() const;
+
+	UPROPERTY(EditAnywhere, BlueprintReadWrite, Category = "MusicClock")
+	ECalibratedMusicTimebase TimebaseForBarAndBeatEvents = ECalibratedMusicTimebase::VideoRenderTime;
+
+	UPROPERTY(BlueprintAssignable, Category = "MusicClock")
+	FPlayStateEvent PlayStateEvent;
+
+	UPROPERTY(BlueprintAssignable, Category = "MusicClock")
+	FBeatEvent BeatEvent;
+
+	UPROPERTY(BlueprintAssignable, Category = "MusicClock")
+	FBarEvent BarEvent;
+
+	UPROPERTY(BlueprintAssignable, Category = "MusicClock")
+	FSectionEvent SectionEvent;
+
+	UPROPERTY(BlueprintAssignable, Category = "MusicClock")
+	FMusicClockConnected MusicClockConnectedEvent;
+
+	UPROPERTY(BlueprintAssignable, Category = "MusicClock")
+	FMusicClockDisconnected MusicClockDisconnectedEvent;
+
+	UPROPERTY(BlueprintGetter = GetCurrentSmoothedAudioRenderSongPos, Category = "MusicClock")
+	FMidiSongPos CurrentSmoothedAudioRenderSongPos;
+
+	UFUNCTION(BlueprintCallable, BlueprintPure, Category = "MusicClock")
+	FMidiSongPos GetCurrentSmoothedAudioRenderSongPos();
+
+	UPROPERTY(BlueprintGetter = GetCurrentVideoRenderSongPos, Category = "MusicClock")
+	FMidiSongPos CurrentVideoRenderSongPos;
+
+	UFUNCTION(BlueprintCallable, BlueprintPure, Category = "MusicClock")
+	FMidiSongPos GetCurrentVideoRenderSongPos();
+
+	UPROPERTY(BlueprintGetter = GetCurrentPlayerExperiencedSongPos, Category = "MusicClock")
+	FMidiSongPos CurrentPlayerExperiencedSongPos;
+
+	UFUNCTION(BlueprintCallable, BlueprintPure, Category = "MusicClock")
+	FMidiSongPos GetCurrentPlayerExperiencedSongPos();
+
+	// Note: Not const as it might cause the clock to update from its source.
+	UFUNCTION(BlueprintPure, Category = "MusicClock")
+	float MeasureSpanProgress(const FMusicalTimeSpan& Span, ECalibratedMusicTimebase Timebase = ECalibratedMusicTimebase::VideoRenderTime);
+
+	UFUNCTION(BlueprintCallable, Category = "Audio|MusicClock", meta = (WorldContext = "WorldContextObject", DefaultToSelf = "WorldContextObject"))
+	static UMusicClockComponent* CreateMetasoundDrivenMusicClock(UObject* WorldContextObject, UAudioComponent* InAudioComponent, FName MetasoundOuputPinName = "Midi Clock", bool Start = true);
+
+	UFUNCTION(BlueprintCallable, Category = "Audio|MusicClock", meta = (WorldContext = "WorldContextObject", DefaultToSelf = "WorldContextObject"))
+	static UMusicClockComponent* CreateWallClockDrivenMusicClock(UObject* WorldContextObject, UMidiFile* WithTempoMap, bool Start = true);
+
+	UFUNCTION(BlueprintCallable, Category = "MusicClock")
+	bool ConnectToMetasoundOnAudioComponent(UAudioComponent* InAudioComponent);
+
+	UFUNCTION(BlueprintCallable, Category = "MusicClock")
+	void ConnectToWallClockForMidi(UMidiFile* InTempoMap);
+
+	// Note: Not const as it might cause the clock to update from its source.
+	FMidiSongPos CalculateSongPosWithOffset(float MsOffset, ECalibratedMusicTimebase Timebase = ECalibratedMusicTimebase::VideoRenderTime);
+
+	const FMidiSongPos& GetRawUnsmoothedAudioRenderPos() const { return RawUnsmoothedAudioRenderPos; }
+
+protected:
+	virtual void TickComponent(float DeltaTime, enum ELevelTick TickType, FActorComponentTickFunction* ThisTickFunction);
+
+private:
+	friend struct FMusicClockDriverBase;
+	friend struct FMetasoundMusicClockDriver;
+	friend struct FWallClockMusicClockDriver;
+
+	TUniquePtr<FMusicClockDriverBase> ClockDriver;
+
+	EMusicClockState State = EMusicClockState::Stopped;
+
+	FSongMaps DefaultMaps;
+	FMidiSongPos RawUnsmoothedAudioRenderPos;
+
+	float     AudioRenderDeltaBarF = 0.0f;
+	float     AudioRenderDeltaBeatF = 0.0f;
+	float     PlayerExperienceDeltaBarF = 0.0f;
+	float     PlayerExperienceDeltaBeatF = 0.0f;
+	float     VideoRenderDeltaBarF = 0.0f;
+	float     VideoRenderDeltaBeatF = 0.0f;
+	int32     LastBroadcastBar = -1;
+	int32     LastBroadcastBeat = -1;
+	FSongSection LastBroadcastSongSection;
+
+	FMidiSongPos PrevAudioRenderSongPos;
+	FMidiSongPos PrevPlayerExperiencedSongPos;
+	FMidiSongPos PrevVideoRenderSongPos;
+
+	uint64    LastUpdateFrame = 0;
+
+	void CreateClockDriver();
+	void CalcDelta();
+	void Broadcast();
+	void MakeDefaultSongMap();
+	bool ConnectToMetasound();
+	void ConnectToWallClock();
+
+	// Ensures the clock will be updated once per frame.  Should only get called on the game thread.
+	void UpdateClockIfNeeded();
+};
+
+struct FMusicClockDriverBase
+{
+public:
+	FMusicClockDriverBase() = delete;
+	FMusicClockDriverBase(UMusicClockComponent* DrivenClock)
+		: Clock(DrivenClock)
+	{}
+	virtual ~FMusicClockDriverBase() = default;
+
+	virtual bool CalculateSongPosWithOffset(float MsOffset, ECalibratedMusicTimebase Timebase, FMidiSongPos& OutResult) const = 0;
+
+	virtual void Disconnect() = 0;
+	virtual bool RefreshCurrentSongPos() = 0;
+	virtual void OnStart() = 0;
+	virtual void OnPause() = 0;
+	virtual void OnContinue() = 0;
+	virtual void OnStop() = 0;
+	virtual const FSongMaps* GetCurrentSongMaps() const = 0;
+
+protected:
+	UMusicClockComponent* Clock;
+};
+
diff --git a/Engine/Plugins/Runtime/Harmonix/Source/HarmonixMetasound/Public/HarmonixMetasound/DataTypes/MidiClock.h b/Engine/Plugins/Runtime/Harmonix/Source/HarmonixMetasound/Public/HarmonixMetasound/DataTypes/MidiClock.h
new file mode 100644
--- /dev/null
+++ b/Engine/Plugins/Runtime/Harmonix/Source/HarmonixMetasound/Public/HarmonixMetasound/DataTypes/MidiClock.h
@@ -0,0 +1,294 @@
+// Copyright Epic Games, Inc. All Rights Reserved.
+
+#pragma once
+
+#include "Tickable.h"
+
+#include "MetasoundDataReference.h"
+#include "MetasoundEnumRegistrationMacro.h"
+#include "MetasoundExecutableOperator.h"
+#include "MetasoundOperatorSettings.h"
+#include "MetasoundVariable.h"
+#include "MetasoundSampleCounter.h"
+
+#include "HarmonixMidi/MidiFile.h"
+#include "HarmonixMidi/SongMaps.h"
+#include "HarmonixMidi/MidiFile.h"
+#include "HarmonixMetasound/DataTypes/MusicTransport.h"
+#include "HarmonixMetasound/DataTypes/MidiClockEvent.h"
+#include "HarmonixMidi/MidiPlayCursorMgr.h"
+#include "HarmonixMidi/MidiPlayCursor.h"
+#include "Sound/QuartzQuantizationUtilities.h"
+
+enum class EMidiClockSubdivisionQuantization : int32
+{
+	Bar						= static_cast<int32>(EQuartzCommandQuantization::Bar),
+	Beat					= static_cast<int32>(EQuartzCommandQuantization::Beat),
+	ThirtySecondNote		= static_cast<int32>(EQuartzCommandQuantization::ThirtySecondNote),
+	SixteenthNote			= static_cast<int32>(EQuartzCommandQuantization::SixteenthNote),
+	EighthNote				= static_cast<int32>(EQuartzCommandQuantization::EighthNote),
+	QuarterNote				= static_cast<int32>(EQuartzCommandQuantization::QuarterNote),
+	HalfNote				= static_cast<int32>(EQuartzCommandQuantization::HalfNote),
+	WholeNote				= static_cast<int32>(EQuartzCommandQuantization::WholeNote),
+	DottedSixteenthNote		= static_cast<int32>(EQuartzCommandQuantization::DottedSixteenthNote),
+	DottedEighthNote		= static_cast<int32>(EQuartzCommandQuantization::DottedEighthNote),
+	DottedQuarterNote		= static_cast<int32>(EQuartzCommandQuantization::DottedQuarterNote),
+	DottedHalfNote			= static_cast<int32>(EQuartzCommandQuantization::DottedHalfNote),
+	DottedWholeNote			= static_cast<int32>(EQuartzCommandQuantization::DottedWholeNote),
+	SixteenthNoteTriplet	= static_cast<int32>(EQuartzCommandQuantization::SixteenthNoteTriplet),
+	EighthNoteTriplet		= static_cast<int32>(EQuartzCommandQuantization::EighthNoteTriplet),
+	QuarterNoteTriplet		= static_cast<int32>(EQuartzCommandQuantization::QuarterNoteTriplet),
+	HalfNoteTriplet			= static_cast<int32>(EQuartzCommandQuantization::HalfNoteTriplet),
+};
+
+namespace Metasound
+{
+	DECLARE_METASOUND_ENUM(
+		EMidiClockSubdivisionQuantization,
+		EMidiClockSubdivisionQuantization::Beat,
+		HARMONIXMETASOUND_API,
+		FEnumMidiClockSubdivisionQuantizationType,
+		FEnumMidiClockSubdivisionQuantizationTypeInfo,
+		FEnumMidiClockSubdivisionQuantizationReadRef,
+		FEnumMidiClockSubdivisionQuantizationWriteRef
+	);
+}
+
+int32 SubdivisionToMidiTicks(EMidiClockSubdivisionQuantization Division, int32 CurrentTick, const FSongMaps& SongMap);
+
+namespace HarmonixMetasound
+{
+	using namespace Metasound;
+
+	struct HARMONIXMETASOUND_API FMidiTimestampTransportState
+	{
+		int32 BlockSampleFrameIndex = 0;
+		float BlockSampleFrameOffset = 0.0f;
+		EMusicPlayerTransportState TransportState = EMusicPlayerTransportState::Invalid;
+	};
+
+	struct HARMONIXMETASOUND_API FMidiTimestampSpeed
+	{
+		int32 BlockSampleFrameIndex = 0;
+		float BlockSampleFrameOffset = 0.0f;
+		float Speed = 1.0f;
+	};
+
+	struct HARMONIXMETASOUND_API FMidiTimestampTempo
+	{
+		int32 BlockSampleFrameIndex = 0;
+		float BlockSampleFrameOffset = 0.0f;
+		float Tempo = 120.0f;
+	};
+
+	constexpr FMidiTimestampTempo InvalidMidiTimestampTempo{ -1, 0.0f, 0.0f };
+
+	class HARMONIXMETASOUND_API FMidiClock
+	{
+	public:
+		explicit FMidiClock(const FOperatorSettings& InSettings);
+		FMidiClock(const FMidiClock& Other);
+		FMidiClock(FMidiClock&& Other);
+
+		FMidiClock& operator=(const FMidiClock& Other);
+
+		virtual ~FMidiClock();
+
+		void ResetAndStart(int32 FrameIndex, bool SeekToStart = true);
+
+		void PrepareBlock();
+
+		bool HasLowResCursors() const;
+		void UpdateLowResCursors();
+
+		void AddTransportStateChangeToBlock(const FMidiTimestampTransportState& NewTransportState);
+		void AddSpeedChangeToBlock(const FMidiTimestampSpeed& NewSpeed);
+
+		bool HasTransportStateChangeInBlock() const
+		{
+			return TransportChangesInBlock.IsEmpty() == false;
+		}
+		bool HasSpeedChangesInBlock() const
+		{
+			return HasSpeedChangeInBlock;
+		}
+		bool HasTempoChangesInBlock() const
+		{
+			return HasTempoChangeInBlock;
+		}
+
+		const TArray<FMidiClockEvent>& GetMidiClockEventsInBlock() const;
+
+		EMusicPlayerTransportState GetTransportStateAtBlockSampleFrame(int32 FrameIndex) const;
+		EMusicPlayerTransportState GetTransportStateAtEndOfBlock() const;
+		const FMidiTimestampTransportState& GetTransportTimestampForBlockSampleFrame(int32 FrameIndex) const;
+		const TArray<FMidiTimestampTransportState>& GetTransportTimestampsInBlock() const
+		{
+			return TransportChangesInBlock;
+		}
+
+		float GetSpeedAtBlockSampleFrame(int32 FrameIndex) const;
+		float GetSpeedAtEndOfBlock() const;
+		const FMidiTimestampSpeed& GetSpeedTimestampForBlockSampleFrame(int32 FrameIndex) const;
+		const TArray<FMidiTimestampSpeed>& GetTempoSpeedTimestampsInBlock() const
+		{
+			return SpeedChangesInBlock;
+		}
+
+		float GetTempoAtBlockSampleFrame(int32 FrameIndex) const;
+		float GetTempoAtEndOfBlock() const;
+		int32 GetNumTempoChangesInBlock() const;
+		FMidiTimestampTempo GetTempoChangeByIndex(int32 Index) const;
+
+		int32 GetCurrentMidiTick() const;
+
+		int32 GetCurrentBlockFrameIndex() const;
+
+		float GetQuarterNoteIncludingCountIn() const;
+
+		/**
+		 * @brief Get the timestamp after the most recent clock update
+		 * @return The current timestamp
+		 */
+		FMusicTimestamp GetCurrentMusicTimestamp() const;
+		
+		/**
+		 * @brief Get the music timestamp at a given frame offset from the last processed audio block.
+		 * @param Offset - The frame index from the beginning of the last processed audio block
+		 * @return The music timestamp
+		 */
+		FMusicTimestamp GetMusicTimestampAtBlockOffset(int32 Offset) const;
+
+		/**
+		 * @brief Get the absolute time in ms for a frame within the last audio block
+		 * @param Offset - The frame index from the beginning of the last processed audio block
+		 * @return The absolute time in ms
+		 */
+		float GetMsAtBlockOffset(int32 Offset) const;
+
+		//*****************************************************************************************
+		// NOTE: These next functions are a facade in front of this class's MidiPlayCursorMgr. 
+		void AttachToMidiResource(TSharedPtr<FMidiFileData> MidiDataProxy, bool ResetCursorsToStart = true, int32 PreRollBars = 0) { DrivingMidiPlayCursorMgr->AttachToMidiResource(MidiDataProxy, ResetCursorsToStart, PreRollBars); }
+		void DetachFromMidiResource() { DrivingMidiPlayCursorMgr->DetachFromMidiResource(); }
+		void AttachToTimeAuthority(const FMidiClock& MidiClockRef);
+		void DetachFromTimeAuthority();
+		void InformOfCurrentAdvanceRate(float AdvanceRate);
+		
+		void LockForMidiDataChanges();
+		void MidiDataChangesComplete(FMidiPlayCursorMgr::EMidiChangePositionCorrectMode Mode = FMidiPlayCursorMgr::EMidiChangePositionCorrectMode::MaintainTick);
+		const FSongMaps& GetSongMaps() const     { return DrivingMidiPlayCursorMgr->GetSongMaps();   }
+		const FTempoMap& GetTempoMap() const     { return DrivingMidiPlayCursorMgr->GetTempoMap();   }
+		const FBarMap& GetBarMap() const         { return DrivingMidiPlayCursorMgr->GetBarMap();     }
+		void AdvanceHiResToMs(int32 BlockFrameIndex, float Ms, bool Broadcast);
+		void SeekTo(const FMusicSeekTarget& Timestamp, int32 PreRollBars);
+
+
+		void SetLoop(int32 StartTick, int32 EndTick) { DrivingMidiPlayCursorMgr->SetLoop(StartTick, EndTick, false, true); }
+		void ClearLoop() { DrivingMidiPlayCursorMgr->ClearLoop(true); }
+		bool DoesLoop() const { return DrivingMidiPlayCursorMgr->DoesLoop(); }
+		float GetLoopStartMs() const { return DrivingMidiPlayCursorMgr->GetLoopStartMs(); }
+		int32 GetLoopStartTick() const { return DrivingMidiPlayCursorMgr->GetLoopStartTick(); }
+		float GetLoopEndMs() const { return DrivingMidiPlayCursorMgr->GetLoopEndMs(); }
+		int32 GetLoopEndTick() const { return DrivingMidiPlayCursorMgr->GetLoopEndTick(); }
+
+		int32 GetCurrentHiResTick() const  { return DrivingMidiPlayCursorMgr->GetCurrentHiResTick();  }
+		float GetCurrentHiResMs() const    { return DrivingMidiPlayCursorMgr->GetCurrentHiResMs();    }
+		float GetElapsedHiResMs() const    { return DrivingMidiPlayCursorMgr->GetElapsedHiResMs();    }
+		int32 GetCurrentLowResTick() const { return DrivingMidiPlayCursorMgr->GetCurrentLowResTick(); }
+		float GetCurrentLowResMs() const   { return DrivingMidiPlayCursorMgr->GetCurrentLowResMs();   }
+		float GetElapsedLowResMs() const   { return DrivingMidiPlayCursorMgr->GetElapsedLowResMs();   }
+		// These next few are const because they need to be callable by things that hold a ReadRef to us!
+		void RegisterHiResPlayCursor(FMidiPlayCursor* PlayCursor, float PreRollMs = -1.0f) const  { DrivingMidiPlayCursorMgr->RegisterHiResPlayCursor(PlayCursor, PreRollMs);  }
+		void RegisterLowResPlayCursor(FMidiPlayCursor* PlayCursor, float PreRollMs = -1.0f) const { DrivingMidiPlayCursorMgr->RegisterLowResPlayCursor(PlayCursor, PreRollMs); }
+		void UnregisterPlayCursor(FMidiPlayCursor* PlayCursor, bool WarnOnFail = true) const      { DrivingMidiPlayCursorMgr->UnregisterPlayCursor(PlayCursor, WarnOnFail);    }
+		void UnregisterAllPlayCursors()                                                           { DrivingMidiPlayCursorMgr->UnregisterAllPlayCursors();                      }
+		//*****************************************************************************************
+
+		void WriteAdvance(int32 StartFrameIndex, int32 EndFrameIndex, float InSpeed = 1.0f);
+		void WriteNoAdvance(int32 StartFrameIndex, int32 EndFrameIndex);
+		void SeekTo(int32 BlockFrameIndex, const FMusicSeekTarget& InTarget, int32 InPrerollBars);
+		// copy speed and tempo changes from in clock to this clock
+		// with an optional Speed multiplier to adjust the out going speed on this clock
+		void CopySpeedAndTempoChanges(const FMidiClock* InClock, float InSpeedMult = 1.0f);
+
+		// Creates a new FMidiFileData with the given starting Tempo and Time Signature
+		// With max song length to be played indefinitely. Useful for Midi Clock Metronomes.
+		static TSharedPtr<FMidiFileData> MakeClockConductorMidiData(float TempoBPM, int32 TimeSigNum, int32 TimeSigDen);
+
+	private:
+		friend class FMidiClockEventCursor;
+		friend class FTempoChangesCursor;
+
+		class FTempoChangesCursor : public FMidiPlayCursor
+		{
+		public:
+			FTempoChangesCursor();
+
+			FTempoChangesCursor(FMidiClock* MidiClock);
+
+			//** BEGIN FMidiPlayCursor
+			void OnTempo(int32 TrackIndex, int32 Tick, int32 Tempo, bool IsPreroll = false) override;
+			//** END FMidiPlayCursor
+
+		private:
+			FMidiClock* MyMidiClock = nullptr;
+		};
+
+		FTempoChangesCursor TempoChangesCursor;
+
+	private:
+		void RegisterForGameThreadUpdates();
+		void UnregisterForGameThreadUpdates();
+		
+		int32 BlockSize;
+		int32 CurrentBlockFrameIndex;
+		float SampleRate;
+		FSampleCount SampleCount;
+		int32 FramesUntilNextProcess = 0;
+		static const int32 kMidiGranularity = 128;
+		FMidiTimestampTransportState CurrentTransportState;
+
+		TArray<FMidiTimestampTransportState> TransportChangesInBlock;
+		bool								 HasSpeedChangeInBlock;
+		TArray<FMidiTimestampSpeed>          SpeedChangesInBlock;
+		bool								 HasTempoChangeInBlock;
+		TArray<FMidiTimestampTempo>          TempoChangesInBlock;
+		
+		// midi clock events are modified by a MidiClockEventCursor
+		// which has to be manually registered to the MidiClock
+		// (this is to avoid always needlessly generating midi clock events)
+		TArray<FMidiClockEvent> MidiClockEventsInBlock;
+
+		bool SmoothingEnabled = false;
+		mutable TSharedPtr<FMidiPlayCursorMgr> DrivingMidiPlayCursorMgr;
+		bool IOwnThePlayCursorMgr = false;
+	};
+
+	// Declare aliases IN the namespace...
+	DECLARE_METASOUND_DATA_REFERENCE_ALIAS_TYPES(FMidiClock, FMidiClockTypeInfo, FMidiClockReadRef, FMidiClockWriteRef)
+
+
+	class HARMONIXMETASOUND_API FMidiClockEventCursor : public FMidiPlayCursor
+	{
+	public:
+
+		FMidiClockEventCursor(FMidiClockWriteRef InClock);
+		virtual ~FMidiClockEventCursor();
+
+		//~ BEGIN FMidiPlayCursor Overrides
+		virtual void Reset(bool ForceNoBroadcast = false) override;
+		virtual void OnLoop(int32 LoopStartTick, int32 LoopEndTick) override;
+		virtual void SeekToTick(int32 Tick) override;
+		virtual void SeekThruTick(int32 Tick) override;
+		virtual void AdvanceThruTick(int32 Tick, bool IsPreRoll) override;
+		//~ END FMidiPlayCursor Overrides
+	private:
+
+		void AddEvent(const FMidiClockEvent& InEvent);
+
+		FMidiClockWriteRef MidiClock;
+	};
+}
+
+// Declare reference types OUT of the namespace...
+DECLARE_METASOUND_DATA_REFERENCE_TYPES_NO_ALIASES(HarmonixMetasound::FMidiClock, HARMONIXMETASOUND_API)
diff --git a/Engine/Plugins/Runtime/Harmonix/Source/HarmonixMetasound/Public/HarmonixMetasound/DataTypes/MidiClockEvent.h b/Engine/Plugins/Runtime/Harmonix/Source/HarmonixMetasound/Public/HarmonixMetasound/DataTypes/MidiClockEvent.h
new file mode 100644
--- /dev/null
+++ b/Engine/Plugins/Runtime/Harmonix/Source/HarmonixMetasound/Public/HarmonixMetasound/DataTypes/MidiClockEvent.h
@@ -0,0 +1,39 @@
+// Copyright Epic Games, Inc. All Rights Reserved.
+
+#pragma once
+
+#include "HarmonixMidi/MidiPlayCursor.h"
+
+namespace HarmonixMetasound
+{
+	struct HARMONIXMETASOUND_API FMidiClockEvent
+	{
+		enum class EType : uint8
+		{
+			Reset,
+			Loop,
+			SeekTo,
+			SeekThru,
+			AdvanceThru
+		};
+
+		const EType Type;
+		const bool  IsPreRoll;
+		const bool  ForceNoBroadcast;
+		const int32 BlockFrameIndex;
+		const int32 StartTick;
+		const int32 EndTick;
+
+
+		static FMidiClockEvent MakeResetEvent(int32 InBlockFrameIndex, int32 Tick, bool ForceNoBroadcast = false);
+		static FMidiClockEvent MakeLoopEvent(int32 InBlockFrameIndex, int32 InStartTick, int32 InEndTick);
+		static FMidiClockEvent MakeSeekToEvent(int32 InBlockFrameIndex, int32 Tick);
+		static FMidiClockEvent MakeSeekThruEvent(int32 InBlockFrameIndex, int32 Tick);
+		static FMidiClockEvent MakeAdvanceThruEvent(int32 InBlockFrameIndex, int32 Tick, bool IsPreRoll);
+
+	private:
+
+		FMidiClockEvent(EType InType, int32 InBlockFrameIndex, int32 InStartTick, int32 InEndTick, bool InIsPreRoll, bool InBroadcastEvents);
+	};
+
+};
\ No newline at end of file
diff --git a/Engine/Plugins/Runtime/Harmonix/Source/HarmonixMetasound/Public/HarmonixMetasound/DataTypes/MidiStream.h b/Engine/Plugins/Runtime/Harmonix/Source/HarmonixMetasound/Public/HarmonixMetasound/DataTypes/MidiStream.h
new file mode 100644
--- /dev/null
+++ b/Engine/Plugins/Runtime/Harmonix/Source/HarmonixMetasound/Public/HarmonixMetasound/DataTypes/MidiStream.h
@@ -0,0 +1,420 @@
+// Copyright Epic Games, Inc. All Rights Reserved.
+
+#pragma once
+
+#include "HarmonixMetasound/DataTypes/MusicTransport.h"
+#include "HarmonixMidi/MidiFile.h"
+#include "HarmonixMidi/MidiMsg.h"
+#include "HarmonixMidi/MidiVoiceId.h"
+#include "Logging/LogMacros.h"
+#include "MetasoundDataReference.h"
+#include "MetasoundOperatorSettings.h"
+#include "MetasoundVariable.h"
+
+#include "MidiClock.h"
+
+#include <functional>
+#include <atomic>
+
+namespace HarmonixMetasound
+{
+	using namespace Metasound;
+
+	DECLARE_LOG_CATEGORY_EXTERN(LogMidiStreamDataType, Log, All)
+
+	struct HARMONIXMETASOUND_API FMidiStreamEvent
+	{
+		int32        BlockSampleFrameIndex  = 0;
+		float        BlockSampleFrameOffset = 0.0f;
+		int32        TrackIndex             = 0;
+		int32        AuthoredMidiTick       = 0;
+		int32        CurrentMidiTick        = 0;
+		float        MsOffset               = 0.0f;
+		FMidiMsg     MidiMessage;
+
+		FMidiStreamEvent(const FMidiVoiceGeneratorBase* Owner, const FMidiMsg& Message)
+			: MidiMessage(Message)
+		{
+			VoiceId = FMidiVoiceId(Owner, MidiMessage);
+		}
+
+		FMidiStreamEvent(uint32 OwnerId, const FMidiMsg& Message)
+			: MidiMessage(Message)
+		{
+			VoiceId = FMidiVoiceId(OwnerId, MidiMessage);
+		}
+
+		FMidiVoiceId GetVoiceId() const { return VoiceId; }
+		void SetVoiceId(FMidiVoiceId InVoiceId) { VoiceId = InVoiceId; }
+
+	private:
+		FMidiVoiceId VoiceId;
+	};
+
+	struct HARMONIXMETASOUND_API FMidiStreamEventTrackChannelFilter
+	{
+		uint16 MidiChannelFilter;
+		uint64 TrackFilter;
+
+		/** 
+		* Parses the FilterString to a Bitfield and assigns it to the MidiChannelFilter
+		* if it fails, will assign MidiChannelFilter to 0
+		*/
+		bool SetMidiChannelFilterFromString(const FString& FilterString, FString& OutErrorMessage);
+
+		/**
+		* Parses the FilterString to a Bitfield and assigns it to the TrackFilter
+		* if it fails, will assign MidiChannelFilter to 0
+		*/
+		bool SetTrackFilterFromString(const FString& FilterString, FString& OutErrorMessage);
+
+		bool operator()(const FMidiStreamEvent& Event) const
+		{
+			if (!(TrackFilter & (uint64(1) << Event.TrackIndex)))
+			{
+				return false;
+			}
+
+			return !Event.MidiMessage.IsStd() || (MidiChannelFilter & (uint16(1) << Event.MidiMessage.GetStdChannel()));
+		}
+
+	private:
+		uint64 StringToBitfield(const FString& InString, FString& OutErrorMesage);
+	};
+
+
+	class HARMONIXMETASOUND_API FMidiStream
+	{
+	public:
+
+		FMidiStream(const FOperatorSettings& InSettings);
+
+		void SetClockSource(const FMidiClockReadRef& ClockSource)
+		{
+			MidiClockSource.Emplace(ClockSource);
+		}
+
+		using MidiSinkFunction = std::function<void()>;
+		void ExecuteBlock(MidiSinkFunction Destination) const
+		{
+		}
+
+		void PrepareBlock();
+
+		void AddTransportStateChangeMessage(int32 SampleFrameIndexInBlock, EMusicPlayerTransportState State);
+		void AddMidiEvent(const FMidiStreamEvent& Event);
+		void InsertMidiEvent(const FMidiStreamEvent& Event);
+		void AddNoteOffEventOrCancelPendingNoteOn(const FMidiStreamEvent& Event);
+		void InsertNoteOffEventOrCancelPendingNoteOn(const FMidiStreamEvent& Event);
+
+		void SetMidiFile(const FMidiFileProxyPtr& MidiFile) { MidiFileSourceOfEvents = MidiFile; }
+		const FString* GetMidiTrackText(int32 TrackNumber, int32 TextIndex) const;
+
+		const TArray<FMidiTimestampTransportState>& GetTransportChangesInBlock() const
+		{
+			return TransportChangesInBlock;
+		}
+		const TArray<FMidiStreamEvent>& GetEventsInBlock() const
+		{
+			return EventsInBlock;
+		}
+
+		const FMidiClockReadRef* GetMidiClockSource() const 
+		{
+			return MidiClockSource.GetPtrOrNull();
+		}
+
+		// Forward declare the read ref for copying other streams
+		using FMidiStreamReadRef = TDataReadReference<FMidiStream>;
+
+		// Copies all midi and transport events from the in stream to this stream. Obliterates any
+		// transport/midi events already in this stream!
+		void Copy(FMidiStreamReadRef InStream, bool EventsOnly = false);
+
+		// Copies all midi and transport events from the in streams to this stream. Obliterates any
+		// transport/midi events already in this stream!
+		void Copy(const TArray<FMidiStreamReadRef>& InStreams, bool EventsOnly = false);
+
+		// Copies all transport events, and midi events that match the predicate, from the in stream to this stream.
+		// Obliterates any transport/midi events already in this stream!
+		template<typename T>
+		void Copy(FMidiStreamReadRef InStream, const T& Predicate, bool EventsOnly = false)
+		{
+			if (GetMidiClockSource() && InStream->GetMidiClockSource() && GetMidiClockSource()->Get() != InStream->GetMidiClockSource()->Get())
+			{
+				UE_LOG(LogMidiStreamDataType, Warning, TEXT("Cannot copy midi events from one stream to another if they have different clocks."));
+				return;
+			}
+			if (!EventsOnly)
+			{
+				CopyTransportEvents(InStream);
+			}
+			CopyMidiEvents(InStream, Predicate);
+		}
+
+		// Copies all transport events, and midi events that match the predicate, from the in streams to this stream.
+		// Obliterates any transport/midi events already in this stream!
+		template<typename T>
+		void Copy(const TArray<FMidiStreamReadRef>& InStreams, const T& Predicate, bool EventsOnly = false)
+		{
+			if (!InStreams.Num())
+			{
+				return;
+			}
+
+			if (InStreams.Num() == 1)
+			{
+				Copy(InStreams[0], Predicate);
+				return;
+			}
+
+			TArray<FMidiStreamReadRef> SameClockStreams;
+			FilterArrayToStreamsWithTheSameClock(InStreams, SameClockStreams);
+			if (!EventsOnly)
+			{
+				CopyTransportEvents(SameClockStreams);
+			}
+			CopyMidiEvents(SameClockStreams, Predicate);
+		}
+
+		// Copies all midi and transport events from the in stream to this stream that match the predicate. Obliterates any
+		// transport/midi events already in this stream!
+		template<typename FILTER, typename TRANSFORMER>
+		void Copy(FMidiStreamReadRef InStream, const FILTER& Predicate, const TRANSFORMER& Transformer, bool EventsOnly = false)
+		{
+			if (GetMidiClockSource() && InStream->GetMidiClockSource() && GetMidiClockSource()->Get() != InStream->GetMidiClockSource()->Get())
+			{
+				UE_LOG(LogMidiStreamDataType, Warning, TEXT("Cannot copy midi events from one stream to another if they have different clocks."));
+				return;
+			}
+			if (!EventsOnly)
+			{
+				CopyTransportEvents(InStream);
+			}
+			CopyMidiEvents(InStream, Predicate, Transformer);
+		}
+
+		int32 GetTicksPerQuarterNote() const { return TicksPerQuarterNote; }
+		void SetTicksPerQuarterNote(int32 InTicksPerQuarterNote) { TicksPerQuarterNote = InTicksPerQuarterNote; }
+
+		// Returns an array of only the in streams that have the same clock as this stream
+		void FilterArrayToStreamsWithTheSameClock(const TArray<FMidiStreamReadRef>& InStreams, TArray<FMidiStreamReadRef>& OutStreams);
+		// Copies all transport events from the in streams to this. Obliterates any existing transport events.
+		void CopyTransportEvents(const TArray<FMidiStreamReadRef>& InStreams);
+		// Copies all transport events from the in stream to this. Obliterates any existing transport events.
+		void CopyTransportEvents(const FMidiStreamReadRef& InStream);
+		// Copies all midi events from the in streams to this. Obliterates and existing midi events.
+		void CopyMidiEvents(const TArray<FMidiStreamReadRef>& InStreams);
+		// Merges all midi events from the in streams to this.
+		void MergeMidiEvents(const TArray<FMidiStreamReadRef>& InStreams);
+		// Copies all midi events from the in stream to this. Obliterates and existing midi events.
+		void CopyMidiEvents(const FMidiStreamReadRef& InStream);
+		// Merges all midi events from the in stream to this.
+		void MergeMidiEvents(const FMidiStreamReadRef& InStream);
+
+		// Copies all midi events that match the predicate from the in stream to this. Obliterates and existing midi events.
+		template<typename T>
+		void CopyMidiEvents(const FMidiStreamReadRef& InStream, const T& Predicate)
+		{
+			const TArray<FMidiStreamEvent>& MidiEvents = InStream->GetEventsInBlock();
+			if (!ensureAlwaysMsgf(EventsInBlock.Num()==0, TEXT("Existing midi events in midi stream are being destroyed during copy!")))
+			{
+				EventsInBlock.Empty(FMath::Max(MidiEvents.Num(), 32));
+				ActiveVoices.Reset();
+			}
+			for (const FMidiStreamEvent& Event : MidiEvents)
+			{
+				if (Predicate(Event))
+				{
+					AddMidiEvent(Event);
+				}
+			}
+		}
+
+		// Merges all midi events that match the predicate from the in stream to this. 
+		template<typename T>
+		void MergeMidiEvents(const FMidiStreamReadRef& InStream, const T& Predicate)
+		{
+			const TArray<FMidiStreamEvent>& MidiEvents = InStream->GetEventsInBlock();
+			for (const FMidiStreamEvent& Event : MidiEvents)
+			{
+				if (Predicate(Event))
+				{
+					InsertMidiEvent(Event);
+				}
+			}
+		}
+
+		// Copies all midi events that match the predicate from the in stream to this. Obliterates and existing midi events.
+		template<typename T>
+		void CopyMidiEvents(const TArray<FMidiStreamReadRef>& InStreams, const T& Predicate)
+		{
+			if (!ensureAlwaysMsgf(EventsInBlock.Num() == 0, TEXT("Existing midi events in midi stream are being destroyed during copy!")))
+			{
+				EventsInBlock.Empty(32);
+				ActiveVoices.Reset();
+			}
+
+			// Define the type so that we can store iterators in an array
+			using MidiStreamIterator = TArray<FMidiStreamEvent>::RangedForConstIteratorType;
+
+			TArray<MidiStreamIterator> WalkingIterators;
+			TArray<MidiStreamIterator> IteratorEnds;
+
+			for (int32 i = 0; i < InStreams.Num(); ++i)
+			{
+				WalkingIterators.Add(InStreams[i]->GetEventsInBlock().begin());
+				IteratorEnds.Add(InStreams[i]->GetEventsInBlock().end());
+			}
+
+			while (1)
+			{
+				int32 StreamWithEarliestEvent = -1;
+				int32 EarliestMidiTick = MAX_int32;
+
+				for (int32 IteratorIndex = 0; IteratorIndex < WalkingIterators.Num(); ++IteratorIndex)
+				{
+					if (WalkingIterators[IteratorIndex] != IteratorEnds[IteratorIndex])
+					{
+						if ((*WalkingIterators[IteratorIndex]).AuthoredMidiTick < EarliestMidiTick)
+						{
+							EarliestMidiTick = (*WalkingIterators[IteratorIndex]).AuthoredMidiTick;
+							StreamWithEarliestEvent = IteratorIndex;
+						}
+					}
+				}
+
+				// done?
+				if (StreamWithEarliestEvent == -1)
+				{
+					break;
+				}
+
+				if (Predicate(*WalkingIterators[StreamWithEarliestEvent]))
+				{
+					AddMidiEvent(*WalkingIterators[StreamWithEarliestEvent]);
+				}
+				++WalkingIterators[StreamWithEarliestEvent];
+			}
+		}
+
+		// Merges all midi events that match the predicate from the in stream to this.
+		template<typename T>
+		void MergeMidiEvents(const TArray<FMidiStreamReadRef>& InStreams, const T& Predicate)
+		{
+			// Define the type so that we can store iterators in an array
+			using MidiStreamIterator = TArray<FMidiStreamEvent>::RangedForConstIteratorType;
+
+			TArray<MidiStreamIterator> WalkingIterators;
+			TArray<MidiStreamIterator> IteratorEnds;
+
+			for (int32 i = 0; i < InStreams.Num(); ++i)
+			{
+				WalkingIterators.Add(InStreams[i]->GetEventsInBlock().begin());
+				IteratorEnds.Add(InStreams[i]->GetEventsInBlock().end());
+			}
+
+			while (1)
+			{
+				int32 StreamWithEarliestEvent = -1;
+				int32 EarliestBlockSampleFrameIndex = MAX_int32;
+
+				for (int32 IteratorIndex = 0; IteratorIndex < WalkingIterators.Num(); ++IteratorIndex)
+				{
+					if (WalkingIterators[IteratorIndex] != IteratorEnds[IteratorIndex])
+					{
+						if ((*WalkingIterators[IteratorIndex]).BlockSampleFrameIndex == EarliestBlockSampleFrameIndex)
+						{
+							if ((*WalkingIterators[StreamWithEarliestEvent]).MidiMessage.IsNoteOn() && (*WalkingIterators[IteratorIndex]).MidiMessage.IsNoteOff())
+							{
+								StreamWithEarliestEvent = IteratorIndex;
+							}
+						}
+						else if ((*WalkingIterators[IteratorIndex]).BlockSampleFrameIndex < EarliestBlockSampleFrameIndex)
+						{
+							EarliestBlockSampleFrameIndex = (*WalkingIterators[IteratorIndex]).BlockSampleFrameIndex;
+							StreamWithEarliestEvent = IteratorIndex;
+						}
+					}
+				}
+
+				// done?
+				if (StreamWithEarliestEvent == -1)
+				{
+					break;
+				}
+
+				if (Predicate(*WalkingIterators[StreamWithEarliestEvent]))
+				{
+					InsertMidiEvent(*WalkingIterators[StreamWithEarliestEvent]);
+				}
+				++WalkingIterators[StreamWithEarliestEvent];
+			}
+		}
+
+		// Copies all midi events that match the predicate from the in stream to this. Optionally transforms the events. Obliterates and existing midi events.
+		template<typename FILTER, typename TRANSFORMER>
+		void CopyMidiEvents(const FMidiStreamReadRef& InStream, const FILTER& Predicate, const TRANSFORMER& Transformer)
+		{
+			const TArray<FMidiStreamEvent>& MidiEvents = InStream->GetEventsInBlock();
+			EventsInBlock.Empty(FMath::Max(MidiEvents.Num(), 32));
+			ActiveVoices.Reset();
+			
+			for (const FMidiStreamEvent& Event : MidiEvents)
+			{
+				if (Predicate(Event))
+				{
+					AddMidiEvent(Transformer(Event));
+				}
+			}
+		}
+
+		// Copies all midi events that match the predicate from the in stream to this. Optionally transforms the events. Obliterates and existing midi events.
+		template<typename FILTER, typename TRANSFORMER>
+		void MergeMidiEvents(const FMidiStreamReadRef& InStream, const FILTER& Predicate, const TRANSFORMER& Transformer)
+		{
+			const TArray<FMidiStreamEvent>& MidiEvents = InStream->GetEventsInBlock();
+			for (const FMidiStreamEvent& Event : MidiEvents)
+			{
+				if (Predicate(Event))
+				{
+					InsertMidiEvent(Transformer(Event));
+				}
+			}
+		}
+
+	private:
+		FMidiFileProxyPtr MidiFileSourceOfEvents;
+		int32 NumFramesPerBlock = 0;
+		FSampleRate SampleRate  = 0;
+		int32 TicksPerQuarterNote = MidiConstants::kTicksPerQuarterNoteInt;
+
+		FMidiTimestampTransportState CurrentTransportState;
+		TArray<FMidiTimestampTransportState> TransportChangesInBlock;
+
+		TArray<FMidiStreamEvent> EventsInBlock;
+
+		friend class FMidiVoiceTracker;
+		void UpdateActiveVoice(const FMidiStreamEvent& Event);
+		TSet<FMidiVoiceId> ActiveVoices;
+
+		TOptional<FMidiClockReadRef> MidiClockSource;
+	};
+
+	// Declare aliases IN the namespace...
+	DECLARE_METASOUND_DATA_REFERENCE_ALIAS_TYPES(FMidiStream, FMidiStreamTypeInfo, FMidiStreamReadRef, FMidiStreamWriteRef)
+
+	class HARMONIXMETASOUND_API FMidiVoiceTracker
+	{
+	public:
+		using FKillVoiceFn = TFunctionRef<void(const FMidiVoiceId&)>;
+
+		void Process(const FMidiStream& MidiStream, const FKillVoiceFn& KillVoiceFn);
+
+	private:
+		TSet<FMidiVoiceId> ActiveVoices;
+	};
+}
+
+// Declare reference types OUT of the namespace...
+DECLARE_METASOUND_DATA_REFERENCE_TYPES_NO_ALIASES(HarmonixMetasound::FMidiStream, HARMONIXMETASOUND_API)
diff --git a/Engine/Plugins/Runtime/Harmonix/Source/HarmonixMetasound/Public/HarmonixMetasound/Nodes/DelayNode.h b/Engine/Plugins/Runtime/Harmonix/Source/HarmonixMetasound/Public/HarmonixMetasound/Nodes/DelayNode.h
new file mode 100644
--- /dev/null
+++ b/Engine/Plugins/Runtime/Harmonix/Source/HarmonixMetasound/Public/HarmonixMetasound/Nodes/DelayNode.h
@@ -0,0 +1,44 @@
+// Copyright Epic Games, Inc. All Rights Reserved.
+
+#pragma once
+
+#include "HarmonixMetasound/Common.h"
+
+namespace HarmonixMetasound::DelayNode
+{
+	namespace Constants
+	{
+		inline constexpr int32 NumChannels = 2;
+		inline constexpr float MaxDelayTime = 2500.0f;
+	}
+
+	namespace Inputs
+	{
+		EXTERN_METASOUND_PARAM(AudioLeft);
+		EXTERN_METASOUND_PARAM(AudioRight);
+		EXTERN_METASOUND_PARAM(MidiClock)
+		EXTERN_METASOUND_PARAM(DelayTimeType);
+		EXTERN_METASOUND_PARAM(DelayTime);
+		EXTERN_METASOUND_PARAM(Feedback);
+		EXTERN_METASOUND_PARAM(DryLevel);
+		EXTERN_METASOUND_PARAM(WetLevel);
+		EXTERN_METASOUND_PARAM(WetFilterEnabled);
+		EXTERN_METASOUND_PARAM(FeedbackFilterEnabled);
+		EXTERN_METASOUND_PARAM(FilterType);
+		EXTERN_METASOUND_PARAM(FilterCutoff);
+		EXTERN_METASOUND_PARAM(FilterQ);
+		EXTERN_METASOUND_PARAM(LFOEnabled);
+		EXTERN_METASOUND_PARAM(LFOTimeType);
+		EXTERN_METASOUND_PARAM(LFOFrequency);
+		EXTERN_METASOUND_PARAM(LFODepth);
+		EXTERN_METASOUND_PARAM(StereoType);
+		EXTERN_METASOUND_PARAM(StereoSpreadLeft);
+		EXTERN_METASOUND_PARAM(StereoSpreadRight);
+	}
+
+	namespace Outputs
+	{
+		EXTERN_METASOUND_PARAM(AudioLeft);
+		EXTERN_METASOUND_PARAM(AudioRight);
+	}
+}
diff --git a/Engine/Plugins/Runtime/Harmonix/Source/HarmonixMetasound/Public/HarmonixMetasound/Nodes/DjFilterNode.h b/Engine/Plugins/Runtime/Harmonix/Source/HarmonixMetasound/Public/HarmonixMetasound/Nodes/DjFilterNode.h
new file mode 100644
--- /dev/null
+++ b/Engine/Plugins/Runtime/Harmonix/Source/HarmonixMetasound/Public/HarmonixMetasound/Nodes/DjFilterNode.h
@@ -0,0 +1,29 @@
+// Copyright Epic Games, Inc. All Rights Reserved.
+
+#pragma once
+
+#include "MetasoundNodeInterface.h"
+
+#include "HarmonixMetasound/Common.h"
+
+namespace HarmonixMetasound::Nodes::DjFilter
+{
+	const HARMONIXMETASOUND_API Metasound::FNodeClassName& GetClassName();
+	
+	namespace Inputs
+	{
+		EXTERN_METASOUND_PARAM(AudioMono);
+		EXTERN_METASOUND_PARAM(Amount);
+		EXTERN_METASOUND_PARAM(Resonance);
+		EXTERN_METASOUND_PARAM(LowPassMinFrequency);
+		EXTERN_METASOUND_PARAM(LowPassMaxFrequency);
+		EXTERN_METASOUND_PARAM(HighPassMinFrequency);
+		EXTERN_METASOUND_PARAM(HighPassMaxFrequency);
+		EXTERN_METASOUND_PARAM(DeadZoneSize);
+	}
+
+	namespace Outputs
+	{
+		EXTERN_METASOUND_PARAM(AudioMono);
+	}
+}
diff --git a/Engine/Plugins/Runtime/Harmonix/Source/HarmonixMetasound/Public/HarmonixMetasound/Nodes/MidiCCTriggerNode.h b/Engine/Plugins/Runtime/Harmonix/Source/HarmonixMetasound/Public/HarmonixMetasound/Nodes/MidiCCTriggerNode.h
new file mode 100644
--- /dev/null
+++ b/Engine/Plugins/Runtime/Harmonix/Source/HarmonixMetasound/Public/HarmonixMetasound/Nodes/MidiCCTriggerNode.h
@@ -0,0 +1,28 @@
+// Copyright Epic Games, Inc. All Rights Reserved.
+
+#pragma once
+
+#include "MetasoundDataReference.h"
+#include "MetasoundNodeInterface.h"
+#include "HarmonixMetasound/Common.h"
+
+namespace HarmonixMetasound::Nodes::MidiCCTriggerNode
+{
+	const HARMONIXMETASOUND_API Metasound::FNodeClassName& GetClassName();
+
+	namespace Inputs
+	{
+		EXTERN_METASOUND_PARAM(Enable);
+		EXTERN_METASOUND_PARAM(MidiTrackNumber);
+		EXTERN_METASOUND_PARAM(MidiChannelNumber);
+		EXTERN_METASOUND_PARAM(InputMidiControllerID);
+		EXTERN_METASOUND_PARAM(MidiStream);
+	}
+
+	namespace Outputs
+	{
+		EXTERN_METASOUND_PARAM(OutputControlChangeValueInt32);
+		EXTERN_METASOUND_PARAM(OutputControlChangeValueFloat);
+		EXTERN_METASOUND_PARAM(OutputTrigger);
+	}
+}
\ No newline at end of file
diff --git a/Engine/Plugins/Runtime/Harmonix/Source/HarmonixMetasound/Public/HarmonixMetasound/Nodes/MultibandAnalyzerNode.h b/Engine/Plugins/Runtime/Harmonix/Source/HarmonixMetasound/Public/HarmonixMetasound/Nodes/MultibandAnalyzerNode.h
new file mode 100644
--- /dev/null
+++ b/Engine/Plugins/Runtime/Harmonix/Source/HarmonixMetasound/Public/HarmonixMetasound/Nodes/MultibandAnalyzerNode.h
@@ -0,0 +1,39 @@
+// Copyright Epic Games, Inc. All Rights Reserved.
+
+#pragma once
+
+#include "MetasoundEnvelopeFollowerTypes.h"
+#include "MetasoundNodeInterface.h"
+
+#include "HarmonixMetasound/Common.h"
+
+namespace HarmonixMetasound::Nodes::MultibandAnalyzer
+{
+	const HARMONIXMETASOUND_API Metasound::FNodeClassName& GetClassName();
+
+	struct FSettings
+	{
+		bool Enable{ true };
+		TArray<float> CrossoverFrequencies{ 200, 400, 800 };
+		bool ApplySmoothing{ true };
+		Metasound::FTime AttackTime = Metasound::FTime::FromMilliseconds(10);
+		Metasound::FTime ReleaseTime = Metasound::FTime::FromMilliseconds(10);
+		Metasound::EEnvelopePeakMode PeakMode = Metasound::EEnvelopePeakMode::Peak;
+	};
+	
+	namespace Inputs
+	{
+		EXTERN_METASOUND_PARAM(Enable);
+		EXTERN_METASOUND_PARAM(AudioMono);
+		EXTERN_METASOUND_PARAM(CrossoverFrequencies);
+		EXTERN_METASOUND_PARAM(ApplySmoothing);
+		EXTERN_METASOUND_PARAM(AttackTime);
+		EXTERN_METASOUND_PARAM(ReleaseTime);
+		EXTERN_METASOUND_PARAM(PeakMode);
+	}
+
+	namespace Outputs
+	{
+		EXTERN_METASOUND_PARAM(BandLevels);
+	}
+}
diff --git a/Engine/Plugins/Runtime/Harmonix/Source/HarmonixMetasound/Public/HarmonixMetasound/Subsystems/MidiClockUpdateSubsystem.h b/Engine/Plugins/Runtime/Harmonix/Source/HarmonixMetasound/Public/HarmonixMetasound/Subsystems/MidiClockUpdateSubsystem.h
new file mode 100644
--- /dev/null
+++ b/Engine/Plugins/Runtime/Harmonix/Source/HarmonixMetasound/Public/HarmonixMetasound/Subsystems/MidiClockUpdateSubsystem.h
@@ -0,0 +1,41 @@
+// Copyright Epic Games, Inc. All Rights Reserved.
+
+#pragma once
+
+#include "Tickable.h"
+#include "Subsystems/EngineSubsystem.h"
+
+#include "MidiClockUpdateSubsystem.generated.h"
+
+namespace HarmonixMetasound
+{
+	class FMidiClock;
+}
+
+/**
+ * Handles updating low-resolution play cursors associated with a MIDI clock.
+ * Long story short, because FMidiClock instances do not have their lifecycle managed by the garbage collector,
+ * we need a way to tick them on the game thread while avoiding races. This gives us a way to register clocks to be
+ * ticked from within their constructors/destructor, and provides no user-facing API.
+ */
+UCLASS()
+class HARMONIXMETASOUND_API UMidiClockUpdateSubsystem final : public UEngineSubsystem, public FTickableGameObject
+{
+	GENERATED_BODY()
+
+public:
+	// Begin FTickableGameObject
+	virtual bool IsTickable() const override;
+	virtual void Tick(float DeltaTime) override;
+	virtual TStatId GetStatId() const override;
+	// End FTickableGameObject
+
+private:
+	friend class HarmonixMetasound::FMidiClock;
+	
+	void TrackClock(HarmonixMetasound::FMidiClock* Clock);
+	void StopTrackingClock(HarmonixMetasound::FMidiClock* Clock);
+
+	mutable FCriticalSection TrackedClocksMutex;
+	TArray<HarmonixMetasound::FMidiClock*> TrackedClocks;
+};
diff --git a/Engine/Plugins/Runtime/Harmonix/Source/HarmonixMetasoundEditor/Private/HarmonixMetasoundEditorModule.cpp b/Engine/Plugins/Runtime/Harmonix/Source/HarmonixMetasoundEditor/Private/HarmonixMetasoundEditorModule.cpp
new file mode 100644
--- /dev/null
+++ b/Engine/Plugins/Runtime/Harmonix/Source/HarmonixMetasoundEditor/Private/HarmonixMetasoundEditorModule.cpp
@@ -0,0 +1,52 @@
+// Copyright Epic Games, Inc. All Rights Reserved.
+
+#include "HarmonixMetasoundEditorModule.h"
+#include "Modules/ModuleManager.h"
+#include "Logging/LogMacros.h"
+#include "MetasoundEditorModule.h"
+#include "HarmonixMetasoundSlateStyle.h"
+#include "MidiStepSequenceDetailCustomization.h"
+#include "HarmonixMetasound/DataTypes/MidiStepSequence.h"
+#include "AssetDefinition_MidiStepSequence.h"
+
+#define LOCTEXT_NAMESPACE "HarmonixMetaSound"
+
+DEFINE_LOG_CATEGORY(LogHarmonixMetasoundEditor)
+
+void FHarmonixMetasoundEditorModule::StartupModule()
+{
+	const HarmonixMetasoundEditor::FSlateStyle& Style = HarmonixMetasoundEditor::FSlateStyle::Get();
+	Metasound::Editor::IMetasoundEditorModule& MetasoundEditorModule = FModuleManager::GetModuleChecked<Metasound::Editor::IMetasoundEditorModule>("MetasoundEditor");
+	MetasoundEditorModule.RegisterPinType("MidiAsset");
+	MetasoundEditorModule.RegisterPinType("MidiStepSequenceAsset");
+	MetasoundEditorModule.RegisterPinType("FusionPatchAsset");
+	MetasoundEditorModule.RegisterPinType("Enum:SubdivisionQuantizationType","Int32");
+	MetasoundEditorModule.RegisterPinType("Enum:DelayFilterType", "Int32");
+	MetasoundEditorModule.RegisterPinType("Enum:DelayStereoType", "Int32");
+	MetasoundEditorModule.RegisterPinType("Enum:TimeSyncOption", "Int32");
+	MetasoundEditorModule.RegisterPinType("Enum:DistortionType", "Int32");
+	MetasoundEditorModule.RegisterPinType("Enum:Harmonix:BiquadFilterType", "Int32");
+	MetasoundEditorModule.RegisterPinType("Enum:Distortion:FilterPasses", "Int32");
+	MetasoundEditorModule.RegisterPinType("Enum:StdMidiControllerID", "Int32");
+
+	MetasoundEditorModule.RegisterPinType("MidiStream", {}, {}, Style.GetMidiStreamConnectedIcon(), Style.GetMidiStreamDisconnectedIcon());
+	MetasoundEditorModule.RegisterPinType("MidiClock", {}, {}, Style.GetMidiClockConnectedIcon(), Style.GetMidiClockDisconnectedIcon());
+	MetasoundEditorModule.RegisterPinType("MusicTransport", {}, {}, Style.GetTransportConnectedIcon(), Style.GetTransportDisconnectedIcon());
+
+	FPropertyEditorModule& PropertyModule = FModuleManager::LoadModuleChecked<FPropertyEditorModule>("PropertyEditor");
+	PropertyModule.RegisterCustomClassLayout(
+		UMidiStepSequence::StaticClass()->GetFName(),
+		FOnGetDetailCustomizationInstance::CreateStatic(&FMidiStepSequenceDetailCustomization::MakeInstance)
+	);
+
+	PropertyModule.NotifyCustomizationModuleChanged();
+}
+
+void FHarmonixMetasoundEditorModule::ShutdownModule()
+{
+
+}
+
+IMPLEMENT_MODULE(FHarmonixMetasoundEditorModule, HarmonixMetasoundEditor);
+
+#undef LOCTEXT_NAMESPACE
diff --git a/Engine/Plugins/Runtime/Harmonix/Source/HarmonixMetasoundEditor/Private/HarmonixMetasoundSlateStyle.cpp b/Engine/Plugins/Runtime/Harmonix/Source/HarmonixMetasoundEditor/Private/HarmonixMetasoundSlateStyle.cpp
new file mode 100644
--- /dev/null
+++ b/Engine/Plugins/Runtime/Harmonix/Source/HarmonixMetasoundEditor/Private/HarmonixMetasoundSlateStyle.cpp
@@ -0,0 +1,68 @@
+// Copyright Epic Games, Inc. All Rights Reserved.
+
+#include "HarmonixMetasoundSlateStyle.h"
+#include "Styling/SlateStyleRegistry.h"
+
+namespace HarmonixMetasoundEditor
+{
+	FSlateStyle::FSlateStyle()
+		: FSlateStyleSet("HarmonixMetasoundSlateStyle")
+	{
+		const FVector2D Icon22x22(22.0f, 22.0f);
+		const FVector2D Icon18x10(18.0f, 10.0f);
+		const FVector2D Icon18x18(18.0f, 18.0f);
+		SetContentRoot(FPaths::EngineDir() / TEXT("Restricted/NotForLicensees/Plugins/Harmonix/Content/Editor"));
+
+#define IMAGE_BRUSH( RelativePath, ... ) FSlateImageBrush( RootToContentDir( RelativePath, TEXT(".png") ), __VA_ARGS__ )
+		Set("HarmonixMetasound.MidiConnectedIcon", new IMAGE_BRUSH(TEXT("Icons/MidiConnectedPin"), Icon22x22));
+		Set("HarmonixMetasound.MidiDisconnectedIcon", new IMAGE_BRUSH(TEXT("Icons/MidiDisconnectedPin"), Icon22x22));
+		Set("HarmonixMetasound.ClockConnectedIcon", new IMAGE_BRUSH(TEXT("Icons/ClockConnectedPin"), Icon18x18));
+		Set("HarmonixMetasound.ClockDisconnectedIcon", new IMAGE_BRUSH(TEXT("Icons/ClockDisconnectedPin"), Icon18x18));
+		Set("HarmonixMetasound.TransportConnectedIcon", new IMAGE_BRUSH(TEXT("Icons/TransportConnectedPin"), Icon18x10));
+		Set("HarmonixMetasound.TransportDisconnectedIcon", new IMAGE_BRUSH(TEXT("Icons/TransportDisconnectedPin"), Icon18x10));
+#undef IMAGE_BRUSH
+
+		FSlateStyleRegistry::RegisterSlateStyle(*this);
+	}
+	
+	const FSlateStyle& FSlateStyle::Get()
+	{
+		static FSlateStyle SlateStyle;
+		return SlateStyle;
+	}
+
+	const FSlateBrush* FSlateStyle::GetMidiStreamConnectedIcon() const
+	{
+		return GetBrush("HarmonixMetasound.MidiConnectedIcon");
+	}
+
+	const FSlateBrush* FSlateStyle::GetMidiStreamDisconnectedIcon() const
+	{
+		return GetBrush("HarmonixMetasound.MidiDisconnectedIcon");
+	}
+
+	const FSlateBrush* FSlateStyle::GetMidiClockConnectedIcon() const
+	{
+		return GetBrush("HarmonixMetasound.ClockConnectedIcon");
+	}
+
+	const FSlateBrush* FSlateStyle::GetMidiClockDisconnectedIcon() const
+	{
+		return GetBrush("HarmonixMetasound.ClockDisconnectedIcon");
+	}
+
+	const FSlateBrush* FSlateStyle::GetTransportConnectedIcon() const
+	{
+		return GetBrush("HarmonixMetasound.TransportConnectedIcon");
+	}
+
+	const FSlateBrush* FSlateStyle::GetTransportDisconnectedIcon() const
+	{
+		return GetBrush("HarmonixMetasound.TransportDisconnectedIcon");
+	}
+
+	FSlateStyle::~FSlateStyle()
+	{
+		FSlateStyleRegistry::UnRegisterSlateStyle(*this);
+	}
+}
diff --git a/Engine/Plugins/Runtime/Harmonix/Source/HarmonixMetasoundEditor/Private/HarmonixMetasoundSlateStyle.h b/Engine/Plugins/Runtime/Harmonix/Source/HarmonixMetasoundEditor/Private/HarmonixMetasoundSlateStyle.h
new file mode 100644
--- /dev/null
+++ b/Engine/Plugins/Runtime/Harmonix/Source/HarmonixMetasoundEditor/Private/HarmonixMetasoundSlateStyle.h
@@ -0,0 +1,23 @@
+// Copyright Epic Games, Inc. All Rights Reserved.
+
+#pragma once
+#include "Styling/SlateStyle.h"
+
+namespace HarmonixMetasoundEditor
+{
+	class FSlateStyle final
+		: public FSlateStyleSet
+	{
+	public:
+		FSlateStyle();
+		static const FSlateStyle& Get();
+		const FSlateBrush* GetMidiStreamConnectedIcon() const;
+		const FSlateBrush* GetMidiStreamDisconnectedIcon() const;
+		const FSlateBrush* GetMidiClockConnectedIcon() const;
+		const FSlateBrush* GetMidiClockDisconnectedIcon() const;
+		const FSlateBrush* GetTransportConnectedIcon() const;
+		const FSlateBrush* GetTransportDisconnectedIcon() const;
+
+		~FSlateStyle();
+	};
+}
diff --git a/Engine/Plugins/Runtime/Harmonix/Source/HarmonixMetasoundEditor/Private/MidiStepSequenceFactory.h b/Engine/Plugins/Runtime/Harmonix/Source/HarmonixMetasoundEditor/Private/MidiStepSequenceFactory.h
new file mode 100644
--- /dev/null
+++ b/Engine/Plugins/Runtime/Harmonix/Source/HarmonixMetasoundEditor/Private/MidiStepSequenceFactory.h
@@ -0,0 +1,16 @@
+// Copyright Epic Games, Inc. All Rights Reserved.
+
+#pragma once
+
+#include "Factories/Factory.h"
+#include "MidiStepSequenceFactory.generated.h"
+
+UCLASS()
+class UMidiStepSequenceFactory : public UFactory
+{
+	GENERATED_BODY()
+public:
+	UMidiStepSequenceFactory();
+	UObject* FactoryCreateNew(UClass* Class, UObject* InParent, FName Name, EObjectFlags Flags, UObject* Context, FFeedbackContext* Warn);
+};
+
diff --git a/Engine/Plugins/Runtime/Harmonix/Source/HarmonixMetasoundTests/Private/DataTypes/MidiClock.spec.cpp b/Engine/Plugins/Runtime/Harmonix/Source/HarmonixMetasoundTests/Private/DataTypes/MidiClock.spec.cpp
new file mode 100644
--- /dev/null
+++ b/Engine/Plugins/Runtime/Harmonix/Source/HarmonixMetasoundTests/Private/DataTypes/MidiClock.spec.cpp
@@ -0,0 +1,105 @@
+// Copyright Epic Games, Inc. All Rights Reserved.
+
+#include "HarmonixMetasound/DataTypes/MidiClock.h"
+
+#include "Misc/AutomationTest.h"
+
+#if WITH_DEV_AUTOMATION_TESTS
+
+namespace HarmonixMetasoundTests::MidiClock
+{
+	BEGIN_DEFINE_SPEC(
+		FHarmonixMetasoundMidiClockSpec,
+		"Harmonix.Metasound.DataTypes.MidiClock",
+		EAutomationTestFlags::EditorContext | EAutomationTestFlags::EngineFilter)
+
+	TUniquePtr<HarmonixMetasound::FMidiClock> TestClock;
+
+	void AddStateAtFrame(HarmonixMetasound::EMusicPlayerTransportState State, int32 Frame) const
+	{
+		const HarmonixMetasound::FMidiTimestampTransportState NewState
+		{
+			Frame,
+			0,
+			State
+		};
+
+		check(TestClock.IsValid());
+		TestClock->AddTransportStateChangeToBlock(NewState);
+	}
+
+	END_DEFINE_SPEC(FHarmonixMetasoundMidiClockSpec)
+
+	void FHarmonixMetasoundMidiClockSpec::Define()
+	{
+		Describe("When a new transport state is added to the block", [this]()
+		{
+			BeforeEach([this]()
+			{
+				TestClock = MakeUnique<HarmonixMetasound::FMidiClock>(Metasound::FOperatorSettings{ 48000, 100 });
+			});
+			
+			It("and the block has no other state changes, it should succeed", [this]()
+			{
+				TestClock->PrepareBlock();
+
+				TestFalse("No transport changes in block", TestClock->HasTransportStateChangeInBlock());
+
+				constexpr HarmonixMetasound::EMusicPlayerTransportState NewState = HarmonixMetasound::EMusicPlayerTransportState::Playing;
+				AddStateAtFrame(NewState, 0);
+
+				TestTrue("There is a transport state change in block", TestClock->HasTransportStateChangeInBlock());
+				TestEqual("State at end of block matches the one we added", TestClock->GetTransportStateAtEndOfBlock(), NewState);
+			});
+
+			It("and the new state happens after the latest one, it should succeed", [this]()
+			{
+				AddStateAtFrame(HarmonixMetasound::EMusicPlayerTransportState::Playing, 0);
+				TestTrue("There is already a transport state change in the block", TestClock->HasTransportStateChangeInBlock());
+
+				const int32 NumInitialStates = TestClock->GetTransportTimestampsInBlock().Num();
+				const int32 LastStateFrame = TestClock->GetTransportTimestampsInBlock().Last().BlockSampleFrameIndex;
+
+				constexpr HarmonixMetasound::EMusicPlayerTransportState NewState = HarmonixMetasound::EMusicPlayerTransportState::Pausing;
+				AddStateAtFrame(NewState, LastStateFrame + 1);
+
+				TestEqual("There is another transport state change in block", TestClock->GetTransportTimestampsInBlock().Num(), NumInitialStates + 1);
+				TestEqual("State at end of block matches the one we added", TestClock->GetTransportStateAtEndOfBlock(), NewState);
+			});
+
+			It("and the new state happens at the same time as the latest one, it should succeed", [this]()
+			{
+				AddStateAtFrame(HarmonixMetasound::EMusicPlayerTransportState::Playing, 0);
+				TestTrue("There is already a transport state change in the block", TestClock->HasTransportStateChangeInBlock());
+
+				const int32 NumInitialStates = TestClock->GetTransportTimestampsInBlock().Num();
+				const int32 LastStateFrame = TestClock->GetTransportTimestampsInBlock().Last().BlockSampleFrameIndex;
+
+				constexpr HarmonixMetasound::EMusicPlayerTransportState NewState = HarmonixMetasound::EMusicPlayerTransportState::Paused;
+				AddStateAtFrame(NewState, LastStateFrame );
+
+				TestEqual("There is another transport state change in block", TestClock->GetTransportTimestampsInBlock().Num(), NumInitialStates + 1);
+				TestEqual("State at end of block matches the one we added", TestClock->GetTransportStateAtEndOfBlock(), NewState);
+			});
+
+			It("and the new state happens before the latest one, it should fail", [this]()
+			{
+				AddStateAtFrame(HarmonixMetasound::EMusicPlayerTransportState::Playing, 0);
+				TestTrue("There is already a transport state change in the block", TestClock->HasTransportStateChangeInBlock());
+
+				const int32 NumInitialStates = TestClock->GetTransportTimestampsInBlock().Num();
+				const HarmonixMetasound::EMusicPlayerTransportState LastState = TestClock->GetTransportStateAtEndOfBlock();
+				const int32 LastStateFrame = TestClock->GetTransportTimestampsInBlock().Last().BlockSampleFrameIndex;
+
+				constexpr HarmonixMetasound::EMusicPlayerTransportState NewState = HarmonixMetasound::EMusicPlayerTransportState::Continuing;
+				AddStateAtFrame(NewState, LastStateFrame - 1);
+
+				TestEqual("There is not another transport state change in block", TestClock->GetTransportTimestampsInBlock().Num(), NumInitialStates);
+				TestEqual("State at end of block matches the one that was already there", TestClock->GetTransportStateAtEndOfBlock(), LastState);
+			});
+		});
+	}
+
+}
+
+#endif
\ No newline at end of file
diff --git a/Engine/Plugins/Runtime/Harmonix/Source/HarmonixMetasoundTests/Private/Nodes/MediaStreamerTransportControllerNodeTests.cpp b/Engine/Plugins/Runtime/Harmonix/Source/HarmonixMetasoundTests/Private/Nodes/MediaStreamerTransportControllerNodeTests.cpp
new file mode 100644
--- /dev/null
+++ b/Engine/Plugins/Runtime/Harmonix/Source/HarmonixMetasoundTests/Private/Nodes/MediaStreamerTransportControllerNodeTests.cpp
@@ -0,0 +1,44 @@
+// Copyright Epic Games, Inc. All Rights Reserved.
+
+#include "NodeTestGraphBuilder.h"
+#include "HarmonixDsp/AudioBuffer.h"
+#include "HarmonixMetasound/Common.h"
+#include "Misc/AutomationTest.h"
+
+#if WITH_DEV_AUTOMATION_TESTS
+
+namespace HarmonixMetasoundTests::MediaStreamerTransportControllerNode
+{
+	using GraphBuilder = Metasound::Test::FNodeTestGraphBuilder;
+	using namespace Metasound;
+	using namespace Metasound::Frontend;
+	using namespace HarmonixMetasound;
+
+	IMPLEMENT_SIMPLE_AUTOMATION_TEST(
+		FMediaStreamerTransportControllerCreateNodeTest,
+		"Harmonix.Metasound.Nodes.MediaStreamerTransportControllerNode.CreateNode",
+		EAutomationTestFlags::EditorContext | EAutomationTestFlags::EngineFilter)
+		bool FMediaStreamerTransportControllerCreateNodeTest::RunTest(const FString&)
+	{
+		// Build the graph.
+		constexpr int32 NumSamplesPerBlock = 256;
+		const TUniquePtr<FMetasoundGenerator> Generator = GraphBuilder::MakeSingleNodeGraph(
+			{ HarmonixMetasound::HarmonixNodeNamespace, "MediaStreamerTransportController", "" },
+			0,
+			48000,
+			NumSamplesPerBlock);
+		UTEST_TRUE("Graph successfully built", Generator.IsValid());
+
+		// execute a block
+		{
+			TAudioBuffer<float> Buffer{ Generator->GetNumChannels(), NumSamplesPerBlock, EAudioBufferCleanupMode::Delete};
+			Generator->OnGenerateAudio(Buffer.GetRawChannelData(0), Buffer.GetNumTotalValidSamples());
+		}
+
+		// No output to validate, finish test
+
+		return true;
+	}
+}
+
+#endif
diff --git a/Engine/Plugins/Runtime/Harmonix/Source/HarmonixMetasoundTests/Private/Nodes/MidiCCTriggerNodeTests.cpp b/Engine/Plugins/Runtime/Harmonix/Source/HarmonixMetasoundTests/Private/Nodes/MidiCCTriggerNodeTests.cpp
new file mode 100644
--- /dev/null
+++ b/Engine/Plugins/Runtime/Harmonix/Source/HarmonixMetasoundTests/Private/Nodes/MidiCCTriggerNodeTests.cpp
@@ -0,0 +1,262 @@
+// Copyright Epic Games, Inc. All Rights Reserved.
+
+#include "NodeTestGraphBuilder.h"
+#include "HarmonixDsp/AudioBuffer.h"
+#include "HarmonixMetasound/Common.h"
+#include "HarmonixMetasound/DataTypes/MidiStream.h"
+#include "HarmonixMetasound/DataTypes/MidiAsset.h"
+#include "HarmonixMetasound/DataTypes/MusicTransport.h"
+#include "Misc/AutomationTest.h"
+#include "HarmonixMidi/MidiMsg.h"
+#include "HarmonixMidi/MidiFile.h"
+#include "HarmonixMetasound/DataTypes/MidiControllerID.h"
+#include "HarmonixMetasound/Nodes/MidiCCTriggerNode.h"
+
+#if WITH_DEV_AUTOMATION_TESTS
+/**
+ * Helper enum class for evaluating different 
+ * blocks containing different Midi Events in a Midi file
+ */
+UENUM()
+enum class EMidiEventTestingBlock : uint8
+{
+	ZeroBlock,
+	MidiChannelMismatchBlock,
+	MidiTrackMismatchBlock,
+	MidiEventTypeMismatchBlock,
+	MidiControllerIDMismatchBlock,
+	Count
+};
+
+namespace HarmonixMetasoundTests::MidiCCTriggerNode
+{
+	//expected values for tests
+	constexpr int32 ExpectedCCInt32ValueOnPass = 127;
+	constexpr int32 ExpectedCCInt32ValueNotPass = 0;
+	constexpr float ExpectedCCFloatValueOnPass = 1.0f;
+	constexpr float ExpectedCCFloatValueNotPass = 0.0f;
+	constexpr int32 ExpectedTriggerCountOnPass = 1;
+	constexpr int32 ExpectedTriggerCountNotPass = 0;
+	constexpr int32 ExpectedTriggerFrameIndexOnPass = 0;
+	constexpr int32 ExpectedTriggerFrameIndexNotPass = -1;
+
+	//testing values
+	constexpr int32 TestInputChannel = 1;
+	constexpr int32 TestDefaultTrackIndex = 1;
+	constexpr uint8 DefaultTestCCValue = 127;
+	constexpr uint8 DifferentTestCCValue = 100;
+	constexpr uint8 TestControllerIDShouldPass = static_cast<uint8>(EStdMidiControllerID::Hold);
+	constexpr uint8 TestControllerIDShouldNotPass = static_cast<uint8>(EStdMidiControllerID::Hold2);
+
+	constexpr int32 DefaultTrackNumber = 0;
+	constexpr int32 DifferentTrackNumber = 1;
+	
+	//time frames of midi events to be added to midi file for testing
+	//for block 2 - 4 (block 1 starts on tick 0, 0.0ms)
+	//each block is 480 samples / 10ms long
+	constexpr float TestingMsBlock2 = 15.0f;
+	constexpr float TestingMsBlock3 = 25.0f;
+	constexpr float TestingMsBlock4 = 35.0f;
+
+	UMidiFile* BuildMidiFile()
+	{
+		UMidiFile* TheMidiFile = NewObject<UMidiFile>();
+
+		//add tempo info bpm = 120
+		TheMidiFile->GetSongMaps()->GetTempoMap().AddTempoInfoPoint(MidiConstants::BPMToMidiTempo(120), 0);
+
+		//Add two midi tracks for testing
+		FMidiTrack* CCTrackDefault = TheMidiFile->AddTrack("CCEventTestTrackDefault");//track 1
+		FMidiTrack* CCTrackDifferent = TheMidiFile->AddTrack("CCEventTestTrackDifferent");//track 2
+
+		/**
+		 * Midi Events for Testing Block 1:
+		 * Add 2 CC events at tick 0 with different channel information and different cc value for testing
+		 */
+		FMidiEvent CCEventChannelMatch(0, FMidiMsg(MidiConstants::kControl, TestControllerIDShouldPass, DefaultTestCCValue));
+		//change channel information to 2 by modifying the lowest 4 bits (should not pass)
+		FMidiEvent CCEventChannelMismatch(0, FMidiMsg(MidiConstants::kControl + 0x01, TestControllerIDShouldPass, DifferentTestCCValue));
+		//add both events to track 1 (default)
+		CCTrackDefault->AddEvent(CCEventChannelMatch);
+		CCTrackDefault->AddEvent(CCEventChannelMismatch);
+
+		/**
+		 * Midi Events for Testing Block 2:
+		 * Add 2 CC events to 2 midi tracks at the same tick, default is track 1
+		 */
+		FMidiEvent CCEventTrackMatch(TheMidiFile->GetSongMaps()->MsToTick(TestingMsBlock2), 
+			FMidiMsg(MidiConstants::kControl, TestControllerIDShouldPass, DefaultTestCCValue));
+		CCTrackDefault->AddEvent(CCEventTrackMatch);//added to track 1
+		//add the other event to track to, which should NOT pass through to output in this test, 
+		//changing the control change value 
+		FMidiEvent CCEventTrackMismatch(TheMidiFile->GetSongMaps()->MsToTick(TestingMsBlock2),
+			FMidiMsg(MidiConstants::kControl, TestControllerIDShouldPass, DifferentTestCCValue));
+		CCTrackDifferent->AddEvent(CCEventTrackMismatch);//added to track 2
+
+		/**
+		* Midi Events for Testing Block 3:
+		* Add 1 event that has a different controller ID (should NOT pass)
+		*/
+		FMidiEvent CCEventControllerIDMismatch(TheMidiFile->GetSongMaps()->MsToTick(TestingMsBlock3),
+			FMidiMsg(MidiConstants::kControl, TestControllerIDShouldNotPass, DefaultTestCCValue));
+		CCTrackDefault->AddEvent(CCEventControllerIDMismatch);
+
+		/**
+		* Midi Events for Testing Block 4:
+		* Add 1 event that has a different type (should NOT pass)
+		*/
+		FMidiEvent CCEventTypeMismatch(TheMidiFile->GetSongMaps()->MsToTick(TestingMsBlock4),
+			FMidiMsg::CreateNoteOn(0, 64, DefaultTestCCValue));
+		CCTrackDefault->AddEvent(CCEventTypeMismatch);
+
+		//call TracksChanged() to update SongLengthData
+		TheMidiFile->TracksChanged();
+
+		//conform midi file length to an integer bar length
+		TheMidiFile->ConformMidiFileLength(EMidiFileLengthConformOption::RoundUp);
+
+		return TheMidiFile;
+	}
+
+	using GraphBuilder = Metasound::Test::FNodeTestGraphBuilder;
+	using namespace Metasound;
+	using namespace Metasound::Frontend;
+	using namespace HarmonixMetasound;
+	using namespace Metasound::Test;
+	using namespace HarmonixMetasound::Nodes::MidiCCTriggerNode;
+	
+	IMPLEMENT_SIMPLE_AUTOMATION_TEST(
+		FMidiCCTriggerNodeTestBasic,
+		"Harmonix.Metasound.Nodes.MidiCCTriggerNode.Basic",
+		EAutomationTestFlags::EditorContext | EAutomationTestFlags::EngineFilter)
+		bool FMidiCCTriggerNodeTestBasic::RunTest(const FString&)
+	{
+		//************************************************************************************************************************
+		// Build the graph.
+		//************************************************************************************************************************
+		FNodeTestGraphBuilder Builder;
+		Builder.AddOutput("AudioOut", GetMetasoundDataTypeName<FAudioBuffer>());
+		//transport
+		Frontend::FNodeHandle TransportNode = Builder.AddNode({ HarmonixMetasound::HarmonixNodeNamespace, "TriggerToTransport", "" }, 0);
+		UTEST_TRUE("Transport Node Created", TransportNode.Get().IsValid());
+		// hoist up play trigger input
+		Builder.AddAndConnectDataReferenceInput(TransportNode, CommonPinNames::Inputs::TransportPlayName, GetMetasoundDataTypeName<FTrigger>());
+		// midi player
+		Frontend::FNodeHandle MidiPlayerNode = Builder.AddNode({ HarmonixMetasound::HarmonixNodeNamespace, "MidiPlayer", "" }, 0);
+		UTEST_TRUE("MidiPlayer Node Created", MidiPlayerNode.Get().IsValid());
+		// hoist up midi asset input
+		Builder.AddAndConnectDataReferenceInput(MidiPlayerNode, CommonPinNames::Inputs::MidiFileAssetName, GetMetasoundDataTypeName<FMidiAsset>());
+		//connect transport to midi player
+		bool ConnectionSuccess = Builder.ConnectNodes(TransportNode, CommonPinNames::Outputs::TransportName, MidiPlayerNode, CommonPinNames::Inputs::TransportName);
+		UTEST_TRUE("Transport Out connected to MidiPlayer node", ConnectionSuccess);
+		// Midi CC trigger node to test
+		FNodeHandle CCTriggerNode = Builder.AddNode({ HarmonixMetasound::HarmonixNodeNamespace, "MidiCCTrigger", "" }, 0);
+		UTEST_TRUE("MidiCCTrigger Node Created", CCTriggerNode.Get().IsValid());
+
+		//Midi Channel Number 
+		Builder.AddAndConnectDataReferenceInput(CCTriggerNode, Inputs::MidiChannelNumberName, GetMetasoundDataTypeName<int32>());
+		//Midi Track Number
+		Builder.AddAndConnectDataReferenceInput(CCTriggerNode, Inputs::MidiTrackNumberName, GetMetasoundDataTypeName<int32>());
+		//Midi Controller ID input 
+		Builder.AddAndConnectDataReferenceInput(CCTriggerNode, Inputs::InputMidiControllerIDName, GetMetasoundDataTypeName<FEnumStdMidiControllerID>());
+
+		//output Midi Control Change Value (int32)
+		Builder.AddAndConnectDataReferenceOutput(CCTriggerNode, Outputs::OutputControlChangeValueInt32Name, GetMetasoundDataTypeName<int32>());
+
+		//output normalized Midi Control Change Value (float)
+		Builder.AddAndConnectDataReferenceOutput(CCTriggerNode, Outputs::OutputControlChangeValueFloatName, GetMetasoundDataTypeName<float>(), Outputs::OutputControlChangeValueFloatName);
+
+		//output trigger
+		Builder.AddAndConnectDataReferenceOutput(CCTriggerNode, Outputs::OutputTriggerName, GetMetasoundDataTypeName<FTrigger>(), Outputs::OutputTriggerName);
+		
+		//connect midi player node to CC Trigger node
+		ConnectionSuccess = Builder.ConnectNodes(MidiPlayerNode, CommonPinNames::Outputs::MidiStreamName, CCTriggerNode, Inputs::MidiStreamName);
+		UTEST_TRUE("MidiPlayer Out successfully connected to Midi CC Trigger node", ConnectionSuccess);
+		
+		//************************************************************************************************************************
+		// Make the generator.
+		//************************************************************************************************************************
+		constexpr int32 NumSamplesPerBlock = 480; // 10ms per block for sample rate = 48000
+		const TUniquePtr<FMetasoundGenerator> Generator = Builder.BuildGenerator(48000, NumSamplesPerBlock);
+		UTEST_TRUE("Graph successfully built", Generator.IsValid());
+		
+		//************************************************************************************************************************
+		// Get reference to outputs we want to check.
+		//************************************************************************************************************************
+		TOptional<FInt32ReadRef> OutputControlChangeValueInt32 = Generator->GetOutputReadReference<int32>(Outputs::OutputControlChangeValueInt32Name);
+		UTEST_TRUE("Output exists", OutputControlChangeValueInt32.IsSet());
+		UTEST_EQUAL("Control Change value (int32) at start", **OutputControlChangeValueInt32, 0);
+
+		TOptional<FFloatReadRef> OutputControlChangeValueFloat = Generator->GetOutputReadReference<float>(Outputs::OutputControlChangeValueFloatName);
+		UTEST_TRUE("Output exists", OutputControlChangeValueFloat.IsSet());
+		UTEST_EQUAL("Control Change value (float) at start", **OutputControlChangeValueFloat, 0.0f);
+
+		TOptional<FTriggerReadRef> CCOutputTrigger = Generator->GetOutputReadReference<FTrigger>(Outputs::OutputTriggerName);
+		UTEST_TRUE("Output exists", CCOutputTrigger.IsSet());
+		UTEST_EQUAL("Number of CC triggers at start", (*CCOutputTrigger)->NumTriggeredInBlock(), 0);
+		UTEST_EQUAL("CC Output Trigger sample", (*CCOutputTrigger)->First(), -1);
+		
+		//************************************************************************************************************************
+		// Make initial inputs
+		//************************************************************************************************************************
+		// Make a Midi File to play by Midi Player
+		UMidiFile* MidiFile = BuildMidiFile();
+		TSharedPtr<Audio::IProxyData> MidiFileProxy = MidiFile->CreateProxyData({ "MidiCCTriggerTest" });
+		Generator->SetInputValue<FMidiAsset>(CommonPinNames::Inputs::MidiFileAssetName, MidiFileProxy);
+		
+		// Set the testing Midi Channel (default = 1)
+		Generator->SetInputValue<int32>(Inputs::MidiChannelNumberName, TestInputChannel);
+		// Set the testing Midi Track (default = 1)
+		Generator->SetInputValue<int32>(Inputs::MidiTrackNumberName, TestDefaultTrackIndex);
+		// Set the testing Midi Controller ID
+		Generator->SetInputValue<FEnumStdMidiControllerID>(Inputs::InputMidiControllerIDName, static_cast<FEnumStdMidiControllerID>(TestControllerIDShouldPass));
+		
+		// Trigger transport
+		Generator->ApplyToInputValue<FTrigger>(CommonPinNames::Inputs::TransportPlayName, [](FTrigger& TriggerIn) { TriggerIn.TriggerFrame(0); });
+
+		//************************************************************************************************************************
+		// Generate 4 10ms blocks for testing various Midi Events in file/stream
+		//************************************************************************************************************************
+		
+		constexpr uint8 FirstTestBlock = static_cast<uint8>(EMidiEventTestingBlock::MidiChannelMismatchBlock);
+		constexpr uint8 NumBlocksToTest = static_cast<uint8>(EMidiEventTestingBlock::Count) - 1;
+
+		TAudioBuffer<float> Buffer{ Generator->GetNumChannels(), NumSamplesPerBlock, EAudioBufferCleanupMode::Delete};
+		for (uint8 BlockIndex = FirstTestBlock; BlockIndex <= NumBlocksToTest; ++BlockIndex)
+		{
+			Generator->OnGenerateAudio(Buffer.GetRawChannelData(0), Buffer.GetNumTotalValidSamples());
+			
+			//Block 1 & 2 each should  have 1 out of 2 events pass through to output with their correlating CC values
+			if (static_cast<EMidiEventTestingBlock>(BlockIndex) == EMidiEventTestingBlock::MidiChannelMismatchBlock ||
+				static_cast<EMidiEventTestingBlock>(BlockIndex) == EMidiEventTestingBlock::MidiTrackMismatchBlock)
+			{
+				//check output CC values
+				UTEST_EQUAL("Control Change value (int32) check", **OutputControlChangeValueInt32, ExpectedCCInt32ValueOnPass);
+				UTEST_EQUAL("Control Change value (float) check", **OutputControlChangeValueFloat, ExpectedCCFloatValueOnPass);
+
+				//check if we have a trigger
+				UTEST_EQUAL("Number of CC Output Triggers", (*CCOutputTrigger)->NumTriggeredInBlock(), ExpectedTriggerCountOnPass);
+				
+				//check trigger sample frame at the first block
+				if (BlockIndex == 1)
+				{
+					UTEST_EQUAL("CC Output Trigger sample", (*CCOutputTrigger)->First(), ExpectedTriggerFrameIndexOnPass);
+				}
+			}
+			else {
+				//Block 3 & 4 each should have 1 event ONLY,and NONE should pass through to output
+				
+				//check output CC value (should both be 0)
+				UTEST_EQUAL("Control Change value (int32) check", **OutputControlChangeValueInt32, ExpectedCCInt32ValueNotPass);
+				UTEST_EQUAL("Control Change value (float) check", **OutputControlChangeValueFloat, ExpectedCCFloatValueNotPass);
+
+				//no trigger should pass through
+				UTEST_EQUAL("Number of CC Output Triggers", (*CCOutputTrigger)->NumTriggeredInBlock(), ExpectedTriggerCountNotPass);
+				UTEST_EQUAL("CC Output Trigger sample", (*CCOutputTrigger)->First(), ExpectedTriggerFrameIndexNotPass);
+			}
+		}
+
+		return true;
+	}
+}
+#endif
\ No newline at end of file
diff --git a/Engine/Plugins/Runtime/Harmonix/Source/HarmonixMetasoundTests/Private/Nodes/MidiClockSubdivisionTriggerNodeTests.cpp b/Engine/Plugins/Runtime/Harmonix/Source/HarmonixMetasoundTests/Private/Nodes/MidiClockSubdivisionTriggerNodeTests.cpp
new file mode 100644
--- /dev/null
+++ b/Engine/Plugins/Runtime/Harmonix/Source/HarmonixMetasoundTests/Private/Nodes/MidiClockSubdivisionTriggerNodeTests.cpp
@@ -0,0 +1,50 @@
+// Copyright Epic Games, Inc. All Rights Reserved.
+
+#include "NodeTestGraphBuilder.h"
+#include "HarmonixDsp/AudioBuffer.h"
+#include "HarmonixMetasound/Common.h"
+#include "Misc/AutomationTest.h"
+
+#include "Tests/AutomationEditorCommon.h"
+
+#if WITH_DEV_AUTOMATION_TESTS
+
+namespace HarmonixMetasoundTests::MidiClockSubdivisionTriggerNode
+{
+	using GraphBuilder = Metasound::Test::FNodeTestGraphBuilder;
+	using namespace Metasound;
+	using namespace Metasound::Frontend;
+	using namespace HarmonixMetasound;
+
+	IMPLEMENT_SIMPLE_AUTOMATION_TEST(
+		FMidiClockSubdivisionTriggerCreateNodeTest,
+		"Harmonix.Metasound.Nodes.MidiClockSubdivisionTriggerNode.CreateNode",
+		EAutomationTestFlags::EditorContext | EAutomationTestFlags::EngineFilter)
+		bool FMidiClockSubdivisionTriggerCreateNodeTest::RunTest(const FString&)
+	{
+		// Build the graph.
+		constexpr int32 NumSamplesPerBlock = 256;
+		const TUniquePtr<FMetasoundGenerator> Generator = GraphBuilder::MakeSingleNodeGraph(
+			{ HarmonixMetasound::HarmonixNodeNamespace, "MidiClockSubdivisionTrigger", "" },
+			0,
+			48000,
+			NumSamplesPerBlock);
+		UTEST_TRUE("Graph successfully built", Generator.IsValid());
+
+		// execute a block
+		{
+			TAudioBuffer<float> Buffer{ Generator->GetNumChannels(), NumSamplesPerBlock, EAudioBufferCleanupMode::Delete};
+			Generator->OnGenerateAudio(Buffer.GetRawChannelData(0), Buffer.GetNumTotalValidSamples());
+		}
+
+		// Validate that default input generate output with no triggers set.
+		TOptional<TDataReadReference<FTrigger>> OutputPin = Generator->GetOutputReadReference<FTrigger>("Trigger Out");
+		UTEST_TRUE("Trigger out exists", OutputPin.IsSet());
+		UTEST_EQUAL("Num triggers", (*OutputPin)->NumTriggeredInBlock(), 0);
+		UTEST_EQUAL("Trigger sample", (*OutputPin)->First(), -1);
+
+		return true;
+	}
+}
+
+#endif
diff --git a/Engine/Plugins/Runtime/Harmonix/Source/HarmonixMetasoundTests/Private/Nodes/MidiClockTimestampTriggerNodeTests.cpp b/Engine/Plugins/Runtime/Harmonix/Source/HarmonixMetasoundTests/Private/Nodes/MidiClockTimestampTriggerNodeTests.cpp
new file mode 100644
--- /dev/null
+++ b/Engine/Plugins/Runtime/Harmonix/Source/HarmonixMetasoundTests/Private/Nodes/MidiClockTimestampTriggerNodeTests.cpp
@@ -0,0 +1,48 @@
+// Copyright Epic Games, Inc. All Rights Reserved.
+
+#include "NodeTestGraphBuilder.h"
+#include "HarmonixDsp/AudioBuffer.h"
+#include "HarmonixMetasound/Common.h"
+#include "Misc/AutomationTest.h"
+
+#if WITH_DEV_AUTOMATION_TESTS
+
+namespace HarmonixMetasoundTests::MidiClockTimestampTriggerNode
+{
+	using GraphBuilder = Metasound::Test::FNodeTestGraphBuilder;
+	using namespace Metasound;
+	using namespace Metasound::Frontend;
+	using namespace HarmonixMetasound;
+
+	IMPLEMENT_SIMPLE_AUTOMATION_TEST(
+		FMidiClockTimestampTriggerCreateNodeTest,
+		"Harmonix.Metasound.Nodes.MidiClockTimestampTriggerNode.CreateNode",
+		EAutomationTestFlags::EditorContext | EAutomationTestFlags::EngineFilter)
+		bool FMidiClockTimestampTriggerCreateNodeTest::RunTest(const FString&)
+	{
+		// Build the graph.
+		constexpr int32 NumSamplesPerBlock = 256;
+		const TUniquePtr<FMetasoundGenerator> Generator = GraphBuilder::MakeSingleNodeGraph(
+			{ HarmonixMetasound::HarmonixNodeNamespace, "MidiClockTimestampTrigger", "" },
+			0,
+			48000,
+			NumSamplesPerBlock);
+		UTEST_TRUE("Graph successfully built", Generator.IsValid());
+
+		// execute a block
+		{
+			TAudioBuffer<float> Buffer{ Generator->GetNumChannels(), NumSamplesPerBlock, EAudioBufferCleanupMode::Delete};
+			Generator->OnGenerateAudio(Buffer.GetRawChannelData(0), Buffer.GetNumTotalValidSamples());
+		}
+
+		// Validate that default input generate output with no triggers set.
+		TOptional<TDataReadReference<FTrigger>> OutputPin = Generator->GetOutputReadReference<FTrigger>("Trigger Out");
+		UTEST_TRUE("Trigger out exists", OutputPin.IsSet());
+		UTEST_EQUAL("Num triggers", (*OutputPin)->NumTriggeredInBlock(), 0);
+		UTEST_EQUAL("Trigger sample", (*OutputPin)->First(), -1);
+
+		return true;
+	}
+}
+
+#endif
