[Harmonix] Move Harmonix plugin into Engine/Plugins/Runtime to make it available to licensees in UE 5.4

diff --git a/Engine/Plugins/Runtime/Harmonix/Source/HarmonixMetasoundTests/Private/Nodes/MidiNoteTriggerNodeTests.cpp b/Engine/Plugins/Runtime/Harmonix/Source/HarmonixMetasoundTests/Private/Nodes/MidiNoteTriggerNodeTests.cpp
new file mode 100644
--- /dev/null
+++ b/Engine/Plugins/Runtime/Harmonix/Source/HarmonixMetasoundTests/Private/Nodes/MidiNoteTriggerNodeTests.cpp
@@ -0,0 +1,482 @@
+// Copyright Epic Games, Inc. All Rights Reserved.
+
+#include "NodeTestGraphBuilder.h"
+#include "HarmonixDsp/AudioBuffer.h"
+#include "HarmonixMetasound/Common.h"
+#include "HarmonixMetasound/DataTypes/MidiStream.h"
+#include "Misc/AutomationTest.h"
+#include "HarmonixMetasound/DataTypes/MidiAsset.h"
+#include "HarmonixMetasound/DataTypes/MusicTransport.h"
+#include "HarmonixMidi/MidiMsg.h"
+#include "HarmonixMidi/MidiFile.h"
+
+#if WITH_DEV_AUTOMATION_TESTS
+
+namespace HarmonixMetasoundTests::MidiNoteTriggerNode
+{
+	//testing Note Numbers
+	constexpr int32 TestNoteNumberShouldPass = 60; //C4, frequency = 261.63 Hz
+	constexpr int32 TestNoteNumberShouldPass2 = 62; //D4, frequency = 293.66 Hz
+
+	constexpr int32 TestNoteNumberShouldNotPass = 72; //C5, frequency = 523.25 Hz
+	//Min and Max Note Numbers for input 
+	constexpr int32 TestMaxNoteNumberDefault = 100;
+	constexpr int32 TestMinNoteNumberDefault = 20;
+	//Note Numbers that are outside of the Min to Max range
+	constexpr int32 TestNoteNumberLessThanMin = 10;
+	constexpr int32 TestNoteNumberGreaterThanMax = 127;
+
+	//testing velocity values
+	constexpr int32 TestNoteVelocityShouldPass = 90;
+	constexpr float TestNoteNormalizedVelocityShouldPass = 90 / 127.f;
+	//Min and Max Velocity values
+	constexpr int32 TestMaxVelocityDefault = 100;
+	constexpr int32 TestMinVelocityDefault = 20;
+	//velocity values outside of the Min to Max range
+	constexpr int32 TestVelocityLessThanMin = 10;
+	constexpr int32 TestVelocityGreaterThanMax = 127;
+	constexpr int32 TestNoteVelocityShouldNotPass = 127;//maybe we don't need this?
+	
+	//test midi tracks
+	constexpr int32 TestMidiTrackShouldPass = 1;
+	constexpr int32 TestMidiTrackShouldNotPass = 2;
+
+	//test midi channel
+	constexpr int32 TestMidiChannelShouldPass = 0;
+	constexpr int32 TestMidiChannelShouldNotPass = 1;
+
+	//frequency outputs
+	constexpr float ExpectedFrequencyOnPass = 261.625549f; //needs verification 
+	constexpr float ExpectedFrequencyNotPass = 0.0f;
+
+	//other expected values if event does not pass through
+	constexpr int32 ExpectedNoteNumberNotPass = 0;
+	constexpr int32 ExpectedVelocityNotPass = 0;
+	constexpr float ExpectedNormalizedVelocityNotPass = 0.0f;
+
+	//expected number of triggers that pass through during each block
+	constexpr int32 ExpectedTriggerCountOnPass = 1;
+	constexpr int32 ExpectedTriggerCountNotPass = 0;
+
+	//expected trigger frames
+	constexpr int32 ExpectedTriggerFrameIndexNotPass = -1;
+
+	//Time indices (Ms) for Midi events in Testing blocks
+	constexpr float TestingMsNoteOnBlock1 = 0.0f;
+	constexpr float TestingMsNoteOffBlock2 = 13.0f;
+
+	constexpr float TestingMsNoteOnBlock2 = 18.0f;
+	constexpr float TestingMsNoteOffBlock3 = 23.0f;
+	constexpr float TestingMsNoteOnBlock3 = 28.0f;
+	constexpr float TestingMsNoteOffBlock4 = 33.0f;	
+	constexpr float TestingMsNoteOnBlock4 = 38.0f;
+	
+	//last block tests for Non-Note events
+	constexpr float TestingMsNoteOffBlock5 = 43.0f;
+	constexpr float TestingMsEventBlock5 = 48.0f;
+
+	//number of test blocks to generate during tests
+	constexpr int32 NumBlocksToTest = 5;
+
+	UMidiFile* BuildMidiFile()
+	{
+		UMidiFile* TheMidiFile = NewObject<UMidiFile>();
+
+		//add tempo info bpm = 120
+		TheMidiFile->GetSongMaps()->GetTempoMap().AddTempoInfoPoint(MidiConstants::BPMToMidiTempo(120), 0);
+
+		//Add two midi tracks for testing
+		FMidiTrack* MidiNoteTrackDefault = TheMidiFile->AddTrack("MidiNoteTestTrackDefault");//track 1
+		FMidiTrack* MidiNoteTrackDifferent = TheMidiFile->AddTrack("MidiNoteTestTrackDifferent");//track 2
+
+		/**
+		* Midi Events for Testing Block 1:
+		* Add 2 Note On events in the first block with different channel information and different note number for testing
+		*/
+		//Note On Events
+		FMidiEvent MidiChannelMatchNoteOn(TheMidiFile->GetSongMaps()->MsToTick(TestingMsNoteOnBlock1),
+			FMidiMsg::CreateNoteOn(TestMidiChannelShouldPass, TestNoteNumberShouldPass, TestNoteVelocityShouldPass));
+		//this should not pass
+		FMidiEvent MidiChannelMismatchNoteOn(TheMidiFile->GetSongMaps()->MsToTick(TestingMsNoteOnBlock1),
+			FMidiMsg::CreateNoteOn(TestMidiChannelShouldNotPass, TestNoteNumberShouldNotPass, TestNoteVelocityShouldNotPass));
+		
+		//Add 2 Note On events to default track 1
+		MidiNoteTrackDefault->AddEvent(MidiChannelMismatchNoteOn);
+		MidiNoteTrackDefault->AddEvent(MidiChannelMatchNoteOn);
+
+
+		/**
+		 * Midi Events for Testing Block 2:
+		 * Add 2 Note Off events for the 2 Note On events in the previous block
+		 * (NOTE: Putting Note On and Note Off Events in separate blocks for testing the number of output triggers
+		 * since Note Offs will reset output)
+		 * Add 2 Note On events to 2 separate midi tracks at the same tick, default is track 1
+		 */
+		 //Note Off Events pairing with Note On events in block 1
+		FMidiEvent MidiChannelMatchNoteOff(TheMidiFile->GetSongMaps()->MsToTick(TestingMsNoteOffBlock2),
+			FMidiMsg::CreateNoteOff(TestMidiChannelShouldPass, TestNoteNumberShouldPass));
+		//this should not pass
+		FMidiEvent MidiChannelMismatchNoteOff(TheMidiFile->GetSongMaps()->MsToTick(TestingMsNoteOffBlock2),
+			FMidiMsg::CreateNoteOff(TestMidiChannelShouldNotPass, TestNoteNumberShouldNotPass));
+		//Add Note Off events to track 1
+		MidiNoteTrackDefault->AddEvent(MidiChannelMismatchNoteOff);
+		MidiNoteTrackDefault->AddEvent(MidiChannelMatchNoteOff);
+
+		FMidiEvent MidiTrackMatchNoteOn(TheMidiFile->GetSongMaps()->MsToTick(TestingMsNoteOnBlock2),
+			FMidiMsg::CreateNoteOn(TestMidiChannelShouldPass, TestNoteNumberShouldPass, TestNoteVelocityShouldPass));
+		//this should not pass
+		FMidiEvent MidiTrackMismatchNoteOn(TheMidiFile->GetSongMaps()->MsToTick(TestingMsNoteOnBlock2),
+			FMidiMsg::CreateNoteOn(TestMidiChannelShouldPass, TestNoteNumberShouldNotPass, TestNoteVelocityShouldNotPass));
+
+		//Add 1 Note On event to default track 1 and the other to track 2
+		MidiNoteTrackDifferent->AddEvent(MidiTrackMismatchNoteOn);
+		MidiNoteTrackDefault->AddEvent(MidiTrackMatchNoteOn);
+
+		/**
+		* Midi Events for Testing Block 3:
+		* Add 2 Note Off events for the 2 Note On events in the previous block
+		* Add 2 pairs of Note On events that have note numbers less than min note number or greater than max note number
+		* Neither should pass through to output 
+		*/
+		//Not Off Events pairing with Note On events from the previous block
+		FMidiEvent MidiTrackMatchNoteOff(TheMidiFile->GetSongMaps()->MsToTick(TestingMsNoteOffBlock3),
+			FMidiMsg::CreateNoteOff(TestMidiChannelShouldPass, TestNoteNumberShouldPass));
+		//this should not pass
+		FMidiEvent MidiTrackMismatchNoteOff(TheMidiFile->GetSongMaps()->MsToTick(TestingMsNoteOffBlock3),
+			FMidiMsg::CreateNoteOff(TestMidiChannelShouldPass, TestNoteNumberShouldNotPass));
+		//Add Note Off events to their corresponding tracks
+		MidiNoteTrackDifferent->AddEvent(MidiTrackMismatchNoteOff);
+		MidiNoteTrackDefault->AddEvent(MidiTrackMatchNoteOff);
+
+		//Note On Events that have Note Numbers greater than the max or less than min Note Number
+		FMidiEvent NoteNumberLessThanMinNoteOn(TheMidiFile->GetSongMaps()->MsToTick(TestingMsNoteOnBlock3),
+			FMidiMsg::CreateNoteOn(TestMidiChannelShouldPass, TestNoteNumberLessThanMin, TestNoteVelocityShouldPass));
+		FMidiEvent NoteNumberGreaterThanMaxNoteOn(TheMidiFile->GetSongMaps()->MsToTick(TestingMsNoteOnBlock3),
+			FMidiMsg::CreateNoteOn(TestMidiChannelShouldPass, TestNoteNumberGreaterThanMax, TestNoteVelocityShouldPass));
+		//Add 2 Note On events to default track 1, neither should pass
+		MidiNoteTrackDefault->AddEvent(NoteNumberLessThanMinNoteOn);
+		MidiNoteTrackDefault->AddEvent(NoteNumberGreaterThanMaxNoteOn);
+
+		/**
+		* Midi Events for Testing Block 4:
+		* Add 2 Note Off events for the 2 Note On events in the previous block
+		* Add 1 event that has a different type (should NOT pass)
+		*/
+		//Not Off events pairing with Note On events from the previous block
+		FMidiEvent NoteNumberLessThanMinNoteOff(TheMidiFile->GetSongMaps()->MsToTick(TestingMsNoteOffBlock4),
+			FMidiMsg::CreateNoteOff(TestMidiChannelShouldPass, TestNoteNumberLessThanMin));
+		FMidiEvent NoteNumberGreaterThanMaxNoteOff(TheMidiFile->GetSongMaps()->MsToTick(TestingMsNoteOffBlock4),
+			FMidiMsg::CreateNoteOff(TestMidiChannelShouldPass, TestNoteNumberGreaterThanMax));
+		//Add 2 Note Off events to default track 1, neither should pass
+		MidiNoteTrackDefault->AddEvent(NoteNumberLessThanMinNoteOff);
+		MidiNoteTrackDefault->AddEvent(NoteNumberGreaterThanMaxNoteOff);
+		
+		//Note On Events that have velocity values greater than the max or less than min velocity value
+		FMidiEvent NoteVelocityLessThanMinNoteOn(TheMidiFile->GetSongMaps()->MsToTick(TestingMsNoteOnBlock4),
+			FMidiMsg::CreateNoteOn(TestMidiChannelShouldPass, TestNoteNumberShouldPass, TestVelocityLessThanMin));
+		FMidiEvent NoteVelocityGreaterThanMaxNoteOn(TheMidiFile->GetSongMaps()->MsToTick(TestingMsNoteOnBlock4),
+			FMidiMsg::CreateNoteOn(TestMidiChannelShouldPass, TestNoteNumberShouldPass, TestVelocityGreaterThanMax));
+		//Add 2 Note On events to default track 1, neither should pass
+		MidiNoteTrackDefault->AddEvent(NoteVelocityLessThanMinNoteOn);
+		MidiNoteTrackDefault->AddEvent(NoteVelocityGreaterThanMaxNoteOn);
+		
+
+
+		/**
+		* Midi Events for Testing Block 5:
+		* Add 2 Note Off events for the 2 Note On events in the previous block
+		* Add 1 event that has a different type (should NOT pass), this one is a Pitch Bend event
+		*/
+		//Not Off events pairing with Note On events from the previous block
+		FMidiEvent NoteVelocityLessThanMinNoteOff(TheMidiFile->GetSongMaps()->MsToTick(TestingMsNoteOffBlock5),
+			FMidiMsg::CreateNoteOff(TestMidiChannelShouldPass, TestNoteNumberShouldPass));
+		FMidiEvent NoteVelocityGreaterThanMaxNoteOff(TheMidiFile->GetSongMaps()->MsToTick(TestingMsNoteOffBlock5),
+			FMidiMsg::CreateNoteOff(TestMidiChannelShouldPass, TestNoteNumberShouldPass));
+
+		//Add Note Off events to default track 1, neither should pass
+		MidiNoteTrackDefault->AddEvent(NoteVelocityLessThanMinNoteOff);
+		MidiNoteTrackDefault->AddEvent(NoteVelocityGreaterThanMaxNoteOff);
+
+		//Add a Pitch Bend event to default track 1, should not pass
+		FMidiEvent EventTypeMismatch(TheMidiFile->GetSongMaps()->MsToTick(TestingMsEventBlock5),
+			FMidiMsg(MidiConstants::kPitch, TestNoteNumberShouldPass2, TestNoteVelocityShouldPass));
+		MidiNoteTrackDefault->AddEvent(EventTypeMismatch);
+
+		//call TracksChanged() to update SongLengthData
+		TheMidiFile->TracksChanged();
+
+		//conform midi file length to an integer bar length
+		TheMidiFile->ConformMidiFileLength(EMidiFileLengthConformOption::RoundUp);
+
+		return TheMidiFile;
+	}
+
+	using GraphBuilder = Metasound::Test::FNodeTestGraphBuilder;
+	using namespace Metasound;
+	using namespace Metasound::Frontend;
+	using namespace HarmonixMetasound;
+
+	IMPLEMENT_SIMPLE_AUTOMATION_TEST(
+		FMidiNoteTriggerCreateNodeTest,
+		"Harmonix.Metasound.Nodes.MidiNoteTriggerNode.CreateNode",
+		EAutomationTestFlags::EditorContext | EAutomationTestFlags::EngineFilter)
+		bool FMidiNoteTriggerCreateNodeTest::RunTest(const FString&)
+	{
+		// Build the graph.
+		constexpr int32 NumSamplesPerBlock = 256;
+		const TUniquePtr<FMetasoundGenerator> Generator = GraphBuilder::MakeSingleNodeGraph(
+			{ HarmonixMetasound::HarmonixNodeNamespace, "MidiNoteTrigger", "" },
+			0,
+			48000,
+			NumSamplesPerBlock);
+		UTEST_TRUE("Graph successfully built", Generator.IsValid());
+
+		// execute a block
+		{
+			TAudioBuffer<float> Buffer{ Generator->GetNumChannels(), NumSamplesPerBlock, EAudioBufferCleanupMode::Delete};
+			Generator->OnGenerateAudio(Buffer.GetRawChannelData(0), Buffer.GetNumTotalValidSamples());
+		}
+
+		// Validate output.
+		TOptional<FTriggerReadRef> OutputNoteOn = Generator->GetOutputReadReference<FTrigger>(CommonPinNames::Outputs::NoteOnName);
+		UTEST_TRUE("Output exists", OutputNoteOn.IsSet());
+		UTEST_EQUAL("Number of Note On triggers", (*OutputNoteOn)->NumTriggeredInBlock(), 0);
+		UTEST_EQUAL("Note On Trigger sample", (*OutputNoteOn)->First(), -1);
+
+		TOptional<FTriggerReadRef> OutputNoteOff = Generator->GetOutputReadReference<FTrigger>(CommonPinNames::Outputs::NoteOffName);
+		UTEST_TRUE("Output exists", OutputNoteOff.IsSet());
+		UTEST_EQUAL("Number of Note Off triggers", (*OutputNoteOff)->NumTriggeredInBlock(), 0);
+		UTEST_EQUAL("Note Off Trigger sample", (*OutputNoteOff)->First(), -1);
+
+		TOptional<FInt32ReadRef> OutputMidiNoteNum = Generator->GetOutputReadReference<int32>("Midi Note #");
+		UTEST_TRUE("Output exists", OutputMidiNoteNum.IsSet());
+		UTEST_EQUAL("Midi Note Num check", **OutputMidiNoteNum, 0);
+
+		TOptional<FFloatReadRef> OutputFrequency = Generator->GetOutputReadReference<float>("Frequency");
+		UTEST_TRUE("Output exists", OutputFrequency.IsSet());
+		UTEST_EQUAL("Frequency check", **OutputFrequency, 0.0f);
+
+		TOptional<FInt32ReadRef> OutputVelocity = Generator->GetOutputReadReference<int32>("Velocity");
+		UTEST_TRUE("Output exists", OutputVelocity.IsSet());
+		UTEST_EQUAL("Velocity check", **OutputVelocity, 0);
+
+		TOptional<FFloatReadRef> OutputNormalizedVelocity = Generator->GetOutputReadReference<float>("Normalized Velocity");
+		UTEST_TRUE("Output exists", OutputNormalizedVelocity.IsSet());
+		UTEST_EQUAL("Normalized Velocity check", **OutputNormalizedVelocity, 0.0f);
+
+		return true;
+	}
+
+	IMPLEMENT_SIMPLE_AUTOMATION_TEST(
+		FMidiNoteTriggerNodeTestBasic,
+		"Harmonix.Metasound.Nodes.MidiNoteTriggerNode.Basic",
+		EAutomationTestFlags::EditorContext | EAutomationTestFlags::EngineFilter)
+		bool FMidiNoteTriggerNodeTestBasic::RunTest(const FString&)
+	{
+		//************************************************************************************************************************
+		// Build the graph.
+		//************************************************************************************************************************
+		GraphBuilder Builder;
+		Builder.AddOutput("AudioOut", GetMetasoundDataTypeName<FAudioBuffer>());
+		//transport
+		Frontend::FNodeHandle TransportNode = Builder.AddNode({ HarmonixMetasound::HarmonixNodeNamespace, "TriggerToTransport", "" }, 0);
+		UTEST_TRUE("Transport Node Created", TransportNode.Get().IsValid());
+		// hoist up play trigger input
+		Builder.AddAndConnectDataReferenceInput(TransportNode, CommonPinNames::Inputs::TransportPlayName, GetMetasoundDataTypeName<FTrigger>());
+		// midi player
+		Frontend::FNodeHandle MidiPlayerNode = Builder.AddNode({ HarmonixMetasound::HarmonixNodeNamespace, "MidiPlayer", "" }, 0);
+		UTEST_TRUE("MidiPlayer Node Created", MidiPlayerNode.Get().IsValid());
+		// hoist up midi asset input
+		Builder.AddAndConnectDataReferenceInput(MidiPlayerNode, CommonPinNames::Inputs::MidiFileAssetName, GetMetasoundDataTypeName<FMidiAsset>());
+		//connect transport to midi player
+		bool ConnectionSuccess = Builder.ConnectNodes(TransportNode, CommonPinNames::Outputs::TransportName, MidiPlayerNode, CommonPinNames::Inputs::TransportName);
+		UTEST_TRUE("Transport Out connected to MidiPlayer node", ConnectionSuccess);
+		// Midi Note Trigger node to test
+		FNodeHandle NoteTriggerNode = Builder.AddNode({ HarmonixMetasound::HarmonixNodeNamespace, "MidiNoteTrigger", "" }, 0);
+		UTEST_TRUE("MidiNoteTrigger Node Created", NoteTriggerNode.Get().IsValid());
+
+		//Midi Input Channel Number 
+		Builder.AddAndConnectDataReferenceInput(NoteTriggerNode, CommonPinNames::Inputs::MidiChannelNumberName, GetMetasoundDataTypeName<int32>());
+		//Midi Input Track Number
+		Builder.AddAndConnectDataReferenceInput(NoteTriggerNode, CommonPinNames::Inputs::MidiTrackNumberName, GetMetasoundDataTypeName<int32>());
+		//min Input Midi Note Number
+		Builder.AddAndConnectDataReferenceInput(NoteTriggerNode, CommonPinNames::Inputs::MinMidiNoteName, GetMetasoundDataTypeName<int32>());
+		//Max Input Midi Note Number
+		Builder.AddAndConnectDataReferenceInput(NoteTriggerNode, CommonPinNames::Inputs::MaxMidiNoteName, GetMetasoundDataTypeName<int32>());
+		//Input Min Velocity
+		Builder.AddAndConnectDataReferenceInput(NoteTriggerNode, CommonPinNames::Inputs::MinMidiVelocityName, GetMetasoundDataTypeName<int32>());
+		//Input Max Velocity
+		Builder.AddAndConnectDataReferenceInput(NoteTriggerNode, CommonPinNames::Inputs::MaxMidiVelocityName, GetMetasoundDataTypeName<int32>());
+		
+
+		//output Note On trigger
+		Builder.AddAndConnectDataReferenceOutput(NoteTriggerNode, CommonPinNames::Outputs::NoteOnName, GetMetasoundDataTypeName<FTrigger>(), CommonPinNames::Outputs::NoteOnName);
+		//output Note Off trigger
+		Builder.AddAndConnectDataReferenceOutput(NoteTriggerNode, CommonPinNames::Outputs::NoteOffName, GetMetasoundDataTypeName<FTrigger>(), CommonPinNames::Outputs::NoteOffName);
+		
+		//output Midi Note number
+		Builder.AddAndConnectDataReferenceOutput(NoteTriggerNode, CommonPinNames::Outputs::MidiNoteNumberName, GetMetasoundDataTypeName<int32>(), CommonPinNames::Outputs::MidiNoteNumberName);
+
+		//output Midi Note Velocity (Int32)
+		Builder.AddAndConnectDataReferenceOutput(NoteTriggerNode, CommonPinNames::Outputs::MidiVelocityName, GetMetasoundDataTypeName<int32>(), CommonPinNames::Outputs::MidiVelocityName);
+		//output normalized Midi Note Velocity (float)
+		Builder.AddAndConnectDataReferenceOutput(NoteTriggerNode, CommonPinNames::Outputs::NormalizedVelocityName, GetMetasoundDataTypeName<float>(), CommonPinNames::Outputs::NormalizedVelocityName);
+
+		//output frequency 
+		Builder.AddAndConnectDataReferenceOutput(NoteTriggerNode, CommonPinNames::Outputs::FrequencyName, GetMetasoundDataTypeName<float>(), CommonPinNames::Outputs::FrequencyName);
+
+		//connect midi player node to Note Trigger node
+		ConnectionSuccess = Builder.ConnectNodes(MidiPlayerNode, CommonPinNames::Outputs::MidiStreamName, NoteTriggerNode, CommonPinNames::Inputs::MidiStreamName);
+		UTEST_TRUE("MidiPlayer Out successfully connected to Midi Note Trigger node", ConnectionSuccess);
+
+		//************************************************************************************************************************
+		// Make the generator.
+		//************************************************************************************************************************
+		constexpr int32 NumSamplesPerBlock = 480; // 10ms per block for sample rate = 48000
+		const TUniquePtr<FMetasoundGenerator> Generator = Builder.BuildGenerator(48000, NumSamplesPerBlock);
+		UTEST_TRUE("Graph successfully built", Generator.IsValid());
+
+		//************************************************************************************************************************
+		// Get reference to outputs we want to check.
+		//************************************************************************************************************************
+		
+		TOptional<FTriggerReadRef> OutputNoteOnTrigger = Generator->GetOutputReadReference<FTrigger>(CommonPinNames::Outputs::NoteOnName);
+		UTEST_TRUE("Output exists", OutputNoteOnTrigger.IsSet());
+		UTEST_EQUAL("Number of Note Triggers at start", (*OutputNoteOnTrigger)->NumTriggeredInBlock(), 0);
+		UTEST_EQUAL("Note On Output Trigger sample", (*OutputNoteOnTrigger)->First(), -1);
+		
+		TOptional<FTriggerReadRef> OutputNoteOffTrigger = Generator->GetOutputReadReference<FTrigger>(CommonPinNames::Outputs::NoteOffName);
+		UTEST_TRUE("Output exists", OutputNoteOffTrigger.IsSet());
+		UTEST_EQUAL("Number of Note Triggers at start", (*OutputNoteOffTrigger)->NumTriggeredInBlock(), 0);
+		UTEST_EQUAL("Note off Output Trigger sample", (*OutputNoteOffTrigger)->First(), -1);
+
+		TOptional<FInt32ReadRef> OutputMidiNoteNumber = Generator->GetOutputReadReference<int32>(CommonPinNames::Outputs::MidiNoteNumberName);
+		UTEST_TRUE("Output exists", OutputMidiNoteNumber.IsSet());
+		UTEST_EQUAL("Midi Note Number at start", **OutputMidiNoteNumber, 0);
+
+		TOptional<FInt32ReadRef> OutputMidiVelocity = Generator->GetOutputReadReference<int32>(CommonPinNames::Outputs::MidiVelocityName);
+		UTEST_TRUE("Output exists", OutputMidiVelocity.IsSet());
+		UTEST_EQUAL("Midi Velocity (int32) at start", **OutputMidiVelocity, 0);
+
+		TOptional<FFloatReadRef> OutputNormalizedVelocity = Generator->GetOutputReadReference<float>(CommonPinNames::Outputs::NormalizedVelocityName);
+		UTEST_TRUE("Output exists", OutputNormalizedVelocity.IsSet());
+		UTEST_EQUAL("Normalized Midi Velocity (float) at start", **OutputNormalizedVelocity, 0.0f);
+
+		TOptional<FFloatReadRef> OutputFrequency= Generator->GetOutputReadReference<float>(CommonPinNames::Outputs::FrequencyName);
+		UTEST_TRUE("Output exists", OutputFrequency.IsSet());
+		UTEST_EQUAL("Normalized Midi Velocity (float) at start", **OutputFrequency, 0.0f);
+
+		//************************************************************************************************************************
+		// Make initial inputs
+		//************************************************************************************************************************
+		// Make a Midi File to play by Midi Player
+		UMidiFile* MidiFile = BuildMidiFile();
+		TSharedPtr<Audio::IProxyData> MidiFileProxy = MidiFile->CreateProxyData({ "MidiNoteTriggerTest" });
+		Generator->SetInputValue<FMidiAsset>(CommonPinNames::Inputs::MidiFileAssetName, MidiFileProxy);
+
+		// Set the testing Midi Channel (default = 1)
+		Generator->SetInputValue<int32>(CommonPinNames::Inputs::MidiChannelNumberName, TestMidiChannelShouldPass + 1);
+		// Set the testing Midi Track (default = 1)
+		Generator->SetInputValue<int32>(CommonPinNames::Inputs::MidiTrackNumberName, TestMidiTrackShouldPass);
+		// Set the testing Min Midi Note
+		Generator->SetInputValue<int32>(CommonPinNames::Inputs::MinMidiNoteName, TestMinNoteNumberDefault);
+		// Set the testing Max Midi Note
+		Generator->SetInputValue<int32>(CommonPinNames::Inputs::MaxMidiNoteName, TestMaxNoteNumberDefault);
+		// Set the testing Min Velocity
+		Generator->SetInputValue<int32>(CommonPinNames::Inputs::MinMidiVelocityName, TestMinVelocityDefault);
+		// Set the testing Max Velocity
+		Generator->SetInputValue<int32>(CommonPinNames::Inputs::MaxMidiVelocityName, TestMaxVelocityDefault);
+
+		// Trigger transport
+		Generator->ApplyToInputValue<FTrigger>(CommonPinNames::Inputs::TransportPlayName, [](FTrigger& TriggerIn) { TriggerIn.TriggerFrame(0); });
+
+		//************************************************************************************************************************
+		// Generate 5 10ms blocks for testing various Midi Events in file/stream
+		//************************************************************************************************************************
+
+		TAudioBuffer<float> Buffer{ Generator->GetNumChannels(), NumSamplesPerBlock, EAudioBufferCleanupMode::Delete};
+			
+			//Block 1:
+			Generator->OnGenerateAudio(Buffer.GetRawChannelData(0), Buffer.GetNumTotalValidSamples());
+			//Expected: 1 Note On Trigger, 0 Note Off Trigger, 
+			// Note Number/Velocity/Frequency outputs contain values equal to inputs
+				
+			//check output triggers
+			UTEST_EQUAL("Number of Note On Output Triggers", (*OutputNoteOnTrigger)->NumTriggeredInBlock(), ExpectedTriggerCountOnPass);
+			UTEST_EQUAL("Number of Note Off Output Triggers", (*OutputNoteOffTrigger)->NumTriggeredInBlock(), ExpectedTriggerCountNotPass);
+
+			//check output values
+			UTEST_EQUAL("Midi Note Number check", **OutputMidiNoteNumber, TestNoteNumberShouldPass);
+			UTEST_EQUAL("Midi Velocity (int32) check", **OutputMidiVelocity, TestNoteVelocityShouldPass);
+			UTEST_EQUAL("Midi Normalized Velocity (float) check", **OutputNormalizedVelocity, TestNoteNormalizedVelocityShouldPass);
+			UTEST_EQUAL("Output frequency check", **OutputFrequency, ExpectedFrequencyOnPass);
+
+			//Block 2:
+			Generator->OnGenerateAudio(Buffer.GetRawChannelData(0), Buffer.GetNumTotalValidSamples());
+			//Expected: 1 Note Off Trigger (for previous Note On that also should pass), 1 Note On Trigger,
+			//Note Number/Velocity/Frequency outputs contain values equal to inputs
+
+			//check output triggers
+			UTEST_EQUAL("Number of Note Off Output Triggers", (*OutputNoteOffTrigger)->NumTriggeredInBlock(), ExpectedTriggerCountOnPass);
+			UTEST_EQUAL("Number of Note On Output Triggers", (*OutputNoteOnTrigger)->NumTriggeredInBlock(), ExpectedTriggerCountOnPass);
+				
+			//check output values
+			UTEST_EQUAL("Midi Note Number check", **OutputMidiNoteNumber, TestNoteNumberShouldPass);
+			UTEST_EQUAL("Midi Velocity (int32) check", **OutputMidiVelocity, TestNoteVelocityShouldPass);
+			UTEST_EQUAL("Midi Normalized Velocity (float) check", **OutputNormalizedVelocity, TestNoteNormalizedVelocityShouldPass);
+			UTEST_EQUAL("Output frequency check", **OutputFrequency, ExpectedFrequencyOnPass);
+
+			//Block 3:
+			Generator->OnGenerateAudio(Buffer.GetRawChannelData(0), Buffer.GetNumTotalValidSamples());
+			//Expected: 1 Note Off Trigger (for previous Note On that also should pass), 0 Note On Trigger
+			//Note Number/Frequency should be the same as the previous block
+			//Velocity should be 0
+				
+			//check output triggers
+			UTEST_EQUAL("Number of Note Off Output Triggers", (*OutputNoteOffTrigger)->NumTriggeredInBlock(), ExpectedTriggerCountOnPass);
+			UTEST_EQUAL("Number of Note On Output Triggers", (*OutputNoteOnTrigger)->NumTriggeredInBlock(), ExpectedTriggerCountNotPass);
+
+			//check output values
+			UTEST_EQUAL("Midi Note Number check", **OutputMidiNoteNumber, TestNoteNumberShouldPass);
+			UTEST_EQUAL("Midi Velocity (int32) check", **OutputMidiVelocity, ExpectedVelocityNotPass);
+			UTEST_EQUAL("Midi Normalized Velocity (float) check", **OutputNormalizedVelocity, ExpectedNormalizedVelocityNotPass);
+			UTEST_EQUAL("Output frequency check", **OutputFrequency, ExpectedFrequencyOnPass);
+
+
+			//Block 4:
+			Generator->OnGenerateAudio(Buffer.GetRawChannelData(0), Buffer.GetNumTotalValidSamples());
+			//Expected: 0 Note On/Note Off trigger,
+			//Note Number/Frequency should be the same as the previous block
+			//Velocity should be 0
+
+			//check output triggers (should all be 0)
+			UTEST_EQUAL("Number of Note Off Output Triggers", (*OutputNoteOffTrigger)->NumTriggeredInBlock(), ExpectedTriggerCountNotPass);
+			UTEST_EQUAL("Number of Note On Output Triggers", (*OutputNoteOnTrigger)->NumTriggeredInBlock(), ExpectedTriggerCountNotPass);
+
+			//check output values
+			UTEST_EQUAL("Midi Note Number check", **OutputMidiNoteNumber, TestNoteNumberShouldPass);
+			UTEST_EQUAL("Midi Velocity (int32) check", **OutputMidiVelocity, ExpectedVelocityNotPass);
+			UTEST_EQUAL("Midi Normalized Velocity (float) check", **OutputNormalizedVelocity, ExpectedNormalizedVelocityNotPass);
+			UTEST_EQUAL("Output frequency check", **OutputFrequency, ExpectedFrequencyOnPass);
+
+			//Block 5:
+			Generator->OnGenerateAudio(Buffer.GetRawChannelData(0), Buffer.GetNumTotalValidSamples());
+			//Expected: 0 Note On/Note Off trigger,
+			//Note Number/Frequency should be the same as the previous block
+			//Velocity should be 0
+
+			//check output triggers (should all be 0)
+			UTEST_EQUAL("Number of Note Off Output Triggers", (*OutputNoteOffTrigger)->NumTriggeredInBlock(), ExpectedTriggerCountNotPass);
+			UTEST_EQUAL("Number of Note On Output Triggers", (*OutputNoteOnTrigger)->NumTriggeredInBlock(), ExpectedTriggerCountNotPass);
+
+			//check output values
+			UTEST_EQUAL("Midi Note Number check", **OutputMidiNoteNumber, TestNoteNumberShouldPass);
+			UTEST_EQUAL("Midi Velocity (int32) check", **OutputMidiVelocity, ExpectedVelocityNotPass);
+			UTEST_EQUAL("Midi Normalized Velocity (float) check", **OutputNormalizedVelocity, ExpectedNormalizedVelocityNotPass);
+			UTEST_EQUAL("Output frequency check", **OutputFrequency, ExpectedFrequencyOnPass);
+
+		return true;
+	}
+}
+
+#endif
diff --git a/Engine/Plugins/Runtime/Harmonix/Source/HarmonixMetasoundTests/Private/Nodes/MidiStreamTransposerNode.cpp b/Engine/Plugins/Runtime/Harmonix/Source/HarmonixMetasoundTests/Private/Nodes/MidiStreamTransposerNode.cpp
new file mode 100644
--- /dev/null
+++ b/Engine/Plugins/Runtime/Harmonix/Source/HarmonixMetasoundTests/Private/Nodes/MidiStreamTransposerNode.cpp
@@ -0,0 +1,44 @@
+// Copyright Epic Games, Inc. All Rights Reserved.
+
+#include "NodeTestGraphBuilder.h"
+#include "HarmonixDsp/AudioBuffer.h"
+#include "HarmonixMetasound/Common.h"
+#include "Misc/AutomationTest.h"
+
+#if WITH_DEV_AUTOMATION_TESTS
+
+namespace HarmonixMetasoundTests::MidiStreamTransposerNode
+{
+	using GraphBuilder = Metasound::Test::FNodeTestGraphBuilder;
+	using namespace Metasound;
+	using namespace Metasound::Frontend;
+	using namespace HarmonixMetasound;
+
+	IMPLEMENT_SIMPLE_AUTOMATION_TEST(
+		FMidiStreamTransposerCreateNodeTest,
+		"Harmonix.Metasound.Nodes.MidiStreamTransposerNode.CreateNode",
+		EAutomationTestFlags::EditorContext | EAutomationTestFlags::EngineFilter)
+		bool FMidiStreamTransposerCreateNodeTest::RunTest(const FString&)
+	{
+		// Build the graph.
+		constexpr int32 NumSamplesPerBlock = 256;
+		const TUniquePtr<FMetasoundGenerator> Generator = GraphBuilder::MakeSingleNodeGraph(
+			{ HarmonixMetasound::HarmonixNodeNamespace, "MidiStreamTransposer", "" },
+			0,
+			48000,
+			NumSamplesPerBlock);
+		UTEST_TRUE("Graph successfully built", Generator.IsValid());
+
+		// execute a block
+		{
+			TAudioBuffer<float> Buffer{ Generator->GetNumChannels(), NumSamplesPerBlock, EAudioBufferCleanupMode::Delete};
+			Generator->OnGenerateAudio(Buffer.GetRawChannelData(0), Buffer.GetNumTotalValidSamples());
+		}
+
+		// No output to validate, finish test
+
+		return true;
+	}
+}
+
+#endif
diff --git a/Engine/Plugins/Runtime/Harmonix/Source/HarmonixMetasoundTests/Private/Nodes/MidiVoiceManagerNodeTests.cpp b/Engine/Plugins/Runtime/Harmonix/Source/HarmonixMetasoundTests/Private/Nodes/MidiVoiceManagerNodeTests.cpp
new file mode 100644
--- /dev/null
+++ b/Engine/Plugins/Runtime/Harmonix/Source/HarmonixMetasoundTests/Private/Nodes/MidiVoiceManagerNodeTests.cpp
@@ -0,0 +1,45 @@
+// Copyright Epic Games, Inc. All Rights Reserved.
+
+#include "NodeTestGraphBuilder.h"
+#include "HarmonixDsp/AudioBuffer.h"
+#include "HarmonixMetasound/Common.h"
+#include "Misc/AutomationTest.h"
+
+#if WITH_DEV_AUTOMATION_TESTS
+
+namespace HarmonixMetasoundTests::MidiVoiceManagerNode
+{
+	using GraphBuilder = Metasound::Test::FNodeTestGraphBuilder;
+	using namespace Metasound;
+	using namespace Metasound::Frontend;
+	using namespace HarmonixMetasound;
+
+	IMPLEMENT_SIMPLE_AUTOMATION_TEST(
+		FMidiVoiceManagerCreateNodeTest,
+		"Harmonix.Metasound.Nodes.MidiVoiceManager.CreateNode",
+		EAutomationTestFlags::EditorContext | EAutomationTestFlags::EngineFilter)
+		bool FMidiVoiceManagerCreateNodeTest::RunTest(const FString&)
+	{
+		// Build graph with user-provided default value
+		const GraphBuilder Builder;
+		const FName NumVoicesInputName = "Num Voices";
+		const FNodeHandle VoiceManagerNode = Builder.AddNode({ HarmonixMetasound::HarmonixNodeNamespace, "MidiVoiceManager", "" }, 0);
+		const FNodeHandle NumVoicesInput = Builder.AddConstructorInput<int32>(NumVoicesInputName, 8);
+		const FOutputHandle OutputToConnect = NumVoicesInput->GetOutputWithVertexName(NumVoicesInputName);
+		const FInputHandle InputToConnect = VoiceManagerNode->GetInputWithVertexName(NumVoicesInputName);
+		UTEST_TRUE("Connected inputs", InputToConnect->Connect(*OutputToConnect));
+
+		// Build graph with vertex default
+		const TUniquePtr<FMetasoundGenerator> Generator = GraphBuilder::MakeSingleNodeGraph(
+			{ HarmonixMetasound::HarmonixNodeNamespace, "MidiVoiceManager", "" },
+			0,
+			48000,
+			128);
+		UTEST_TRUE("Generator created", Generator.IsValid());
+
+		return true;
+
+	}
+}
+
+#endif
diff --git a/Engine/Plugins/Runtime/Harmonix/Source/HarmonixMetasoundTests/Private/Nodes/MusicTimestampBuilderNodeTests.cpp b/Engine/Plugins/Runtime/Harmonix/Source/HarmonixMetasoundTests/Private/Nodes/MusicTimestampBuilderNodeTests.cpp
new file mode 100644
--- /dev/null
+++ b/Engine/Plugins/Runtime/Harmonix/Source/HarmonixMetasoundTests/Private/Nodes/MusicTimestampBuilderNodeTests.cpp
@@ -0,0 +1,44 @@
+// Copyright Epic Games, Inc. All Rights Reserved.
+
+#include "NodeTestGraphBuilder.h"
+#include "HarmonixDsp/AudioBuffer.h"
+#include "HarmonixMetasound/Common.h"
+#include "Misc/AutomationTest.h"
+
+#if WITH_DEV_AUTOMATION_TESTS
+
+namespace HarmonixMetasoundTests::MusicTimestampBuilderNode
+{
+	using GraphBuilder = Metasound::Test::FNodeTestGraphBuilder;
+	using namespace Metasound;
+	using namespace Metasound::Frontend;
+	using namespace HarmonixMetasound;
+
+	IMPLEMENT_SIMPLE_AUTOMATION_TEST(
+		FMusicTimestampBuilderCreateNodeTest,
+		"Harmonix.Metasound.Nodes.MusicTimestampBuilder.CreateNode",
+		EAutomationTestFlags::EditorContext | EAutomationTestFlags::EngineFilter)
+		bool FMusicTimestampBuilderCreateNodeTest::RunTest(const FString&)
+	{
+		// Build the graph.
+		constexpr int32 NumSamplesPerBlock = 256;
+		const TUniquePtr<FMetasoundGenerator> Generator = GraphBuilder::MakeSingleNodeGraph(
+			{ HarmonixMetasound::HarmonixNodeNamespace, "MusicTimestampBuilder", "" },
+			0,
+			48000,
+			NumSamplesPerBlock);
+		UTEST_TRUE("Graph successfully built", Generator.IsValid());
+
+		// execute a block
+		{
+			TAudioBuffer<float> Buffer{ Generator->GetNumChannels(), NumSamplesPerBlock, EAudioBufferCleanupMode::Delete};
+			Generator->OnGenerateAudio(Buffer.GetRawChannelData(0), Buffer.GetNumTotalValidSamples());
+		}
+
+		// No need to validate output pin, finish test
+
+		return true;
+	}
+}
+
+#endif
diff --git a/Engine/Plugins/Runtime/Harmonix/Source/HarmonixMetasoundTests/Private/Nodes/StepSequencePlayerNodeTests.cpp b/Engine/Plugins/Runtime/Harmonix/Source/HarmonixMetasoundTests/Private/Nodes/StepSequencePlayerNodeTests.cpp
new file mode 100644
--- /dev/null
+++ b/Engine/Plugins/Runtime/Harmonix/Source/HarmonixMetasoundTests/Private/Nodes/StepSequencePlayerNodeTests.cpp
@@ -0,0 +1,44 @@
+// Copyright Epic Games, Inc. All Rights Reserved.
+
+#include "NodeTestGraphBuilder.h"
+#include "HarmonixDsp/AudioBuffer.h"
+#include "HarmonixMetasound/Common.h"
+#include "Misc/AutomationTest.h"
+
+#if WITH_DEV_AUTOMATION_TESTS
+
+namespace HarmonixMetasoundTests::StepSequencePlayerNode
+{
+	using GraphBuilder = Metasound::Test::FNodeTestGraphBuilder;
+	using namespace Metasound;
+	using namespace Metasound::Frontend;
+	using namespace HarmonixMetasound;
+
+	IMPLEMENT_SIMPLE_AUTOMATION_TEST(
+		FStepSequencePlayerCreateNodeTest,
+		"Harmonix.Metasound.Nodes.StepSequencePlayerNode.CreateNode",
+		EAutomationTestFlags::EditorContext | EAutomationTestFlags::EngineFilter)
+		bool FStepSequencePlayerCreateNodeTest::RunTest(const FString&)
+	{
+		// Build the graph.
+		constexpr int32 NumSamplesPerBlock = 256;
+		const TUniquePtr<FMetasoundGenerator> Generator = GraphBuilder::MakeSingleNodeGraph(
+			{ HarmonixMetasound::HarmonixNodeNamespace, "StepSequencePlayer", "" },
+			0,
+			48000,
+			NumSamplesPerBlock);
+		UTEST_TRUE("Graph successfully built", Generator.IsValid());
+
+		// execute a block
+		{
+			TAudioBuffer<float> Buffer{ Generator->GetNumChannels(), NumSamplesPerBlock, EAudioBufferCleanupMode::Delete};
+			Generator->OnGenerateAudio(Buffer.GetRawChannelData(0), Buffer.GetNumTotalValidSamples());
+		}
+
+		// No output to validate, finish test
+
+		return true;
+	}
+}
+
+#endif
diff --git a/Engine/Plugins/Runtime/Harmonix/Source/HarmonixMetasoundTests/Private/Nodes/TempoMapMetronomeNodeTests.cpp b/Engine/Plugins/Runtime/Harmonix/Source/HarmonixMetasoundTests/Private/Nodes/TempoMapMetronomeNodeTests.cpp
new file mode 100644
--- /dev/null
+++ b/Engine/Plugins/Runtime/Harmonix/Source/HarmonixMetasoundTests/Private/Nodes/TempoMapMetronomeNodeTests.cpp
@@ -0,0 +1,48 @@
+// Copyright Epic Games, Inc. All Rights Reserved.
+
+#include "NodeTestGraphBuilder.h"
+#include "HarmonixDsp/AudioBuffer.h"
+#include "HarmonixMetasound/Common.h"
+#include "HarmonixMetasound/DataTypes/MidiClock.h"
+#include "Misc/AutomationTest.h"
+
+#if WITH_DEV_AUTOMATION_TESTS
+
+namespace HarmonixMetasoundTests::TempoMapMetronomeNode
+{
+	using GraphBuilder = Metasound::Test::FNodeTestGraphBuilder;
+	using namespace Metasound;
+	using namespace Metasound::Frontend;
+	using namespace HarmonixMetasound;
+
+	IMPLEMENT_SIMPLE_AUTOMATION_TEST(
+		FTempoMapMetronomeCreateNodeTest,
+		"Harmonix.Metasound.Nodes.TempoMapMetronomeNode.CreateNode",
+		EAutomationTestFlags::EditorContext | EAutomationTestFlags::EngineFilter)
+		bool FTempoMapMetronomeCreateNodeTest::RunTest(const FString&)
+	{
+		// Build the graph.
+		constexpr int32 NumSamplesPerBlock = 256;
+		const TUniquePtr<FMetasoundGenerator> Generator = GraphBuilder::MakeSingleNodeGraph(
+			{ HarmonixMetasound::HarmonixNodeNamespace, "TempoMapMetronome", "" },
+			0,
+			48000,
+			NumSamplesPerBlock);
+		UTEST_TRUE("Graph successfully built", Generator.IsValid());
+
+		// execute a block
+		{
+			TAudioBuffer<float> Buffer{ Generator->GetNumChannels(), NumSamplesPerBlock, EAudioBufferCleanupMode::Delete};
+			Generator->OnGenerateAudio(Buffer.GetRawChannelData(0), Buffer.GetNumTotalValidSamples());
+		}
+
+		// Validate output.
+		TOptional<FMidiClockReadRef> OutputMidiClock = Generator->GetOutputReadReference<FMidiClock>(CommonPinNames::Outputs::MidiClockName);
+		UTEST_TRUE("Output exists", OutputMidiClock.IsSet());
+		UTEST_EQUAL("Current Midi Tick Test", (*OutputMidiClock)->GetCurrentMidiTick(), -1);
+
+		return true;
+	}
+}
+
+#endif
diff --git a/Engine/Plugins/Runtime/Harmonix/Source/HarmonixMetasoundTests/Private/Subsystems/MidiClockUpdateSubsystemTests.cpp b/Engine/Plugins/Runtime/Harmonix/Source/HarmonixMetasoundTests/Private/Subsystems/MidiClockUpdateSubsystemTests.cpp
new file mode 100644
--- /dev/null
+++ b/Engine/Plugins/Runtime/Harmonix/Source/HarmonixMetasoundTests/Private/Subsystems/MidiClockUpdateSubsystemTests.cpp
@@ -0,0 +1,97 @@
+// Copyright Epic Games, Inc. All Rights Reserved.
+
+#include "Engine/Engine.h"
+#include "HarmonixMetasound/DataTypes/MidiClock.h"
+#include "HarmonixMetasound/Subsystems/MidiClockUpdateSubsystem.h"
+#include "Misc/AutomationTest.h"
+
+#if WITH_DEV_AUTOMATION_TESTS
+
+namespace HarmonixMetasoundTests::MidiClockUpdateSubsystem
+{
+	namespace Helpers
+	{
+		HarmonixMetasound::FMidiClock MakeAndStartClock(
+			const Metasound::FOperatorSettings& OperatorSettings,
+			float Tempo,
+			int32 TimeSigNum,
+			int32 TimeSigDenom)
+		{
+			const TSharedPtr<FMidiFileData> MidiData = MakeShared<FMidiFileData>();
+			check(MidiData);
+			
+			MidiData->Tracks.Add(FMidiTrack(TEXT("conductor")));
+			MidiData->LastEventTick = std::numeric_limits<int32>::max();
+			MidiData->Tracks[0].AddEvent(FMidiEvent(0, FMidiMsg(static_cast<uint8>(TimeSigNum), static_cast<uint8>(TimeSigDenom))));
+			FBarMap& BarMap = MidiData->SongMaps.GetBarMap();
+			BarMap.AddTimeSignatureAtBarIncludingCountIn(0, TimeSigNum, TimeSigDenom);
+			const int32 MidiTempo = MidiConstants::BPMToMidiTempo(Tempo);
+			MidiData->Tracks[0].AddEvent(FMidiEvent(0, FMidiMsg(MidiTempo)));
+			FTempoMap& TempoMap = MidiData->SongMaps.GetTempoMap();
+			TempoMap.AddTempoInfoPoint(MidiTempo, 0);
+			MidiData->Tracks[0].Sort();
+			MidiData->SongMaps.GetSongLengthData().LastTick = std::numeric_limits<int32>::max();
+			MidiData->SongMaps.GetSongLengthData().LengthTicks = std::numeric_limits<int32>::max();
+			MidiData->SongMaps.GetSongLengthData().LengthBars = std::numeric_limits<int32>::max();
+
+			HarmonixMetasound::FMidiClock Clock{ OperatorSettings };
+			Clock.AttachToMidiResource(MidiData);
+			Clock.ResetAndStart(0);
+
+			return Clock;
+		}	
+	}
+	
+	IMPLEMENT_SIMPLE_AUTOMATION_TEST(
+	FMidiClockUpdateSubsystemBasicTest,
+	"Harmonix.Metasound.Subsystems.MidiClockUpdateSubsystem.Basic",
+	EAutomationTestFlags::EditorContext | EAutomationTestFlags::EngineFilter)
+	bool FMidiClockUpdateSubsystemBasicTest::RunTest(const FString&)
+	{
+		using namespace HarmonixMetasound;
+
+		UTEST_NOT_NULL("GEngine exists", GEngine);
+		UMidiClockUpdateSubsystem* Subsystem = GEngine->GetEngineSubsystem<UMidiClockUpdateSubsystem>();
+		UTEST_NOT_NULL("Subsystem exists", Subsystem);
+
+		constexpr float Tempo = 95;
+		constexpr int32 TimeSigNum = 3;
+		constexpr int32 TimeSigDenom = 4;
+		Metasound::FOperatorSettings OperatorSettings { 48000, 100 };
+		
+		FMidiClock Clock = Helpers::MakeAndStartClock(OperatorSettings, Tempo, TimeSigNum, TimeSigDenom);
+
+		// Update the clock and tick the subsystem a few times to make sure we're updating the low-r clock
+		constexpr int32 NumIterations = 100;
+		const int32 NumSamples = OperatorSettings.GetNumFramesPerBlock();
+		int32 SampleRemainder = 0;
+		int32 SampleCount = 0;
+		
+		for (int32 i = 0; i < NumIterations; ++i)
+		{
+			// Advance the high-resolution clock
+			SampleRemainder += NumSamples;
+			constexpr int32 MidiGranularity = 128;
+			while (SampleRemainder >= MidiGranularity)
+			{
+				SampleCount += MidiGranularity;
+				SampleRemainder -= MidiGranularity;
+				const float AdvanceToMs = static_cast<float>(SampleCount) * 1000.0f / OperatorSettings.GetSampleRate();
+				Clock.AdvanceHiResToMs(0, AdvanceToMs, true);
+			}
+
+			// Tick the subsystem (low-resolution clocks)
+			Subsystem->Tick(0);
+
+			// Check that the high- and low-resolution clocks are at the same place
+			UTEST_EQUAL(
+				FString::Printf(TEXT("High- and low-res are at the same tick: iteration %i"), i),
+				Clock.GetCurrentHiResTick(),
+				Clock.GetCurrentLowResTick());
+		}
+
+		return true;
+	}
+}
+
+#endif
\ No newline at end of file
diff --git a/Engine/Plugins/Runtime/Harmonix/Source/HarmonixMidi/HarmonixMidi.Build.cs b/Engine/Plugins/Runtime/Harmonix/Source/HarmonixMidi/HarmonixMidi.Build.cs
new file mode 100644
--- /dev/null
+++ b/Engine/Plugins/Runtime/Harmonix/Source/HarmonixMidi/HarmonixMidi.Build.cs
@@ -0,0 +1,40 @@
+// Copyright Epic Games, Inc. All Rights Reserved.
+
+using UnrealBuildTool;
+
+public class HarmonixMidi : ModuleRules
+{
+	public HarmonixMidi(ReadOnlyTargetRules Target) : base(Target)
+	{
+		//OptimizeCode = CodeOptimization.Never;
+
+		PCHUsage = ModuleRules.PCHUsageMode.UseExplicitOrSharedPCHs;
+
+		PublicDependencyModuleNames.AddRange(
+			new string[]
+			{
+				"Core",
+				"CoreUObject",
+			});
+		PrivateDependencyModuleNames.AddRange(
+			new string[]
+			{
+				"AudioExtensions",
+				"Harmonix",
+			});
+		if (Target.bBuildEditor == true)
+		{
+			PrivateDependencyModuleNames.Add("AssetRegistry");
+			PrivateDependencyModuleNames.Add("UnrealEd");
+			PrivateDependencyModuleNames.Add("Engine");
+		}
+		else
+		{
+			// This next little gem is needed because when this module is built without
+			// an Engine dependency the MidiFile.gen.cpp file has an UNUSED forward declaration
+			// at the top of it that is prefixed by "ENGINE_API", so the compiler sees an undefined
+			// symbol. :-( - Buzz
+			PrivateDefinitions.Add("ENGINE_API=");
+		}
+	}
+}
diff --git a/Engine/Plugins/Runtime/Harmonix/Source/HarmonixMidi/Private/HarmonixMidi/BarMap.cpp b/Engine/Plugins/Runtime/Harmonix/Source/HarmonixMidi/Private/HarmonixMidi/BarMap.cpp
new file mode 100644
--- /dev/null
+++ b/Engine/Plugins/Runtime/Harmonix/Source/HarmonixMidi/Private/HarmonixMidi/BarMap.cpp
@@ -0,0 +1,447 @@
+// Copyright Epic Games, Inc. All Rights Reserved.
+#include "HarmonixMidi/BarMap.h"
+#include "HarmonixMidi/MidiConstants.h"
+#include "HarmonixMidi/TempoMap.h"
+#include "HarmonixMidi/BeatMap.h" 
+
+#include "Misc/ScopeLock.h"
+#include "Serialization/Archive.h"
+
+
+void FBarMap::Empty()
+{
+	Points.Empty();
+}
+
+void FBarMap::Copy(const FBarMap& Other, int32 StartTick, int32 EndTick)
+{
+	TicksPerQuarterNote = Other.TicksPerQuarterNote;
+	FMusicMapUtl::Copy(Other.Points, Points, StartTick, EndTick);
+}
+
+bool FBarMap::IsEmpty() const
+{
+	return Points.IsEmpty();
+}
+
+int32 FBarMap::MusicTimestampToTick(const FMusicTimestamp& Timestamp) const
+{
+	float Beat = Timestamp.Beat;
+	if (Beat < 1.0f)
+	{
+		if (Beat >= 0.0f)
+		{
+			UE_LOG(LogMidi, Warning, TEXT("Beat %f specified in Music Timestamp! Beats in Music Timestamps are '1' based. Beat 1.0 is the first beat in the specified bar. Adding 1.0 to the specified Beat!"), Beat);
+			Beat += 1.0f;
+		}
+		else
+		{
+			UE_LOG(LogMidi, Warning, TEXT("Negative beat (%f) specified in Music Timestamp. This is not allowed! Beats are '1' based in Music Timestamps and must be positive. Using beat 1!"), Beat);
+			Beat = 1.0f;
+		}
+	}
+
+	int32 BarIndex = Timestamp.Bar - StartBar;
+	Beat -= 1.0f;
+	if (Points.IsEmpty())
+	{
+		return	BarIndex * 4 * TicksPerQuarterNote +
+			FMath::FloorToInt(Beat * TicksPerQuarterNote);
+	}
+	
+	int32 Index = Algo::UpperBound(Points, BarIndex, FTimeSignaturePoint::BarLessThan());
+
+	if (Index > 0)
+	{
+		--Index;
+	}
+
+	int32 BarDelta = BarIndex - Points[Index].BarIndex;
+	return (int32)(Points[Index].StartTick +
+		BarDelta * GetTicksInBarAfterPoint(Index) +
+		(int32)(Beat * GetTicksInBeatAfterPoint(Index)));
+}
+
+int32 FBarMap::MusicTimestampBarToTick(int32 InBar, int32* OutBeatsPerBar, int32* OutTicksPerBeat) const
+{
+	int32 BarIndex = InBar - StartBar;
+	return BarIncludingCountInToTick(BarIndex, OutBeatsPerBar, OutTicksPerBeat);
+}
+
+int32 FBarMap::BarIncludingCountInToTick(int32 BarIndex, int32 * OutBeatsPerBar, int32 * OutTicksPerBeat) const
+{
+	if (Points.IsEmpty())
+	{
+		if (OutBeatsPerBar)
+		{
+			*OutBeatsPerBar = 4;
+		}
+		if (OutTicksPerBeat)
+		{
+			*OutTicksPerBeat = TicksPerQuarterNote;
+		}
+		return	BarIndex * 4 * TicksPerQuarterNote;
+	}
+
+	int32 TimeSigIndex = Algo::UpperBound(Points, BarIndex, FTimeSignaturePoint::BarLessThan());
+
+	if (TimeSigIndex > 0)
+	{
+		--TimeSigIndex;
+	}
+
+	if (OutBeatsPerBar)
+	{
+		*OutBeatsPerBar = Points[TimeSigIndex].TimeSignature.Numerator;
+	}
+	if (OutTicksPerBeat)
+	{
+		*OutTicksPerBeat = GetTicksInBeatAfterPoint(TimeSigIndex);
+	}
+	int32 BarDelta = BarIndex - Points[TimeSigIndex].BarIndex;
+	return Points[TimeSigIndex].StartTick + (BarDelta * GetTicksInBarAfterPoint(TimeSigIndex));
+}
+
+int32 FBarMap::MusicTimestampBarBeatTickToTick(int32 Bar, int32 BeatInBar, int32 TickInBeat) const
+{
+	int32 BarIndex = Bar - StartBar;
+	return BarBeatTickIncludingCountInToTick(BarIndex, BeatInBar, TickInBeat);
+}
+
+int32 FBarMap::BarBeatTickIncludingCountInToTick(int32 BarIndex, int32 BeatInBar, int32 TickInBeat) const
+{
+	if (BeatInBar < 1)
+	{
+		UE_LOG(LogMidi, Warning, TEXT("Beat %d specified as a \"beat in bar\". Beat 1 is the first beat in a bar. Using Beat 1!"), BeatInBar);
+		BeatInBar = 0;
+	}
+	else
+	{
+		BeatInBar--;
+	}
+
+	if (Points.IsEmpty())
+	{
+		// Assume 4/4 time.
+		return	(BarIndex * 4 * TicksPerQuarterNote) +
+				(BeatInBar * TicksPerQuarterNote) +
+			     TickInBeat;
+	}
+
+	int32 TimeSigIndex = Algo::UpperBound(Points, BarIndex, FTimeSignaturePoint::BarLessThan());
+
+	if (TimeSigIndex > 0)
+	{
+		--TimeSigIndex;
+	}
+
+	int32 BarDelta = BarIndex - Points[TimeSigIndex].BarIndex;
+	int32 TicksInBeat = GetTicksInBeatAfterPoint(TimeSigIndex);
+	if (BeatInBar >= Points[TimeSigIndex].TimeSignature.Numerator)
+	{
+		UE_LOG(LogMidi, Warning, TEXT("BarBeatTickToAbsoluteTick: Supplied 'BeatInBar' is greater than the number of beats in the specified bar!"));
+	}
+	if (TickInBeat >= TicksInBeat)
+	{
+		UE_LOG(LogMidi, Warning, TEXT("BarBeatTickToAbsoluteTick: Supplied 'TickInBeat' is greater than the number of ticks in each beat at the specified bar!"));
+	}
+	return Points[TimeSigIndex].StartTick +
+			(BarDelta * GetTicksInBarAfterPoint(TimeSigIndex)) +
+			(BeatInBar * TickInBeat) +
+			TickInBeat;
+}
+
+FMusicTimestamp FBarMap::TickToMusicTimestamp(int32 Tick, int32* OutBeatsPerBar) const
+{
+	FMusicTimestamp Result;
+	if (Points.IsEmpty())
+	{
+		// Assume 4/4 time.
+		Result.Bar = Tick / (TicksPerQuarterNote * 4);
+		Tick -= (Result.Bar * TicksPerQuarterNote * 4);
+		if (Tick < 0)
+		{
+			Result.Bar--;
+			Tick += TicksPerQuarterNote * 4;
+		}
+		Result.Beat = (float)Tick / (float)TicksPerQuarterNote;
+		Result.Bar += StartBar;
+		Result.Beat += 1.0f; // 1 based
+		if (OutBeatsPerBar) *OutBeatsPerBar = 4;
+		return Result;
+	}
+
+	int32 TimeSigIndex = FMusicMapUtl::GetPointIndexForTick(Points, Tick);
+	if (TimeSigIndex < 0)
+	{
+		TimeSigIndex = 0;
+	}
+
+	int32 TicksPerBar = GetTicksInBarAfterPoint(TimeSigIndex);
+	int32 TicksPassed = Tick - Points[TimeSigIndex].StartTick;
+	if (TicksPassed < 0)
+	{
+		int32 BarsPassed = (- TicksPassed / TicksPerBar) + 1;
+		TicksPassed += BarsPassed * TicksPerBar;
+		float BeatsPassed = (float)TicksPassed / (float)GetTicksInBeatAfterPoint(TimeSigIndex);
+		Result.Bar = -BarsPassed + StartBar; // already - 1 based
+		Result.Beat = BeatsPassed + 1.0f; // 1 based
+	}
+	else
+	{
+		int32 BarsPassed = TicksPassed / TicksPerBar;
+		TicksPassed -= BarsPassed * TicksPerBar;
+		float BeatsPassed = (float)TicksPassed / (float)GetTicksInBeatAfterPoint(TimeSigIndex);
+		Result.Bar = Points[TimeSigIndex].BarIndex + BarsPassed + StartBar; // 1 based
+		Result.Beat = BeatsPassed + 1.0; // 1 based
+	}
+	if (OutBeatsPerBar) *OutBeatsPerBar = Points[TimeSigIndex].TimeSignature.Numerator;
+	return Result;
+}
+
+FMusicTimestamp FBarMap::TickFromBarOneToMusicTimestamp(int32 InTickFromBarOne, int32* OutBeatsPerBar) const
+{
+	return TickToMusicTimestamp(InTickFromBarOne + GetTickOfBarOne(), OutBeatsPerBar);
+}
+/*
+
+void FBarMap::TickToRawBarBeatTick(int32 InTick, int32& OutBar, int32& OutBeat, int32& OutTickInBeat, int32* OutBeatsPerBar, int32* OutTicksPerBeat) const
+{
+	if (Points.IsEmpty())
+	{
+		// Assume 4/4 time.
+		if (InTick == 0)
+		{
+			OutBar = 1;
+			OutBeat = 1;
+			OutTickInBeat = 0;
+		}
+		else if (InTick > 0)
+		{
+			OutBar = InTick / (TicksPerQuarterNote * 4);
+			InTick -= (OutBar * TicksPerQuarterNote * 4);
+			OutBeat = (float)InTick / (float)TicksPerQuarterNote;
+			OutTickInBeat = InTick - (OutBeat * TicksPerQuarterNote);
+			OutBar += 1; // 1 based
+			OutBeat += 1; // 1 based
+		}
+		else // InTick < 0
+		{
+			OutBar = InTick / (TicksPerQuarterNote * 4);
+			OutBar -= 1; // 1 based
+			InTick -= (OutBar * TicksPerQuarterNote * 4);
+			OutBeat = (float)InTick / (float)TicksPerQuarterNote;
+			OutTickInBeat = InTick - (OutBeat * TicksPerQuarterNote);
+			OutBeat += 1;
+		}
+		if (OutBeatsPerBar) *OutBeatsPerBar = 4;
+		if (OutTicksPerBeat) *OutTicksPerBeat = TicksPerQuarterNote;
+		return;
+	}
+
+	int32 Index = FMusicMapUtl::GetPointIndexForTick(Points, InTick);
+	if (Index < 0)
+	{
+		Index = 0;
+	}
+
+	int32 TicksPerBar = GetTicksInBarAfterPoint(Index);
+	int32 TicksPerBeat = GetTicksInBeatAfterPoint(Index);
+	if (InTick < 0)
+	{
+		// before the beginning? 
+		int32 BarsPassed = (-InTick / TicksPerBar) + 1;
+		int32 TickInBar = (BarsPassed * TicksPerBar) + InTick;
+		int32 BeatsPassed = TickInBar / TicksPerBeat;
+		int32 TickInBeat = TickInBar - (BeatsPassed * TicksPerBeat);
+		OutBar = -BarsPassed; // already 1 based
+		OutBeat = BeatsPassed + 1; // 1 based
+		OutTickInBeat = TickInBeat; // 0 based
+	}
+	else
+	{
+		int32 TicksPassed = InTick - Points[Index].StartTick;
+		int32 BarsPassed = TicksPassed / TicksPerBar;
+		TicksPassed -= BarsPassed * TicksPerBar;
+		int32 BeatsPassed = TicksPassed / TicksPerBeat;
+		TicksPassed -= BeatsPassed * TicksPerBeat;
+		OutBar = Points[Index].Bar + BarsPassed + 1; // 1 based
+		OutBeat = BeatsPassed + 1; // 1 based
+		OutTickInBeat = TicksPassed; // 0 based
+	}
+	if (OutBeatsPerBar) *OutBeatsPerBar = Points[Index].TimeSignature.Numerator;
+	if (OutTicksPerBeat) *OutTicksPerBeat = TicksPerBeat;
+}
+*/
+
+/*
+float FBarMap::TickToFractionalMusicalTimestampBar(int32 Tick) const
+{
+	int32 BeatsPerBar;
+	FMusicTimestamp Timestamp = TickToMusicTimestamp(Tick, &BeatsPerBar);
+	return (float)Timestamp.Bar + ((Timestamp.Beat - 1.0f) / (float)BeatsPerBar);
+}
+*/
+
+float FBarMap::TickToFractionalBarIncludingCountIn(int32 Tick) const
+{
+	if (Points.IsEmpty())
+	{
+		// assume 4/4
+		return (float)Tick / (float)(TicksPerQuarterNote * 4);
+	}
+
+	int32 TimeSigIndex = FMusicMapUtl::GetPointIndexForTick(Points, Tick);
+	if (TimeSigIndex < 0)
+	{
+		TimeSigIndex = 0;
+	}
+	int32 TicksPast = Tick - Points[TimeSigIndex].StartTick;
+	int32 TicksPerBar = GetTicksInBarAfterPoint(TimeSigIndex);
+	float BarsPassed = (float)TicksPast / (float)TicksPerBar;
+	return (float)Points[TimeSigIndex].BarIndex + BarsPassed;
+}
+
+int32 FBarMap::TickToBarIncludingCountIn(int32 Tick) const
+{
+	if (Points.IsEmpty())
+	{
+		// assume 4/4
+		return FMath::FloorToInt32((float)Tick / (float)(TicksPerQuarterNote * 4));
+	}
+	int32 TimeSigIndex = FMusicMapUtl::GetPointIndexForTick(Points, Tick);
+	if (TimeSigIndex < 0)
+	{
+		TimeSigIndex = 0;
+	}
+	int32 TicksPast = Tick - Points[TimeSigIndex].StartTick;
+	int32 TicksPerBar = GetTicksInBarAfterPoint(TimeSigIndex);
+	int32 BarsPassed = TicksPast / TicksPerBar;
+	return Points[TimeSigIndex].BarIndex + BarsPassed;
+}
+
+int32 FBarMap::FractionalBarIncludingCountInToTick(float Bar) const
+{
+	if (Points.IsEmpty())
+	{
+		// assume 4/4
+		return FMath::FloorToInt32(Bar * 4.0f * TicksPerQuarterNote);
+	}
+
+	int32 BarIndex = FMath::FloorToInt32(Bar);
+	int32 TimeSigIndex = Algo::UpperBound(Points, BarIndex, FTimeSignaturePoint::BarLessThan());
+
+	if (TimeSigIndex > 0)
+	{
+		--TimeSigIndex;
+	}
+
+	float BarsPassed = Bar - (float)Points[TimeSigIndex].BarIndex;
+	int32 TicksPerBar = GetTicksInBarAfterPoint(TimeSigIndex);
+	return FMath::FloorToInt32((float)Points[TimeSigIndex].StartTick + (BarsPassed * (float)TicksPerBar));
+}
+
+bool FBarMap::AddTimeSignatureAtMusicTimestampBar(int32 BarNumber, int32 Numerator, int32 Denominator, bool SortNow, bool FailOnError)
+{
+	int BarIndex = BarNumber - StartBar;
+	return AddTimeSignatureAtBarIncludingCountIn(BarIndex, Numerator, Denominator, SortNow, FailOnError);
+}
+
+bool FBarMap::AddTimeSignatureAtBarIncludingCountIn(int32 BarIndex, int32 Numerator, int32 Denominator, bool SortNow, bool FailOnError)
+{
+	if (BarIndex < 0)
+	{
+		BarIndex = 0;
+	}
+
+	if (BarIndex == 0)
+	{
+		if (!Points.IsEmpty())
+		{
+			if (FailOnError)
+			{
+				checkf(false, TEXT("Multiple time signatures at start of song"));
+			}
+			return false;
+		}
+		FMusicMapUtl::AddPoint<FTimeSignaturePoint, int32, int32, const FTimeSignature&>(0, 0, FTimeSignature(Numerator, Denominator), Points, 0, SortNow);
+	}
+	else
+	{
+		FTimeSignaturePoint& LastSig = Points.Last();
+		int32 BarsPassed = BarIndex - LastSig.BarIndex;
+		if (BarsPassed <= 0)
+		{
+			if (FailOnError)
+			{
+				checkf(false, TEXT("Multiple time signatures at bar %d"), BarIndex + StartBar);
+			}
+			return false;
+		}
+		int32 TicksPassed  = BarsPassed * GetTicksInBarAfterPoint(Points.Num()-1);
+		int32 ThisTick     = LastSig.StartTick + TicksPassed;
+		int32 BeatsPassed = BarsPassed * LastSig.TimeSignature.Numerator;
+		int32 ThisBeat    = LastSig.BeatIndex + BeatsPassed;
+		FMusicMapUtl::AddPoint<FTimeSignaturePoint, int32, int32, const FTimeSignature&>(BarIndex, ThisBeat, FTimeSignature(Numerator, Denominator), Points, ThisTick, SortNow);
+	}
+	return true;
+}
+
+int32 FBarMap::GetNumTimeSignaturePoints() const
+{
+	return Points.Num();
+}
+
+int32 FBarMap::GetPointIndexForTick(int32 Tick) const
+{
+	return FMusicMapUtl::GetPointIndexForTick(Points, Tick);
+}
+
+const FTimeSignaturePoint& FBarMap::GetTimeSignaturePoint(int32 Index) const
+{
+	return Points[Index];
+}
+
+FTimeSignaturePoint& FBarMap::GetTimeSignaturePoint(int32 Index)
+{
+	return Points[Index];
+}
+
+const FTimeSignature& FBarMap::GetTimeSignatureAtTick(int32 Tick) const
+{
+	check(!Points.IsEmpty());
+		
+	int32 Index = FMusicMapUtl::GetPointIndexForTick(Points, Tick);
+
+	if (Index < 0)
+	{
+		Index = 0;
+	}
+
+	return Points[Index].TimeSignature;
+}
+
+const FTimeSignature& FBarMap::GetTimeSignatureAtBar(int32 Bar) const
+{
+	check(!Points.IsEmpty());
+
+	Bar -= StartBar;
+
+	int32 Index = Algo::UpperBound(Points, Bar, FTimeSignaturePoint::BarLessThan());
+
+	if (Index > 0)
+	{
+		--Index;
+	}
+
+	return Points[Index].TimeSignature;
+}
+
+void FBarMap::Finalize(int32 InLastTick)
+{
+	if (Points.IsEmpty())
+	{
+		SupplyDefault();
+	}
+	FMusicMapUtl::Finalize(Points, InLastTick);
+}
diff --git a/Engine/Plugins/Runtime/Harmonix/Source/HarmonixMidi/Private/HarmonixMidi/Blueprint/MidiFunctionLibrary.cpp b/Engine/Plugins/Runtime/Harmonix/Source/HarmonixMidi/Private/HarmonixMidi/Blueprint/MidiFunctionLibrary.cpp
new file mode 100644
--- /dev/null
+++ b/Engine/Plugins/Runtime/Harmonix/Source/HarmonixMidi/Private/HarmonixMidi/Blueprint/MidiFunctionLibrary.cpp
@@ -0,0 +1,68 @@
+// Copyright Epic Games, Inc. All Rights Reserved.
+
+#include "HarmonixMidi/Blueprint/MidiFunctionLibrary.h"
+
+#include "HarmonixMidi/MidiConstants.h"
+
+#include UE_INLINE_GENERATED_CPP_BY_NAME(MidiFunctionLibrary)
+
+FMidiNote UMidiNoteFunctionLibrary::GetMinMidiNote()
+{
+	return FMidiNote(MidiConstants::kMinNote);
+}
+
+FMidiNote UMidiNoteFunctionLibrary::GetMaxMidiNote()
+{
+	return FMidiNote(MidiConstants::kMaxNote);
+}
+
+uint8 UMidiNoteFunctionLibrary::GetMinNoteNumber()
+{
+	return MidiConstants::kMinNote;
+}
+
+uint8 UMidiNoteFunctionLibrary::GetMaxNoteNumber()
+{
+	return MidiConstants::kMaxNote;
+}
+
+int UMidiNoteFunctionLibrary::GetMaxNumNotes()
+{
+	return MidiConstants::kMaxNumNotes;
+}
+
+uint8 UMidiNoteFunctionLibrary::GetMinNoteVelocity()
+{
+	return MidiConstants::kMinVelocity;
+}
+
+uint8 UMidiNoteFunctionLibrary::GetMaxNoteVelocity()
+{
+	return MidiConstants::kMaxVelocity;
+}
+
+float UMusicalTickFunctionLibrary::GetTicksPerQuarterNote()
+{
+	return MidiConstants::kTicksPerQuarterNote;
+}
+
+int32 UMusicalTickFunctionLibrary::GetTicksPerQuarterNoteInt()
+{
+	return MidiConstants::kTicksPerQuarterNoteInt;
+}
+
+float UMusicalTickFunctionLibrary::GetQuarterNotesPerTick()
+{
+	return MidiConstants::kQuarterNotesPerTick;
+}
+
+float UMusicalTickFunctionLibrary::TickToQuarterNote(float InTick)
+{
+	return InTick * MidiConstants::kQuarterNotesPerTick;
+}
+
+float UMusicalTickFunctionLibrary::QuarterNoteToTick(float InQuarterNote)
+{
+	return InQuarterNote * MidiConstants::kTicksPerQuarterNote;
+}
+
diff --git a/Engine/Plugins/Runtime/Harmonix/Source/HarmonixMidi/Private/HarmonixMidi/MidiConstants.cpp b/Engine/Plugins/Runtime/Harmonix/Source/HarmonixMidi/Private/HarmonixMidi/MidiConstants.cpp
new file mode 100644
--- /dev/null
+++ b/Engine/Plugins/Runtime/Harmonix/Source/HarmonixMidi/Private/HarmonixMidi/MidiConstants.cpp
@@ -0,0 +1,225 @@
+// Copyright Epic Games, Inc. All Rights Reserved.
+#include "HarmonixMidi/MidiConstants.h"
+
+DEFINE_LOG_CATEGORY(LogMidi);
+
+FString MidiConstants::GetTextTypeName(uint8 TextType)
+{
+	switch (TextType)
+	{
+	case kMeta_Text: return TEXT("Text");
+	case kMeta_Copyright: return TEXT("Copyright");
+	case kMeta_TrackName: return TEXT("Track Name");
+	case kMeta_InstrumentName: return TEXT("Instrument Name");
+	case kMeta_Lyric: return TEXT("Lyric");
+	case kMeta_Marker: return TEXT("Marker");
+	case kMeta_CuePoint: return TEXT("Cue Point");
+	default:
+		checkNoEntry();
+		break;
+	}
+	return TEXT("");
+}
+
+float MidiConstants::RoundToStandardBeatPrecision(float InBeat, int TimeSignatureDenominator)
+{
+	// this is the precision to which the "rounded" elapsed beats get presented.
+	const double kQuarterNotePrecision = 0.01 * ((double)TimeSignatureDenominator / 4.0); // account for beat being a quarter note, eighth note, etc.
+	double x = (double)InBeat;
+	x += 0.5 * kQuarterNotePrecision;
+	x *= 1.0 / kQuarterNotePrecision;
+	x = FMath::Floor(x);
+	x *= kQuarterNotePrecision;
+	return x;
+}
+
+FString MidiConstants::GetControllerName(EControllerID ControllerId)
+{
+	switch (ControllerId)
+	{
+	case EControllerID::BankSelection:       return FString::Format(TEXT("({0}) BankSelection"), { ControllerId });
+	case EControllerID::ModWheel:			 return FString::Format(TEXT("({0}) ModWheel"), { ControllerId });
+	case EControllerID::Breath:				 return FString::Format(TEXT("({0}) Breath"), { ControllerId });
+	case EControllerID::PortamentoTime:		 return FString::Format(TEXT("({0}) PortamentoTime"), { ControllerId });
+	case EControllerID::DataCoarse:			 return FString::Format(TEXT("({0}) DataCoarse"), { ControllerId });
+	case EControllerID::Volume:				 return FString::Format(TEXT("({0}) Volume"), { ControllerId });
+	case EControllerID::Balance:			 return FString::Format(TEXT("({0}) Balance"), { ControllerId });
+	case EControllerID::PanRight:			 return FString::Format(TEXT("({0}) Pan (0 = left, 1 = right)"), { ControllerId });
+	case EControllerID::Expression:			 return FString::Format(TEXT("({0}) Expression"), { ControllerId });
+	case EControllerID::BitCrushWetMix:		 return FString::Format(TEXT("({0}) BitCrushWetMix"), { ControllerId });
+	case EControllerID::BitCrushLevel:		 return FString::Format(TEXT("({0}) BitCrushLevel"), { ControllerId });
+	case EControllerID::BitCrushSampleHold:	 return FString::Format(TEXT("({0}) BitCrushSampleHold"), { ControllerId });
+	case EControllerID::LFO0Frequency:		 return FString::Format(TEXT("({0}) LFO0Frequency"), { ControllerId });
+	case EControllerID::LFO1Frequency:		 return FString::Format(TEXT("({0}) LFO1Frequency"), { ControllerId });
+	case EControllerID::LFO0Depth:			 return FString::Format(TEXT("({0}) LFO0Depth"), { ControllerId });
+	case EControllerID::LFO1Depth:			 return FString::Format(TEXT("({0}) LFO1Depth"), { ControllerId });
+	case EControllerID::DelayTime:			 return FString::Format(TEXT("({0}) DelayTime"), { ControllerId });
+	case EControllerID::DelayDryGain:		 return FString::Format(TEXT("({0}) DelayDryGain"), { ControllerId });
+	case EControllerID::DelayWetGain:		 return FString::Format(TEXT("({0}) DelayWetGain"), { ControllerId });
+	case EControllerID::DelayFeedback:		 return FString::Format(TEXT("({0}) DelayFeedback"), { ControllerId });
+	case EControllerID::CoarsePitchBend:	 return FString::Format(TEXT("({0}) CoarsePitchBend"), { ControllerId });
+	case EControllerID::SampleStartTime:	 return FString::Format(TEXT("({0}) SampleStartTime"), { ControllerId });
+	case EControllerID::DataFine:			 return FString::Format(TEXT("({0}) DataFine"), { ControllerId });
+	case EControllerID::SubStreamVol1:		 return FString::Format(TEXT("({0}) SubStreamVol1"), { ControllerId });
+	case EControllerID::SubStreamVol2:		 return FString::Format(TEXT("({0}) SubStreamVol2"), { ControllerId });
+	case EControllerID::SubStreamVol3:		 return FString::Format(TEXT("({0}) SubStreamVol3"), { ControllerId });
+	case EControllerID::SubStreamVol4:		 return FString::Format(TEXT("({0}) SubStreamVol4"), { ControllerId });
+	case EControllerID::SubStreamVol5:		 return FString::Format(TEXT("({0}) SubStreamVol5"), { ControllerId });
+	case EControllerID::SubStreamVol6:		 return FString::Format(TEXT("({0}) SubStreamVol6"), { ControllerId });
+	case EControllerID::SubStreamVol7:		 return FString::Format(TEXT("({0}) SubStreamVol7"), { ControllerId });
+	case EControllerID::SubStreamVol8:		 return FString::Format(TEXT("({0}) SubStreamVol8"), { ControllerId });
+	case EControllerID::DelayEQEnabled:		 return FString::Format(TEXT("({0}) DelayEQEnabled"), { ControllerId });
+	case EControllerID::DelayEQType:		 return FString::Format(TEXT("({0}) DelayEQType"), { ControllerId });
+	case EControllerID::DelayEQFreq:		 return FString::Format(TEXT("({0}) DelayEQFreq"), { ControllerId });
+	case EControllerID::DelayEQQ:			 return FString::Format(TEXT("({0}) DelayEQQ"), { ControllerId });
+	case EControllerID::Hold:				 return FString::Format(TEXT("({0}) Hold"), { ControllerId });
+	case EControllerID::PortamentoSwitch:	 return FString::Format(TEXT("({0}) PortamentoSwitch"), { ControllerId });
+	case EControllerID::Sustenuto:			 return FString::Format(TEXT("({0}) Sustenuto"), { ControllerId });
+	case EControllerID::SoftPedal:			 return FString::Format(TEXT("({0}) SoftPedal"), { ControllerId });
+	case EControllerID::Legato:				 return FString::Format(TEXT("({0}) Legato"), { ControllerId });
+	case EControllerID::Hold2:				 return FString::Format(TEXT("({0}) Hold2"), { ControllerId });
+	case EControllerID::FilterQ:			 return FString::Format(TEXT("({0}) FilterQ"), { ControllerId });
+	case EControllerID::Release:			 return FString::Format(TEXT("({0}) Release"), { ControllerId });
+	case EControllerID::Attack:				 return FString::Format(TEXT("({0}) Attack"), { ControllerId });
+	case EControllerID::FilterFrequency:	 return FString::Format(TEXT("({0}) FilterFrequency"), { ControllerId });
+	case EControllerID::TimeStretchEnvelopeOrder:		 
+		return FString::Format(TEXT("({0}) TimeStretchEnvelopeOrder"), { ControllerId });
+	case EControllerID::DelayLFOBeatSync:	 return FString::Format(TEXT("({0}) DelayLFOBeatSync"), { ControllerId });
+	case EControllerID::DelayLFOEnabled:	 return FString::Format(TEXT("({0}) DelayLFOEnabled"), { ControllerId });
+	case EControllerID::DelayLFORate:		 return FString::Format(TEXT("({0}) DelayLFORate"), { ControllerId });
+	case EControllerID::DelayLFODepth:		 return FString::Format(TEXT("({0}) DelayLFODepth"), { ControllerId });
+	case EControllerID::DelayStereoType:	 return FString::Format(TEXT("({0}) DelayStereoType"), { ControllerId });
+	case EControllerID::DelayPanLeft:		 return FString::Format(TEXT("({0}) DelayPanLeft"), { ControllerId });
+	case EControllerID::DelayPanRight:		 return FString::Format(TEXT("({0}) DelayPanRight"), { ControllerId });
+	case EControllerID::RPNFine:			 return FString::Format(TEXT("({0}) RPNFine"), { ControllerId });
+	case EControllerID::RPNCourse:			 return FString::Format(TEXT("({0}) RPNCourse"), { ControllerId });
+	case EControllerID::AllSoundOff:		 return FString::Format(TEXT("({0}) AllSoundOff"), { ControllerId });
+	case EControllerID::Reset:				 return FString::Format(TEXT("({0}) Reset"), { ControllerId });
+	case EControllerID::AllNotesOff:		 return FString::Format(TEXT("({0}) AllNotesOff"), { ControllerId });
+	}
+	return FString::Format(TEXT("({0}) - <unassigned>"), { ControllerId });
+}
+
+void MidiConstants::GetControllerNames(TArray<FString>& Names)
+{
+	for (int32 i = 0; i < 128; ++i)
+	{
+		Names.Add(GetControllerName((EControllerID)i));
+	}
+}
+
+const int8 MidiConstants::GetNoteNumberFromNoteName(const char* InName)
+{
+	if (isdigit(*InName))
+		return (int8)atoi(InName);
+
+	char Name[6];
+	size_t Index = 0;
+	for (; Index < 5 && Index < strlen(InName); Index++)
+		Name[Index] = (char)toupper(InName[Index]);
+	Name[Index] = 0;
+	char* Walk = Name;
+
+	int32 Note = -1;
+	switch (*Walk)
+	{
+	case 'C': Note = 0;  break;
+	case 'D': Note = 2;  break;
+	case 'E': Note = 4;  break;
+	case 'F': Note = 5;  break;
+	case 'G': Note = 7;  break;
+	case 'A': Note = 9;  break;
+	case 'B': Note = 11; break;
+	}
+	if (Note == -1)
+		return 0;
+	Walk++;
+	if (*Walk == '#') 
+	{
+		Note++; 
+		Walk++; 
+	}
+	else if (*Walk == 'B')
+	{
+		Note--;
+		Walk++; 
+	}
+
+	int32 Octave = 0;
+	if (*Walk != '-')
+		Octave = atoi(Walk) + 1;
+
+	return int8(Note + (Octave * kNotesPerOctave));
+}
+
+const char* MidiConstants::GetNoteNameFromNoteNumber(uint8 MidiNoteNumber, ENoteNameEnharmonicStyle Style)
+{
+	typedef char const* CStr;
+	static CStr sSharpAndFlat[kNotesPerOctave];
+	static CStr sSharp[kNotesPerOctave];
+	static CStr sFlat[kNotesPerOctave];
+
+	sSharpAndFlat[0] = "C";
+	sSharpAndFlat[1] = "C#/Db";
+	sSharpAndFlat[2] = "D";
+	sSharpAndFlat[3] = "D#/Eb";
+	sSharpAndFlat[4] = "E";
+	sSharpAndFlat[5] = "F";
+	sSharpAndFlat[6] = "F#/Gb";
+	sSharpAndFlat[7] = "G";
+	sSharpAndFlat[8] = "G#/Ab";
+	sSharpAndFlat[9] = "A";
+	sSharpAndFlat[10] = "A#/Bb";
+	sSharpAndFlat[11] = "B";
+
+	sSharp[0] = "C";
+	sSharp[1] = "C#";
+	sSharp[2] = "D";
+	sSharp[3] = "D#";
+	sSharp[4] = "E";
+	sSharp[5] = "F";
+	sSharp[6] = "F#";
+	sSharp[7] = "G";
+	sSharp[8] = "G#";
+	sSharp[9] = "A";
+	sSharp[10] = "A#";
+	sSharp[11] = "B";
+
+	sFlat[0] = "C";
+	sFlat[1] = "Db";
+	sFlat[2] = "D";
+	sFlat[3] = "Eb";
+	sFlat[4] = "E";
+	sFlat[5] = "F";
+	sFlat[6] = "Gb";
+	sFlat[7] = "G";
+	sFlat[8] = "Ab";
+	sFlat[9] = "A";
+	sFlat[10] = "Bb";
+	sFlat[11] = "B";
+
+	CStr const* NameTable = nullptr;
+	switch (Style)
+	{
+	case ENoteNameEnharmonicStyle::Flat:
+		NameTable = sFlat;
+		break;
+
+	case ENoteNameEnharmonicStyle::Sharp:
+		NameTable = sSharp;
+		break;
+	default:
+	case ENoteNameEnharmonicStyle::SharpAndFlat:
+		NameTable = sSharpAndFlat;
+		break;
+	}
+
+	check(NameTable != nullptr);
+	uint8 Index = MidiNoteNumber % kNotesPerOctave;
+	return NameTable[Index];
+}
+
+int8 MidiConstants::GetNoteOctaveFromNoteNumber(uint8 MidiNoteNumber)
+{
+	return (int8)MidiNoteNumber / kNotesPerOctave - 1;
+}
