[Harmonix] Move Harmonix plugin into Engine/Plugins/Runtime to make it available to licensees in UE 5.4

diff --git a/Engine/Plugins/Runtime/Harmonix/Source/HarmonixMidi/Private/HarmonixMidi/MidiFile.cpp b/Engine/Plugins/Runtime/Harmonix/Source/HarmonixMidi/Private/HarmonixMidi/MidiFile.cpp
new file mode 100644
--- /dev/null
+++ b/Engine/Plugins/Runtime/Harmonix/Source/HarmonixMidi/Private/HarmonixMidi/MidiFile.cpp
@@ -0,0 +1,888 @@
+// Copyright Epic Games, Inc. All Rights Reserved.
+#include "HarmonixMidi/MidiFile.h"
+#include "HarmonixMidi/MidiTrack.h"
+#include "HarmonixMidi/MidiReader.h"
+#include "HarmonixMidi/MidiReceiver.h"
+#include "HarmonixMidi/MidiMsg.h"
+#include "HarmonixMidi/MidiEvent.h"
+#include "HarmonixMidi/MidiWriter.h"
+#include "HarmonixMidi/BarMap.h"
+#include "HarmonixMidi/SongMaps.h"
+#include "HarmonixMidi/SongMapReceiver.h"
+#include "Misc/Paths.h"
+#include "Algo/ForEach.h"
+#include "UObject/AssetRegistryTagsContext.h"
+#if WITH_EDITORONLY_DATA
+#include "EditorFramework/AssetImportData.h"
+#endif
+
+//////////////////////////////////////////////////////////////
+//
+// private helper class for reading a std midi file into a UMidiFile
+//
+
+class FMidiEventReceiver : public IMidiReceiver
+{
+public:
+	FMidiEventReceiver(UMidiFile& target);
+	int32 GetLastTick() const { return LastTick; }
+
+private:
+	virtual void Reset() override
+	{
+		LastTick = 0;
+		CurrentTrackHasName = false;
+	}
+	virtual void OnNewTrack(int32 NewTrackIndex) override;
+	virtual void OnEndOfTrack(int32 InLastTick) override;
+	virtual void OnAllTracksRead() override {}
+	virtual void OnMidiMessage(int32 Tick, uint8 Status, uint8 Data1, uint8 Data2) override;
+	virtual void OnTempo(int32 Tick, int32 Tempo) override;
+	virtual void OnText(int32 Tick, const FString& Str, uint8 Type) override;
+	virtual void OnTimeSignature(int32 Tick, int32 Numerator, int32 Denominator, bool FailOnError = true) override;
+
+	virtual void AddDummyConductorTrack() override;
+
+	UMidiFile& File;
+	int32  LastTick = 0;
+	bool   CurrentTrackHasName = false;
+};
+
+
+////////////////////////////////////////
+UMidiFile::UMidiFile()
+{
+	// add a single conductor track as the 0th track:
+	TheMidiData.Tracks.Emplace("Conductor");
+}
+
+void UMidiFile::GetAssetRegistryTags(FAssetRegistryTagsContext Context) const
+{
+	Super::GetAssetRegistryTags(Context);
+	
+	Context.AddTag(UObject::FAssetRegistryTag("NumTracks", FString::FromInt(GetNumTracks()), UObject::FAssetRegistryTag::TT_Numerical));
+	
+	FString Tempo = FString::Printf(TEXT("%.2f BPM"), TheMidiData.SongMaps.GetTempoAtTick(0));
+	Context.AddTag(UObject::FAssetRegistryTag("InitialTempo",  Tempo, UObject::FAssetRegistryTag::TT_Numerical));
+	
+	FString Length = FString::Printf(TEXT("%d Bars"), TheMidiData.SongMaps.GetSongLengthData().LengthBars);
+	Context.AddTag(UObject::FAssetRegistryTag("Length", Length, UObject::FAssetRegistryTag::TT_Numerical));
+
+#if WITH_EDITORONLY_DATA
+	// this seems useful
+	// this would allow us to inspect import data about this asset without fully loading the asset
+	// storing it as a JSON is something that is done in other places (Like in StaticMesh)
+	if (AssetImportData)
+	{
+		Context.AddTag(FAssetRegistryTag(
+			GET_MEMBER_NAME_CHECKED(UMidiFile, AssetImportData),
+			AssetImportData->GetSourceData().ToJson(),
+			UObject::FAssetRegistryTag::TT_Hidden)
+		);
+	}
+#endif
+
+}
+
+#if WITH_EDITORONLY_DATA
+FString UMidiFile::GetImportedSrcFilePath() const
+{
+	return AssetImportData->GetFirstFilename();
+}
+
+#endif
+
+TSharedPtr<Audio::IProxyData> UMidiFile::CreateProxyData(const Audio::FProxyDataInitParams& InitParams)
+{
+	if (!RenderableCopyOfMidiFileData)
+	{
+		RenderableCopyOfMidiFileData = MakeShared<FMidiFileData>(TheMidiData);
+	}
+	TSharedPtr<FMidiFileProxy> Proxy = MakeShared<FMidiFileProxy>(RenderableCopyOfMidiFileData);
+	return Proxy;
+}
+
+
+void UMidiFile::BeginDestroy()
+{
+	RenderableCopyOfMidiFileData = nullptr;
+	Super::BeginDestroy();
+}
+
+void UMidiFile::PostInitProperties()
+{
+	Super::PostInitProperties();
+
+#if WITH_EDITORONLY_DATA
+	if (!AssetImportData && !HasAnyFlags(RF_ClassDefaultObject))
+	{
+		AssetImportData = NewObject<UAssetImportData>(this, TEXT("AssetImportData"));
+	}
+#endif
+}
+
+void UMidiFile::LoadStdMidiFile(const FString& FilePath, int32 DesiredTicksPerQuarterNote, MidiConstants::EMidiTextEventEncoding InTextEncoding)
+{
+	FString Filename = FPaths::GetCleanFilename(FilePath);
+	IPlatformFile& PlatformFileApi = FPlatformFileManager::Get().GetPlatformFile();
+	IFileHandle* FileHandle = PlatformFileApi.OpenRead(*FilePath);
+	TSharedPtr<FArchiveFileReaderGeneric> Archive = MakeShared<FArchiveFileReaderGeneric>(FileHandle, *Filename, FileHandle->Size());
+	LoadStdMidiFile(Archive, Filename, DesiredTicksPerQuarterNote, InTextEncoding);
+
+#if WITH_EDITORONLY_DATA
+	if (!AssetImportData)
+	{
+		AssetImportData = NewObject<UAssetImportData>(this, TEXT("AssetImportData"));
+	}
+	AssetImportData->Update(FilePath);
+#endif
+}
+
+void UMidiFile::LoadStdMidiFile(void* Buffer, int32 BufferSize, const FString& Filename, int32 DesiredTicksPerQuarterNote, MidiConstants::EMidiTextEventEncoding InTextEncoding)
+{
+	TSharedPtr<FBufferReader> BufferArchive = MakeShared<FBufferReader>(Buffer, BufferSize, false);
+	LoadStdMidiFile(BufferArchive, Filename, DesiredTicksPerQuarterNote, InTextEncoding);
+}
+
+void UMidiFile::LoadStdMidiFile(TSharedPtr<FArchive> Archive, const FString& Filename, int32 DesiredTicksPerQuarterNote, MidiConstants::EMidiTextEventEncoding InTextEncoding)
+{
+	TheMidiData.TicksPerQuarterNote = DesiredTicksPerQuarterNote;
+	TheMidiData.SongMaps.Init(TheMidiData.TicksPerQuarterNote);
+	TheMidiData.Tracks.Empty();
+	// Create a receiver list...
+	FMidiReceiverList AllReceivers;
+	// Add the event receiver that will populate the midi tracks with midi events...
+	FMidiEventReceiver EventReceiver(*this);
+	AllReceivers.Add(&EventReceiver);
+	// Add the map receiver that will populate the song maps...
+	FSongMapReceiver MapReceiver(&TheMidiData.SongMaps);
+	AllReceivers.Add(&MapReceiver);
+
+	// Create the reader and read all the data...
+	FStdMidiFileReader reader(Archive, Filename, &AllReceivers, TheMidiData.TicksPerQuarterNote, InTextEncoding);
+	reader.ReadAllTracks();
+
+	TheMidiData.LastEventTick = EventReceiver.GetLastTick();
+	if (reader.GetFailed())
+	{
+		UE_LOG(LogMidi, Error, TEXT("MIDI import failed. Midi data is malformed."));
+	}
+
+	// we are now "dirty"... so make sure any new requests for renderable data 
+	// get a new copy...
+	RenderableCopyOfMidiFileData = nullptr;
+}
+
+void UMidiFile::SaveStdMidiFile(const FString& FilePath)
+{
+	FString Filename = FPaths::GetCleanFilename(FilePath);
+	IPlatformFile& PlatformFileApi = FPlatformFileManager::Get().GetPlatformFile();
+	IFileHandle* FileHandle = PlatformFileApi.OpenWrite(*FilePath);
+	TSharedPtr<FArchiveFileWriterGeneric> Archive = MakeShared<FArchiveFileWriterGeneric>(FileHandle, *Filename, 0);
+	SaveStdMidiFile(Archive, Filename);
+}
+
+void UMidiFile::SaveStdMidiFile(TSharedPtr<FArchive> Archive, const FString& Filename)
+{
+	FMidiWriter Writer(*Archive, TheMidiData.TicksPerQuarterNote);
+	Algo::ForEach(TheMidiData.Tracks, [&](FMidiTrack& track) { track.WriteStdMidi(Writer); });
+}
+
+FMidiTrack* UMidiFile::AddTrack(const FString& Name)
+{
+	TheMidiData.Tracks.Emplace(Name);
+	return &TheMidiData.Tracks.Last();
+}
+
+void UMidiFile::Empty()
+{
+	TheMidiData.Tracks.Empty();
+	TheMidiData.Tracks.Emplace("Conductor");
+}
+
+void UMidiFile::SetConductorTrack(const FTempoMap* TempoMap, const FBarMap* BarMap)
+{
+	check(!TheMidiData.Tracks.IsEmpty());
+	// use assignment operator to properly clear string table, events, and other data.
+	TheMidiData.Tracks[0] = FMidiTrack("Conductor");
+
+	int32 numTempoChanges = TempoMap->GetNumTempoChangePoints();
+
+	int32 i;
+	for (i = 0; i < numTempoChanges; ++i)
+	{
+		int32 tick = TempoMap->GetTempoChangePointTick(i);
+		float msPerQuarterNote = TempoMap->GetMsPerQuarterNoteAtTick(tick);
+		int32 usecPerQuarterNote = int32(msPerQuarterNote * 1000);
+		TheMidiData.Tracks[0].AddEvent(FMidiEvent(tick, FMidiMsg(usecPerQuarterNote)));
+	}
+
+	int32 numTimeSigChanges = BarMap->GetNumTimeSignaturePoints();
+
+	for (i = 0; i < numTimeSigChanges; ++i)
+	{
+		const FTimeSignaturePoint& sig = BarMap->GetTimeSignaturePoint(i);
+		TheMidiData.Tracks[0].AddEvent(FMidiEvent(sig.StartTick, FMidiMsg(sig.TimeSignature.Numerator, sig.TimeSignature.Denominator)));
+	}
+
+	TheMidiData.Tracks[0].Sort();
+}
+
+void UMidiFile::SortAllTracks()
+{
+	for (auto& Track : TheMidiData.Tracks)
+	{
+		Track.Sort();
+	}
+}
+
+const FMidiTrack* UMidiFile::FindTrackByName(const FString& Name) const
+{
+	int32 trackIndex = FindTrackIndexByName(Name);
+	if (trackIndex == INDEX_NONE)
+	{
+		return nullptr;
+	}
+	return &TheMidiData.Tracks[trackIndex];
+}
+
+int32 UMidiFile::FindTrackIndexByName(const FString& Name) const
+{
+	for (int32 i = 0; i < TheMidiData.Tracks.Num(); ++i)
+	{
+		const FMidiTrack& Track = TheMidiData.Tracks[i];
+		if (Track.GetName() && *Track.GetName() == Name)
+		{
+			return i;
+		}
+	}
+	return INDEX_NONE;
+}
+
+int32 UMidiFile::FindTextEvent(const FString& eventText, const FString* TrackName)
+{
+	if (TrackName)
+	{
+		int32 TrackIndex = FindTrackIndexByName(*TrackName);
+		if (TrackIndex == INDEX_NONE)
+		{
+			return INDEX_NONE;
+		}
+		return FindTextEvent(eventText, TrackIndex);
+	}
+	else
+	{
+		for (int32 i = 0; i < TheMidiData.Tracks.Num(); ++i)
+		{
+			int32 FoundIndex = FindTextEvent(eventText, i);
+			if (FoundIndex >= 0)
+			{
+				return FoundIndex;
+			}
+		}
+	}
+	return INDEX_NONE;
+}
+
+int32 UMidiFile::FindTextEvent(const FString& EventText, int32 TrackIndex)
+{
+	if (TrackIndex >= TheMidiData.Tracks.Num())
+	{
+		return INDEX_NONE;
+	}
+	if (!TheMidiData.Tracks[TrackIndex].GetTextRepository())
+	{
+		return INDEX_NONE;
+	}
+
+	// first lets see if the request text is even in the track's repository...
+	int32 StringIndex = 0;
+	FMidiTextRepository* TextRepo = TheMidiData.Tracks[TrackIndex].GetTextRepository();
+	for (; StringIndex < TextRepo->Num(); StringIndex++)
+	{
+		if ((*TextRepo)[StringIndex] == EventText)
+		{
+			break;
+		}
+	}
+	if (StringIndex == TextRepo->Num())
+	{
+		return INDEX_NONE;
+	}
+
+	// it's in the list, so now let's find the event...
+	for (const FMidiEvent& m : TheMidiData.Tracks[TrackIndex].GetEvents())
+	{
+		if (m.GetMsg().MsgType() == FMidiMsg::EType::Text && m.GetMsg().GetTextIndex() == StringIndex)
+		{
+			return m.GetTick();
+		}
+	}
+	// not found? How can this be if it was in the textRepo?
+	return INDEX_NONE;
+}
+
+TArray<int32> UMidiFile::FindAllTextEvents(const FString& EventText, const FString* TrackName)
+{
+	TArray<int32> ticks;
+
+	if (TrackName)
+	{
+		int32 TrackIndex = FindTrackIndexByName(*TrackName);
+		if (TrackIndex != INDEX_NONE)
+		{
+			FindAllTextEvents(EventText, TrackIndex, ticks);
+		}
+	}
+	else
+	{
+		for (int32 i = 0; i < TheMidiData.Tracks.Num(); ++i)
+		{
+			FindAllTextEvents(EventText, i, ticks);
+		}
+	}
+	return ticks;
+}
+
+void UMidiFile::FindAllTextEvents(const FString& EventText, int32 TrackIndex, TArray<int32> OutIndexes)
+{
+	// first lets see if the request text is even in the track's repository...
+	int32 StringIndex = 0;
+	FMidiTextRepository* TextRepo = TheMidiData.Tracks[TrackIndex].GetTextRepository();
+	for (; StringIndex < TextRepo->Num(); StringIndex++)
+	{
+		if ((*TextRepo)[StringIndex] == EventText)
+		{
+			break;
+		}
+	}
+
+	// if we found it, then look through the events...
+	if (StringIndex != TextRepo->Num())
+	{
+		for (const FMidiEvent& m : TheMidiData.Tracks[TrackIndex].GetEvents())
+		{
+			if (m.GetMsg().MsgType() == FMidiMsg::EType::Text && m.GetMsg().GetTextIndex() == StringIndex)
+			{
+				OutIndexes.Add(m.GetTick());
+			}
+		}
+	}
+}
+
+void UMidiFile::TracksChanged()
+{
+	int32 NewLastEventTick = 0;
+	bool bIsFileEmpty = true;
+	for (auto& Track : TheMidiData.Tracks)
+	{
+		FMidiEventList& Events = Track.GetRawEvents();
+		//check if the file is empty
+		if (bIsFileEmpty && !Events.IsEmpty())
+		{
+			bIsFileEmpty = false;
+		}
+		int32 LastTickOnTrack = Events.IsEmpty() ? 0 : Events.Last().GetTick();
+		if (NewLastEventTick < LastTickOnTrack)
+		{
+			NewLastEventTick = LastTickOnTrack;
+		}
+	}
+
+	if (NewLastEventTick == 0 && !bIsFileEmpty)
+	{
+		TheMidiData.LastEventTick = NewLastEventTick;
+		FSongLengthData& LengthData = TheMidiData.SongMaps.GetSongLengthData();
+		//if the last event tick in a file is 0 and the file is not empty, we know that the length in tick is 1
+		LengthData.LengthTicks = 1;
+		LengthData.LastTick = NewLastEventTick;
+		LengthData.LengthBars = TheMidiData.SongMaps.GetBarMap().TickToBarIncludingCountIn(LengthData.LengthTicks);
+		return;
+	}
+
+	//update length data in song length data if needed 
+	if (NewLastEventTick != TheMidiData.LastEventTick)
+	{
+		TheMidiData.LastEventTick = NewLastEventTick;
+		FMusicTimestamp Timestamp = TheMidiData.SongMaps.GetBarMap().TickToMusicTimestamp(NewLastEventTick + 1);
+		FSongLengthData& LengthData = TheMidiData.SongMaps.GetSongLengthData();
+		LengthData.LengthTicks = TheMidiData.SongMaps.GetBarMap().MusicTimestampToTick(Timestamp);
+		LengthData.LastTick = LengthData.LengthTicks - 1;
+		LengthData.LengthBars = TheMidiData.SongMaps.GetBarMap().TickToBarIncludingCountIn(LengthData.LengthTicks);
+	}
+}
+
+bool UMidiFile::ShouldConformMidiFileLength(EMidiFileLengthConformOption Option)
+{
+	//if file length is already rounded down, there's nothing more can be done
+	if (bLengthRoundedDown) return false;
+
+	//if a file's length is rounded up, it can be rounded down once more, but nothing else 
+	if (bLengthRoundedUp)
+	{
+		if (Option == EMidiFileLengthConformOption::RoundDown)
+		{
+			//if the file's length is less than 1, cannot round down
+			return TheMidiData.SongMaps.GetSongLengthData().LengthBars > 1 ? true : false;
+		}
+		else
+		{
+			return false;
+		}
+	}
+
+	if (bLengthRoundedToNearest && Option == EMidiFileLengthConformOption::Nearest) return false;
+
+	//if bar length is 1 bar (or less, it will round up to 1 bar if the file is imported, 0 if the file is manually created), 
+	//it CANNOT be rounded down to 0 bar, can only be rounded up
+	if (TheMidiData.SongMaps.GetSongLengthData().LengthBars <= 1 && Option != EMidiFileLengthConformOption::RoundUp)
+	{
+		return false;
+	}
+
+	//Special case: a manually created Midi file only has events on tick 0
+	//This should still be able to be conformed by rounding up ONLY
+	bool bIsFileEmpty = true;
+	for (auto& Track : TheMidiData.Tracks)
+	{
+		FMidiEventList& Events = Track.GetRawEvents();
+		//check if the file is empty
+		if (bIsFileEmpty && !Events.IsEmpty())
+		{
+			bIsFileEmpty = false;
+			break;
+		}
+	}
+
+	//get the last event tick in file
+	int32 LastEventTick = TheMidiData.LastEventTick;
+
+	if (LastEventTick == 0 && !bIsFileEmpty && Option == EMidiFileLengthConformOption::RoundUp) return true;
+
+	//if the last event tick is at the very end of an integer bar, the next tick should be the at the next bar, beat 1.0 
+	FMusicTimestamp Timestamp = TheMidiData.SongMaps.GetBarMap().TickToMusicTimestamp(LastEventTick + 1);
+	//for manually created Midi files, we can specify that the last event tick is at the very end of an integer bar
+	//however,some DAWs places the very last event of an integer bar at the very beginning of the next bar, causing a 1-tick difference
+	//the current solution is to tolerate this 1-tick error when checking the last event tick
+	double OneTickTimeStampBeat = TheMidiData.SongMaps.GetBarMap().TickToMusicTimestamp(1).Beat;
+	double OneTickError = OneTickTimeStampBeat - 1.0;
+	if (!FMath::IsNearlyEqual(Timestamp.Beat, 1.0f, OneTickError))
+	{
+		return true;
+	}
+	return false;
+}
+
+void UMidiFile::ConformMidiFileLength(EMidiFileLengthConformOption Option)
+{
+	//check whether midi file length needs to be conformed
+	if (!ShouldConformMidiFileLength(Option))
+	{
+		UE_LOG(LogMidi, Log, TEXT("Midi file does not need to be conformed."));
+		return;
+	}
+	
+	int32 LastTickInSongLengthData = TheMidiData.SongMaps.GetSongLengthData().LastTick;
+	//Currently, file length is automatically rounded up to the nearest integer bar for imported midi files, if the conform option is Round Up, early out
+	//(NOTE: Round Up doesn't add/change/remove anything to the midi file or events)
+	if (Option == EMidiFileLengthConformOption::RoundUp)
+	{
+		//for manually created Midi files, update the song length data to keep up with 
+		//results of imported Midi files WITHOUT adding/changing/removing events
+		FMusicTimestamp Timestamp = TheMidiData.SongMaps.GetBarMap().TickToMusicTimestamp(TheMidiData.LastEventTick + 1);
+		if (!FMath::IsNearlyEqual(Timestamp.Beat, 1.0f))
+		{
+			Timestamp.Bar++;
+			Timestamp.Beat = 1.0f;
+			FSongLengthData& LengthData = TheMidiData.SongMaps.GetSongLengthData();
+			LengthData.LengthTicks = TheMidiData.SongMaps.GetBarMap().MusicTimestampToTick(Timestamp);
+			LengthData.LastTick = LengthData.LengthTicks - 1;
+			LengthData.LengthBars = TheMidiData.SongMaps.GetBarMap().TickToBarIncludingCountIn(LengthData.LengthTicks);
+		}
+		
+		bLengthRoundedUp = true;
+		return;
+	}
+
+	//Get file's fractional bar length (not rounded up)
+	float OriginalLengthFractionalBar = TheMidiData.SongMaps.GetBarIncludingCountInAtTick(TheMidiData.LastEventTick);
+
+	//Round Down to the nearest integer bar
+	int32 ConformedLastBar = (int32)OriginalLengthFractionalBar;
+
+	//rounding to nearest is still either round down or round up
+	//depending on the result from FMath::RoundToInt()
+	if (Option == EMidiFileLengthConformOption::Nearest)
+	{
+		ConformedLastBar = FMath::RoundToInt(OriginalLengthFractionalBar);
+		EMidiFileLengthConformOption NearestRoundingOption = ConformedLastBar > (int32)OriginalLengthFractionalBar ? EMidiFileLengthConformOption::RoundUp : EMidiFileLengthConformOption::RoundDown;
+		
+		//if round to nearest ends up being rounding up, early out
+		if (NearestRoundingOption == EMidiFileLengthConformOption::RoundUp)
+		{
+			//round the length up by updating the song length data
+			FMusicTimestamp Timestamp = TheMidiData.SongMaps.GetBarMap().TickToMusicTimestamp(TheMidiData.LastEventTick + 1);
+			if (!FMath::IsNearlyEqual(Timestamp.Beat, 1.0f))
+			{
+				Timestamp.Bar++;
+				Timestamp.Beat = 1.0f;
+				FSongLengthData& LengthData = TheMidiData.SongMaps.GetSongLengthData();
+				LengthData.LengthTicks = TheMidiData.SongMaps.GetBarMap().MusicTimestampToTick(Timestamp);
+				LengthData.LastTick = LengthData.LengthTicks - 1;
+				LengthData.LengthBars = TheMidiData.SongMaps.GetBarMap().TickToBarIncludingCountIn(LengthData.LengthTicks);
+			}
+			bLengthRoundedUp = true;
+			bLengthRoundedToNearest = true;
+			return;
+		}
+	}
+
+	//the actual tick to conform to by rounding down
+	int32 ConformedLastEventTick = TheMidiData.SongMaps.GetBarMap().BarIncludingCountInToTick(ConformedLastBar) - 1;
+
+	//Retrieve midi file name for logging after moving/removing midi events
+	FString MidiFileName = this->GetName();
+
+	//keep track of the events that are moved to the last tick and excessive events removed after moving
+	int32 NumEventsMovedToLastTick = 0;
+	int32 TotalNumEventsRemoved = 0;
+	int32 NumNoteOnNoteOffPairsRemoved = 0;
+	int32 NumAftertouchEventsRemoved = 0;
+	int32 NumPitchEventsRemoved = 0;
+	int32 NumControlEventsRemoved = 0;
+
+	//go through midi events in each midi track 
+	for (int32 TrackIndex = 0; TrackIndex < TheMidiData.Tracks.Num(); ++TrackIndex)
+	{
+		//Get the raw midi events and the last tick of the last event
+		FMidiTrack& Track = TheMidiData.Tracks[TrackIndex];
+		FMidiEventList& Events = Track.GetRawEvents();
+		int32 LastEventTick = Events.Last().GetTick();
+	
+		//if the last tick of this event is less than the tick of the conformed last tick,
+		//this midi track does not need to be conformed, go to the next midi track
+		if (LastEventTick <= ConformedLastEventTick) continue;
+
+		//move all events with tick > ConfromedLastEventTick to ConformedLastTick
+		for (int32 EventIndex = Events.Num() - 1; EventIndex >= 0 && Events[EventIndex].GetTick() > ConformedLastEventTick; --EventIndex)
+		{
+			// NOTE: It is safe for us to do this BECAUSE...
+			// We do not have to re-sort the midi data on the tracks because while we may have slid events earlier,
+			// their relative position to each other MUST be unchanged!
+			Events[EventIndex].Tick = ConformedLastEventTick;
+			NumEventsMovedToLastTick++;
+		}
+
+		//filter the midi events (at the new conformed last tick) after moving them, remove note on/note off pairs on the last tick, and other excessive events
+		int32 NumItemsRemoved = 0;
+		int32 CurrentEventIndex = Events.Num() - 1;
+
+		while (CurrentEventIndex >= 0 && Events[CurrentEventIndex].GetTick() == ConformedLastEventTick)
+		{
+			FMidiEvent& Event = Events[CurrentEventIndex];
+			FMidiMsg& Msg = Event.GetMsg();
+
+			//filter Note On/Note off pairs on the last tick
+			if (Msg.IsNoteOff())
+			{
+				int32 NoteOnEventIndex = -1;
+				for (int32 i = CurrentEventIndex - 1; i > 0 && Events[i].GetTick() == ConformedLastEventTick; --i)
+				{
+					//check equality of note on/note off events' midi note number (data1) 
+					if (Events[i].GetMsg().IsNoteOn() && Events[i].GetMsg().GetStdData1() == Msg.GetStdData1())
+					{
+						NoteOnEventIndex = i;
+						break;
+					}
+				}
+				//remove note on/note off pairs on the conformed last tick
+				if (NoteOnEventIndex != -1) 
+				{
+					Events.RemoveAt(CurrentEventIndex);
+					NumItemsRemoved++;
+					Events.RemoveAt(NoteOnEventIndex);
+					NumItemsRemoved++;
+					NumNoteOnNoteOffPairsRemoved++;
+					TotalNumEventsRemoved += 2;
+				}
+			}
+			uint8 MsgStatus = Msg.Status;
+
+			//filter chan press/pitch bend/program change events
+			if (MsgStatus == MidiConstants::kChanPres|| MsgStatus == MidiConstants::kPitch || MsgStatus == MidiConstants::kProgram)
+			{
+				int32 OtherEventIndexToRemove = -1;
+				//check if there exist multiple events with same status on the last tick
+				for (int32 i = CurrentEventIndex - 1; i >= 0 && Events[i].GetTick() == ConformedLastEventTick; --i)
+				{
+					if (Events[i].GetMsg().Status == MsgStatus) {
+						OtherEventIndexToRemove = CurrentEventIndex;
+						break;
+					}
+				}
+				//remove excessive events on the conformed last tick
+				if (OtherEventIndexToRemove != -1)
+				{
+					Events.RemoveAt(OtherEventIndexToRemove);
+					NumItemsRemoved++;
+					if (MsgStatus == MidiConstants::kChanPres || MsgStatus == MidiConstants::kPolyPres)
+					{
+						NumAftertouchEventsRemoved++;
+					}
+					else {
+						NumPitchEventsRemoved++;
+					}
+					TotalNumEventsRemoved++;
+				}
+
+			}
+
+			//filter Control Change events and Poly Pres events
+			if (MsgStatus == MidiConstants::kControl || MsgStatus == MidiConstants::kPolyPres)
+			{
+				uint8 CurrentEventControllerId = Msg.Data1;
+				int32 ControlEventIndexToRemove = -1; 
+				for (int32 i = CurrentEventIndex - 1; i >= 0 && Events[i].GetTick() == ConformedLastEventTick; --i)
+				{
+					//check control change events for identical controller ID (data1) on the same tick and remove the later one
+					//check poly pres events for the same note number (data1) on the same tick and remove the later one
+					if (Events[i].GetMsg().Data1 == CurrentEventControllerId) {
+						ControlEventIndexToRemove = CurrentEventIndex;
+						break;
+					}
+				}
+				//remove excessive control changes or poly pressure on conformed last tick
+				if (ControlEventIndexToRemove != -1)
+				{
+					Events.RemoveAt(ControlEventIndexToRemove);
+					NumItemsRemoved++;
+					NumControlEventsRemoved++;
+					TotalNumEventsRemoved++;
+				}
+			}
+
+			//update indices after removing excessive events
+			if (NumItemsRemoved == 0)
+			{
+				CurrentEventIndex -= 1;
+			}
+			else {
+				CurrentEventIndex -= NumItemsRemoved;
+				NumItemsRemoved = 0;
+			}
+		}
+	}
+
+	TracksChanged();
+
+	//update the conformed length in SongLengthData if conformed to an integer bar, and marks the file is being conformed with the current conform option(s)
+	if (Option == EMidiFileLengthConformOption::RoundDown || Option == EMidiFileLengthConformOption::Nearest)
+	{
+		FSongLengthData& LengthData = TheMidiData.SongMaps.GetSongLengthData();
+		LengthData.LengthTicks = ConformedLastEventTick + 1;
+		LengthData.LastTick = ConformedLastEventTick;
+		LengthData.LengthBars = TheMidiData.SongMaps.GetBarMap().TickToBarIncludingCountIn(LengthData.LengthTicks);
+		bLengthRoundedDown = true;
+
+		//if we get here by rounding up first (or round up by rounding to the nearest integer) and then round down
+		//this means the file is ultimately rounded down, NOT rounded up
+		if (bLengthRoundedUp) bLengthRoundedUp = false;
+
+		bLengthRoundedToNearest = Option == EMidiFileLengthConformOption::Nearest ? true : false;
+		
+		//Log conform summary 
+		FString MidiFileLengthConformSummary = FString::Printf(
+			TEXT("Midi file '%s.mid' length conformance summary: original length: %f bars, rounded down to new length: %d bar(s), moved %d Midi event(s) to tick %d"),
+			*MidiFileName,
+			OriginalLengthFractionalBar,
+			ConformedLastBar,
+			NumEventsMovedToLastTick,
+			ConformedLastEventTick
+		);
+		if (TotalNumEventsRemoved > 0)
+		{
+			MidiFileLengthConformSummary.Append(FString::Printf(TEXT("; removed %d Midi event(s):"),TotalNumEventsRemoved));
+			if (NumNoteOnNoteOffPairsRemoved > 0) MidiFileLengthConformSummary.Append(FString::Printf(TEXT(" %d Note On/Note off event pair(s),"), NumNoteOnNoteOffPairsRemoved));
+			if (NumAftertouchEventsRemoved > 0) MidiFileLengthConformSummary.Append(FString::Printf(TEXT(" %d Aftertouch event(s),"), NumAftertouchEventsRemoved));
+			if (NumPitchEventsRemoved > 0) MidiFileLengthConformSummary.Append(FString::Printf(TEXT(" %d Pitch Bend event(s),"), NumPitchEventsRemoved));
+			if (NumControlEventsRemoved > 0) MidiFileLengthConformSummary.Append(FString::Printf(TEXT(" %d Control Change event(s),"), NumControlEventsRemoved));
+			MidiFileLengthConformSummary.Append(FString::Printf(TEXT(" at tick %d "), ConformedLastEventTick));
+		}
+		UE_LOG(LogMidi, Log, TEXT("%s"),*MidiFileLengthConformSummary);
+	}
+
+	//This prevents the current conform option affecting any midi file that is currently being obtained from RenderableCopyOfMidiFileData
+	//e.g. if file length is being conformed while the un-conformed file is occupied during playback, changes to the midi file won't be reflected until the 
+	//current playback is stopped and starts playing again (we only create a new RenderableCopyOfMidiFileData if there isn't one existed already)
+	RenderableCopyOfMidiFileData = nullptr;
+}
+
+#if WITH_EDITORONLY_DATA
+void UMidiFile::PostEditChangeProperty(FPropertyChangedEvent& PropertyChangedEvent)
+{	
+	Super::PostEditChangeProperty(PropertyChangedEvent);
+
+	FName StartBarPropertyName = GET_MEMBER_NAME_CHECKED(UMidiFile, StartBar);
+	
+	if (PropertyChangedEvent.Property->GetName() == StartBarPropertyName)
+	{
+		TheMidiData.SongMaps.GetBarMap().SetStartBar(StartBar);
+	}
+}
+#endif
+
+FMidiEventReceiver::FMidiEventReceiver(UMidiFile& file)
+	: File(file)
+{
+}
+
+void FMidiEventReceiver::AddDummyConductorTrack()
+{
+	// Add a dummy conductor track for midi type 0 files
+	File.GetTracks().Emplace("Conductor");
+}
+
+void FMidiEventReceiver::OnNewTrack(int32 newTrackIndex)
+{
+	File.GetTracks().Emplace();
+	CurrentTrackHasName = false;
+}
+
+void FMidiEventReceiver::OnEndOfTrack(int InLastTick)
+{
+	if (InLastTick > LastTick)
+	{
+		LastTick = InLastTick;
+	}
+	if (!CurrentTrackHasName)
+	{
+		check(!File.GetTracks().IsEmpty());
+
+		FMidiTrack& Track = File.GetTracks().Last();
+		uint16 StringIndex = Track.AddText("track_1");
+
+		Track.AddEvent(FMidiEvent(0, FMidiMsg::CreateText(StringIndex, MidiConstants::kMeta_TrackName)));
+		Track.Sort();
+		CurrentTrackHasName = true;
+	}
+}
+
+void FMidiEventReceiver::OnMidiMessage(int32 Tick, uint8 Status, uint8 Data1, uint8 Data2)
+{
+	check(!File.GetTracks().IsEmpty());
+
+	File.GetTracks().Last().AddEvent(FMidiEvent(Tick, FMidiMsg(Status, Data1, Data2)));
+}
+
+void FMidiEventReceiver::OnTempo(int32 Tick, int32 Tempo)
+{
+	check(!File.GetTracks().IsEmpty());
+
+	File.GetTracks().Last().AddEvent(FMidiEvent(Tick, FMidiMsg(Tempo)));
+}
+
+void FMidiEventReceiver::OnText(int32 Tick, const FString& Str, uint8 Type)
+{
+	check(!File.GetTracks().IsEmpty());
+
+	FMidiTrack& Track = File.GetTracks().Last();
+	uint16 StringIndex = Track.AddText(Str);
+
+	if (Type == MidiConstants::kMeta_TrackName)
+	{
+		// Track name should always appear at tick 0!
+		if (Tick != 0 && !CurrentTrackHasName)
+		{
+			UE_LOG(LogMidi, Warning, TEXT("MIDI track name (\"%s\") event found at tick %d but none found at tick 0! The first track name event on a midi track should be on tick 0. Shifting it to tick 0."),
+				*Str, Tick);
+			Tick = 0;
+		}
+		else if (CurrentTrackHasName && Str != *Track.GetName())
+		{
+			UE_LOG(LogMidi, Warning, TEXT("2nd Track Name event found on MIDI track %s --> %s at tick %d"),
+				**Track.GetName(), *Str , Tick);
+		}
+		CurrentTrackHasName = true;
+	}
+
+	Track.AddEvent(FMidiEvent(Tick, FMidiMsg::CreateText(StringIndex, Type)));
+}
+
+void FMidiEventReceiver::OnTimeSignature(int32 Tick, int32 Numerator, int32 Denominator, bool FailOnError)
+{
+	check(!File.GetTracks().IsEmpty());
+
+	File.GetTracks().Last().AddEvent(FMidiEvent(Tick, FMidiMsg(Numerator, Denominator)));
+}
+
+int32 FMidiFileData::FindTrackIndexByName(const FString& TrackName)
+{
+	for (int32 i = 0; i < Tracks.Num(); ++i)
+	{
+		const FMidiTrack& Track = Tracks[i];
+		if (Track.GetName() && *Track.GetName() == TrackName)
+		{
+			return i;
+		}
+	}
+	return INDEX_NONE;
+}
+
+void FMidiFileData::PostSerialize(const FArchive& Ar)
+{
+	if (Ar.IsLoading() && SongMaps.BarMapIsEmpty() && Tracks.Num() > 0)
+	{
+		UE_LOG(LogMidi, Warning, TEXT("Empty bar map. Rebuilding."));
+		const FMidiEventList& Events = Tracks[0].GetEvents();
+		FBarMap& BarMap = SongMaps.GetBarMap();
+		BarMap.SetStartBar(1);
+		for (auto& Event : Events)
+		{
+			if (Event.GetMsg().MsgType() == FMidiMsg::EType::TimeSig)
+			{
+				int32 Tick = Event.GetTick();
+				int32 BarIndex;
+				check(Tick == 0 || BarMap.GetNumTimeSignaturePoints() > 0);
+				if (Tick == 0)
+				{
+					BarIndex = 0;
+				}
+				else
+				{
+					BarIndex = BarMap.TickToBarIncludingCountIn(Tick);
+				}
+				BarMap.AddTimeSignatureAtBarIncludingCountIn(BarIndex, Event.GetMsg().GetTimeSigNumerator(), Event.GetMsg().GetTimeSigDenominator());
+			}
+		}
+		BarMap.Finalize(LastEventTick);
+	}
+}
+
+void FMidiFileData::AddTempoChange(int32 TrackIdx, int32 Tick, float TempoBPM)
+{
+	check(Tracks.IsValidIndex(TrackIdx));
+
+	FTempoMap& TempoMap = SongMaps.GetTempoMap();
+	int32 MidiTempo = MidiConstants::BPMToMidiTempo(TempoBPM);
+	Tracks[TrackIdx].AddEvent(FMidiEvent(Tick, FMidiMsg(MidiTempo)));
+	Tracks[TrackIdx].Sort();
+	TempoMap.AddTempoInfoPoint(MidiTempo, Tick);
+}
+
+void FMidiFileData::AddTimeSigChange(int32 TrackIdx, int32 Tick, int32 InTimeSigNum, int32 InTimeSigDenom)
+{
+	check(Tracks.IsValidIndex(TrackIdx));
+
+	FBarMap& BarMap = SongMaps.GetBarMap();
+
+	// Time signature changes can only happen at the beginning of a bar,
+	// so round up to the next bar boundary...
+	int32 AbsoluteBar = FMath::CeilToInt32(SongMaps.GetBarIncludingCountInAtTick(Tick));
+	Tick = BarMap.BarBeatTickIncludingCountInToTick(AbsoluteBar, 1, 0);
+	int32 TimeSigNum = FMath::Clamp(InTimeSigNum, 1, 64);
+	int32 TimeSigDenom = FMath::Clamp(InTimeSigDenom, 1, 64);
+	Tracks[TrackIdx].AddEvent(FMidiEvent(Tick, FMidiMsg((uint8)TimeSigNum, (uint8)TimeSigDenom)));
+	Tracks[TrackIdx].Sort();
+	BarMap.AddTimeSignatureAtBarIncludingCountIn(AbsoluteBar, TimeSigNum, TimeSigDenom);
+}
+
diff --git a/Engine/Plugins/Runtime/Harmonix/Source/HarmonixMidi/Private/HarmonixMidi/MidiPlayCursor.cpp b/Engine/Plugins/Runtime/Harmonix/Source/HarmonixMidi/Private/HarmonixMidi/MidiPlayCursor.cpp
new file mode 100644
--- /dev/null
+++ b/Engine/Plugins/Runtime/Harmonix/Source/HarmonixMidi/Private/HarmonixMidi/MidiPlayCursor.cpp
@@ -0,0 +1,1247 @@
+// Copyright Epic Games, Inc. All Rights Reserved.
+#include "HarmonixMidi/MidiPlayCursor.h"
+#include "HarmonixMidi/MidiPlayCursorMgr.h"
+#include <functional>
+
+const float FMidiPlayCursor::kSmallMs = 0.1f;
+
+FMidiPlayCursor::FMidiPlayCursor()
+	: CurrentMs(0.0f)
+	, CurrentTick(0)
+	, LoopCount(0)
+	, UnregisterASAP(false)
+	, Owner(nullptr)
+	, Tracker(nullptr)
+	, LookaheadType(ELookaheadType::Ticks)
+	, SyncOpts(ESyncOptions::NoBroadcastNoPreRoll)
+	, LookaheadTicks(0)
+	, LookaheadMs(0.0f)
+	, WatchTrack(-1)
+	, FilterPassFlags(EFilterPassFlags::All)
+{
+
+}
+
+FMidiPlayCursor::~FMidiPlayCursor()
+{
+	if (Owner)
+	{
+		Owner->UnregisterPlayCursor(this);
+		Owner = nullptr;
+		Tracker = nullptr;
+	}
+}
+
+void FMidiPlayCursor::SetOwner(FMidiPlayCursorMgr* NewOwner, FMidiPlayCursorTracker* NewTracker, float PreRollMs)
+{
+	check(Owner == nullptr || NewOwner == nullptr);
+
+	Owner = NewOwner;
+	Tracker = NewTracker;
+
+	if (Owner)
+	{
+		TrackNextEventIndexs.SetNumUninitialized(Owner->Tracks().Num());
+		Reset(PreRollMs > 0.0f);
+
+		if (Tracker && PreRollMs > 0.0f)
+		{
+			if (PreRollMs > Tracker->CurrentMs)
+			{
+				PreRollMs = Tracker->CurrentMs;
+			}
+			int32 Tick = (int32)Owner->GetTempoMap().MsToTick(Tracker->CurrentMs - PreRollMs);
+			float Ms = Owner->GetTempoMap().TickToMs(Tick);
+			Tracker->ResetNewCursor(this, Tick, Ms, false);
+		}
+	}
+}
+
+void FMidiPlayCursor::Reset(bool ForceNoBroadcast)
+{
+	CurrentTick = -1;
+	LoopCount = 0;
+	if (Owner)
+	{
+		CurrentMs = Owner->GetTempoMap().TickToMs(CurrentTick);
+		TrackNextEventIndexs.SetNumUninitialized(Owner->Tracks().Num());
+
+		for (int32 i = 0; i < TrackNextEventIndexs.Num(); i++)
+		{
+			TrackNextEventIndexs[i] = 0;
+		}
+
+		if (LookaheadType == ELookaheadType::Ticks)
+		{
+			PrepareLookAheadTicks(ForceNoBroadcast);
+		}
+		else
+		{
+			PrepareLookAheadMs(ForceNoBroadcast);
+		}
+	}
+	else
+	{
+		TrackNextEventIndexs.Empty();
+		CurrentMs = 0.0f;
+	}
+}
+
+void FMidiPlayCursor::SetupTickLookahead(int32 Ticks, ESyncOptions Opts)
+{
+	LookaheadType = ELookaheadType::Ticks;
+	LookaheadTicks = Ticks;
+	SyncOpts = Opts;
+
+	// may be called before or after we are attached to a play cursor manager!
+	if (Owner)
+	{
+		PrepareLookAheadTicks();
+	}
+}
+
+void FMidiPlayCursor::RecalcNextEventsDueToMidiChanges(FMidiPlayCursorMgr::EMidiChangePositionCorrectMode PositionMode)
+{
+	// MUST ONLY BE CALLED BY THE PLAYCURSORMGR WITH CURSORS LOCKED!
+	check(Owner);
+
+	if (PositionMode == FMidiPlayCursorMgr::EMidiChangePositionCorrectMode::MaintainTick)
+	{
+		CurrentMs = Owner->GetTempoMap().TickToMs(CurrentTick);
+	}
+	else // if (positionMode == MidiChangePositionCorrectMode::MaintainTime)
+	{
+		CurrentTick = Owner->GetTempoMap().MsToTick(CurrentMs);
+	}
+
+	// changes to tracks might actually be more complicated...
+	const UMidiFile::FMidiTrackList& Tracks = Owner->Tracks();
+	if (Tracks.Num() != TrackNextEventIndexs.Num())
+	{
+		int32 OldNum = TrackNextEventIndexs.Num();
+		TrackNextEventIndexs.SetNumUninitialized(Tracks.Num());
+	}
+
+	for (int32 i = 0; i < Tracks.Num(); ++i)
+	{
+		const FMidiTrack& Track = Tracks[i];
+		const FMidiEventList& Events = Track.GetEvents();
+		if (Events.Num() == 0)
+		{
+			TrackNextEventIndexs[i] = -1;
+			continue;
+		}
+		else
+		{
+			// search from the beginning...
+			TrackNextEventIndexs[i] = 0;
+			while (TrackNextEventIndexs[i] != -1 && CurrentTick >= Events[TrackNextEventIndexs[i]].GetTick())
+			{
+				TrackNextEventIndexs[i]++;
+				if (TrackNextEventIndexs[i] == Events.Num())
+				{
+					TrackNextEventIndexs[i] = -1;
+				}
+			}
+		}
+	}
+}
+
+void FMidiPlayCursor::PrepareLookAheadTicks(bool ForceNoBroadcast)
+{
+	check(Owner);
+
+	bool Broadcast = !ForceNoBroadcast &&
+		(SyncOpts != ESyncOptions::NoBroadcastNoPreRoll) &&
+		(SyncOpts == ESyncOptions::BroadcastImmediately || (SyncOpts == ESyncOptions::PreRollIfAtStartOrBroadcast && !Tracker->IsAtStart()));
+
+	if ((FilterPassFlags & EFilterPassFlags::Reset) != EFilterPassFlags::None)
+	{
+		OnReset();
+	}
+
+	if (LookaheadTicks > 0 &&
+		Tracker->IsAtStart() &&
+		(SyncOpts == ESyncOptions::PreRollIfAtStartOrBroadcast || SyncOpts == ESyncOptions::PreRollIfAtStartOrNoBroadcast))
+	{
+		Owner->RecalculatePreRollDueToCursorPosition(this);
+		return;
+	}
+
+	if (Broadcast)
+	{
+		AdvanceByTicks(true);
+	}
+	else
+	{
+		SeekThruTick(Tracker->CurrentTick + LookaheadTicks);
+	}
+}
+
+void FMidiPlayCursor::SetupMsLookahead(float Ms, ESyncOptions Opts)
+{
+	LookaheadType = ELookaheadType::Time;
+	LookaheadMs = Ms;
+	SyncOpts = Opts;
+
+	// may be called before or after we are attached to a play cursor manager!
+	if (Owner)
+	{
+		PrepareLookAheadMs();
+	}
+}
+
+void FMidiPlayCursor::PrepareLookAheadMs(bool ForceNoBroadcast)
+{
+	check(Owner);
+
+	bool Broadcast = !ForceNoBroadcast &&
+		(SyncOpts != ESyncOptions::NoBroadcastNoPreRoll) &&
+		(SyncOpts == ESyncOptions::BroadcastImmediately || (SyncOpts == ESyncOptions::PreRollIfAtStartOrBroadcast && !Tracker->IsAtStart()));
+
+	if ((FilterPassFlags & EFilterPassFlags::Reset) != EFilterPassFlags::None)
+	{
+		OnReset();
+	}
+
+	if (LookaheadMs > 0.0f &&
+		Tracker->IsAtStart() &&
+		(SyncOpts == ESyncOptions::PreRollIfAtStartOrBroadcast || SyncOpts == ESyncOptions::PreRollIfAtStartOrNoBroadcast))
+	{
+		Owner->RecalculatePreRollDueToCursorPosition(this);
+		return;
+	}
+
+	if (Broadcast)
+	{
+		AdvanceByMs(true);
+	}
+	else
+	{
+		float NewMs = Tracker->CurrentMs + LookaheadMs;
+		if (NewMs > 0.0f)
+		{
+			CurrentTick = (int32)(Owner->GetTempoMap().MsToTick(NewMs) + 0.5f);
+		}
+		else
+		{
+			CurrentTick = (int32)(Owner->GetTempoMap().MsToTick(NewMs) - 0.5f);
+		}
+		SeekThruTick(CurrentTick);
+		CurrentMs = NewMs;
+	}
+}
+
+bool FMidiPlayCursor::Advance(bool IsLowRes)
+{
+	check(Owner);
+	bool ProcessLoops = IsLowRes ||
+		!(Owner->IsDirectMappedTimeFollower() && LookaheadTicks == 0);
+	switch (LookaheadType)
+	{
+	case ELookaheadType::Ticks: AdvanceByTicks(ProcessLoops); break;
+	case ELookaheadType::Time:  AdvanceByMs(ProcessLoops); break;
+	default: checkNoEntry(); break;
+	}
+	return !UnregisterASAP;
+}
+
+bool FMidiPlayCursor::AdvanceAsPreRoll()
+{
+	check(Owner);
+	switch (LookaheadType)
+	{
+	case ELookaheadType::Ticks: AdvanceByTicks(true, true, true); break;
+	case ELookaheadType::Time:  AdvanceByMs(true, true, true); break;
+	default: checkNoEntry(); break;
+	}
+	return !UnregisterASAP;
+}
+
+
+void FMidiPlayCursor::AdvanceByTicks(bool ProcessLoops, bool Broadcast, bool IsPreRoll)
+{
+	if (LookaheadTicks < 0)
+	{
+		DoAdvanceForLaggingTickCursor(Broadcast, IsPreRoll);
+	}
+	else
+	{
+		DoAdvanceForLeadingTickCursor(Broadcast, ProcessLoops, IsPreRoll);
+	}
+}
+
+void FMidiPlayCursor::DoAdvanceForLaggingTickCursor(bool Broadcast, bool IsPreRoll)
+{
+	int32 NewTick = Tracker->CurrentTick + LookaheadTicks;
+	if (Owner->DoesLoop())
+	{
+		int32 LoopStartTick = Owner->GetLoopStartTick();
+		int32 LoopEndTick = Owner->GetLoopEndTick();
+		int32 TrackerTick = Tracker->CurrentTick;
+
+		// loop or jump?
+		if (LoopStartTick < LoopEndTick)
+		{
+			// possible loop back...
+			if (CurrentTick < TrackerTick)
+			{
+				// cases 1,2 & 3...
+				if (Broadcast)
+				{
+					AdvanceThruTick(NewTick, IsPreRoll);
+				}
+				else
+				{
+					SeekThruTick(NewTick);
+				}
+			}
+			else if (NewTick < LoopStartTick)
+			{
+				// wrap new tick back to loop end...
+				NewTick = LoopEndTick - (LoopStartTick - NewTick);
+				// now advance from current position to new position...
+				if (Broadcast)
+				{
+					AdvanceThruTick(NewTick, IsPreRoll);
+				}
+				else
+				{
+					SeekThruTick(NewTick);
+				}
+			}
+			else
+			{
+				if (Broadcast)
+				{
+					// advance from current position to loop end...
+					AdvanceThruTick(LoopEndTick - 1, IsPreRoll);
+					// notify...
+					if ((FilterPassFlags & EFilterPassFlags::Loop) != EFilterPassFlags::None)
+					{
+						OnLoop(LoopStartTick, LoopEndTick);
+					}
+					// advance from loop start to new position...
+					SeekToTick(LoopStartTick);
+					AdvanceThruTick(NewTick, IsPreRoll);
+				}
+				else
+				{
+					SeekThruTick(NewTick);
+				}
+			}
+		}
+		else
+		{
+			// possible jump forward...
+			if ((CurrentTick < LoopEndTick && TrackerTick < LoopEndTick) ||
+				(CurrentTick > LoopStartTick && TrackerTick > LoopStartTick))
+			{
+				// easy
+				if (Broadcast)
+				{
+					AdvanceThruTick(NewTick, IsPreRoll);
+				}
+				else
+				{
+					SeekThruTick(NewTick);
+				}
+			}
+			else
+			{
+				if (NewTick < LoopStartTick)
+				{
+					// wrap new tick back from loopEnd...
+					NewTick = LoopEndTick - (LoopStartTick - NewTick);
+					if (Broadcast)
+					{
+						AdvanceThruTick(NewTick, IsPreRoll);
+					}
+					else
+					{
+						SeekThruTick(NewTick);
+					}
+				}
+				else
+				{
+					if (Broadcast)
+					{
+						// advance from current position to loop end...
+						AdvanceThruTick(LoopEndTick - 1, IsPreRoll);
+						// notify...
+						if ((FilterPassFlags & EFilterPassFlags::Loop) != EFilterPassFlags::None)
+						{
+							OnLoop(LoopStartTick, LoopEndTick);
+						}
+						// advance from loop start to new position...
+						SeekToTick(LoopStartTick);
+						AdvanceThruTick(NewTick, IsPreRoll);
+					}
+					else
+					{
+						SeekThruTick(NewTick);
+					}
+				}
+			}
+		}
+	}
+	else
+	{
+		// simple... no loop...
+		if (Broadcast)
+		{
+			AdvanceThruTick(NewTick, IsPreRoll);
+		}
+		else
+		{
+			SeekThruTick(NewTick);
+		}
+	}
+}
+
+void FMidiPlayCursor::DoAdvanceForLeadingTickCursor(bool Broadcast, bool ProcessLoops, bool IsPreRoll)
+{
+	int32 NewTick = Tracker->CurrentTick + LookaheadTicks;
+	if (Owner->DoesLoop() && ProcessLoops) 
+	{
+		int32 LoopStartTick = Owner->GetLoopStartTick();
+		int32 LoopEndTick = Owner->GetLoopEndTick();
+		int32 TrackerTick = Tracker->CurrentTick;
+
+		// loop or jump?
+		if (LoopStartTick < LoopEndTick)
+		{
+			// possible loop back...
+			if (NewTick >= LoopEndTick)
+			{
+				// wrap NewTick...
+				NewTick = LoopStartTick + (NewTick - LoopEndTick);
+			}
+			if (NewTick < CurrentTick)
+			{
+				if (Broadcast)
+				{
+					// advance from current position to loop end...
+					AdvanceThruTick(LoopEndTick - 1, IsPreRoll);
+					// notify...
+					if ((FilterPassFlags & EFilterPassFlags::Loop) != EFilterPassFlags::None)
+					{
+						OnLoop(LoopStartTick, LoopEndTick);
+					}
+					// advance from loop start to new position...
+					SeekToTick(LoopStartTick);
+					AdvanceThruTick(NewTick, IsPreRoll);
+				}
+				else
+				{
+					SeekThruTick(NewTick);
+				}
+			}
+			else
+			{
+				if (Broadcast)
+				{
+					AdvanceThruTick(NewTick, IsPreRoll);
+				}
+				else
+				{
+					SeekThruTick(NewTick);
+				}
+			}
+		}
+		else
+		{
+			// possible jump forward...
+			if ((NewTick < LoopEndTick && CurrentTick < LoopEndTick) ||
+				(NewTick >= LoopStartTick && CurrentTick >= LoopStartTick))
+			{
+				if (Broadcast)
+				{
+					AdvanceThruTick(NewTick, IsPreRoll);
+				}
+				else
+				{
+					SeekThruTick(NewTick);
+				}
+			}
+			else
+			{
+				if (NewTick > (LoopEndTick - 1) && NewTick < LoopStartTick)
+				{
+					// jump NewTick...
+					NewTick = LoopStartTick + (NewTick - LoopEndTick);
+				}
+
+				if (CurrentTick < LoopEndTick)
+				{
+					// advance from current position to loop end...
+					AdvanceThruTick(LoopEndTick - 1, IsPreRoll);
+					// notify...
+					if ((FilterPassFlags & EFilterPassFlags::Loop) != EFilterPassFlags::None)
+					{
+						OnLoop(LoopStartTick, LoopEndTick);
+					}
+					// advance from loop start to new position...
+					SeekToTick(LoopStartTick);
+					AdvanceThruTick(NewTick, IsPreRoll);
+				}
+				else
+				{
+					if (Broadcast)
+					{
+						AdvanceThruTick(NewTick, IsPreRoll);
+					}
+					else
+					{
+						SeekThruTick(NewTick);
+					}
+				}
+			}
+		}
+	}
+	else
+	{
+		// simple... no loop...
+		if (Broadcast)
+		{
+			AdvanceThruTick(NewTick, IsPreRoll);
+		}
+		else
+		{
+			SeekThruTick(NewTick);
+		}
+	}
+}
+
+void FMidiPlayCursor::AdvanceByMs(bool ProcessLoops, bool Broadcast, bool IsPreRoll)
+{
+	if (LookaheadMs < 0.0f)
+	{
+		DoAdvanceForLaggingMsCursor(Broadcast, IsPreRoll);
+	}
+	else
+	{
+		DoAdvanceForLeadingMsCursor(Broadcast, ProcessLoops, IsPreRoll);
+	}
+}
+
+void FMidiPlayCursor::DoAdvanceForLaggingMsCursor(bool Broadcast, bool IsPreRoll)
+{
+	float NewMs = Tracker->CurrentMs + LookaheadMs;
+	int32 NewTick = (int32)(Owner->GetTempoMap().MsToTick(NewMs) + 0.5f);
+	if (Owner->DoesLoop())
+	{
+		float LoopStartMs = Owner->GetLoopStartMs();
+		float LoopEndMs = Owner->GetLoopEndMs();
+		int32 LoopStartTick = Owner->GetLoopStartTick();
+		int32 LoopEndTick = Owner->GetLoopEndTick();
+		int32 TrackerTick = Tracker->CurrentTick;
+
+		// loop or jump?
+		if (LoopStartTick < LoopEndTick)
+		{
+			// possible loop back...
+			if (CurrentTick < TrackerTick)
+			{
+				// cases 1,2 & 3...
+				if (Broadcast)
+				{
+					AdvanceThruTick(NewTick, IsPreRoll);
+				}
+				else
+				{
+					SeekThruTick(NewTick);
+				}
+			}
+			else if (NewMs < LoopStartMs)
+			{
+				// wrap new tick back to loop end...
+				NewMs = LoopEndMs - (LoopStartMs - NewMs);
+				NewTick = (int32)(Owner->GetTempoMap().MsToTick(NewMs) + 0.5f);
+				// now advance from current position to new position...
+				if (Broadcast)
+				{
+					AdvanceThruTick(NewTick, IsPreRoll);
+				}
+				else
+				{
+					SeekThruTick(NewTick);
+				}
+			}
+			else
+			{
+				if (Broadcast)
+				{
+					// advance from current position to loop end...
+					AdvanceThruTick(LoopEndTick - 1, IsPreRoll);
+					// notify...
+					if ((FilterPassFlags & EFilterPassFlags::Loop) != EFilterPassFlags::None)
+					{
+						OnLoop(LoopStartTick, LoopEndTick);
+					}
+					// advance from loop start to new position...
+					SeekToTick(LoopStartTick);
+					AdvanceThruTick(NewTick, IsPreRoll);
+				}
+				else
+				{
+					SeekThruTick(NewTick);
+				}
+			}
+		}
+		else
+		{
+			// possible jump forward...
+			if ((CurrentTick < LoopEndTick && TrackerTick < LoopEndTick) ||
+				(CurrentTick > LoopStartTick && TrackerTick > LoopStartTick))
+			{
+				// easy
+				if (Broadcast)
+				{
+					AdvanceThruTick(NewTick, IsPreRoll);
+				}
+				else
+				{
+					SeekThruTick(NewTick);
+				}
+			}
+			else
+			{
+				if (NewMs < LoopStartMs)
+				{
+					// wrap new tick back from loopEnd...
+					NewMs = LoopEndMs - (LoopStartMs - NewMs);
+					NewTick = (int32)(Owner->GetTempoMap().MsToTick(NewMs) + 0.5f);
+					if (Broadcast)
+					{
+						AdvanceThruTick(NewTick, IsPreRoll);
+					}
+					else
+					{
+						SeekThruTick(NewTick);
+					}
+				}
+				else
+				{
+					if (Broadcast)
+					{
+						// advance from current position to loop end...
+						AdvanceThruTick(LoopEndTick - 1, IsPreRoll);
+						// notify...
+						if ((FilterPassFlags & EFilterPassFlags::Loop) != EFilterPassFlags::None)
+						{
+							OnLoop(LoopStartTick, LoopEndTick);
+						}
+						// advance from loop start to new position...
+						SeekToTick(LoopStartTick);
+						AdvanceThruTick(NewTick, IsPreRoll);
+					}
+					else
+					{
+						SeekThruTick(NewTick);
+					}
+				}
+			}
+		}
+	}
+	else
+	{
+		// simple... no loop...
+		if (Broadcast)
+		{
+			AdvanceThruTick(NewTick, IsPreRoll);
+		}
+		else
+		{
+			SeekThruTick(NewTick);
+		}
+	}
+	CurrentMs = NewMs;
+}
+
+void FMidiPlayCursor::DoAdvanceForLeadingMsCursor(bool Broadcast, bool ProcessLoops, bool IsPreRoll)
+{
+	float NewMs = Tracker->CurrentMs + LookaheadMs;
+	int32 NewTick = (int32)(Owner->GetTempoMap().MsToTick(NewMs) + 0.5f);
+	if (Owner->DoesLoop() && ProcessLoops) 
+	{
+		float LoopStartMs = Owner->GetLoopStartMs();
+		float LoopEndMs = Owner->GetLoopEndMs();
+		int32 LoopStartTick = Owner->GetLoopStartTick();
+		int32 LoopEndTick = Owner->GetLoopEndTick();
+		int32 TrackerTick = Tracker->CurrentTick;
+
+		// loop or jump?
+		if (LoopStartTick < LoopEndTick)
+		{
+			// possible loop back...
+			if (NewMs >= LoopEndMs)
+			{
+				// wrap NewTick...
+				NewMs = LoopStartMs + (NewMs - LoopEndMs);
+				NewTick = (int32)(Owner->GetTempoMap().MsToTick(NewMs) + 0.5f);
+			}
+			if (NewTick < CurrentTick)
+			{
+				if (Broadcast)
+				{
+					// advance from current position to loop end...
+					AdvanceThruTick(LoopEndTick - 1, IsPreRoll);
+					// notify...
+					if ((FilterPassFlags & EFilterPassFlags::Loop) != EFilterPassFlags::None)
+					{
+						OnLoop(LoopStartTick, LoopEndTick);
+					}
+					// advance from loop start to new position...
+					SeekToTick(LoopStartTick);
+					AdvanceThruTick(NewTick, IsPreRoll);
+				}
+				else
+				{
+					SeekThruTick(NewTick);
+				}
+			}
+			else
+			{
+				if (Broadcast)
+				{
+					AdvanceThruTick(NewTick, IsPreRoll);
+				}
+				else
+				{
+					SeekThruTick(NewTick);
+				}
+			}
+		}
+		else
+		{
+			// possible jump forward...
+			if ((NewTick < LoopEndTick && CurrentTick < LoopEndTick) ||
+				(NewTick >= LoopStartTick && CurrentTick >= LoopStartTick))
+			{
+				if (Broadcast)
+				{
+					AdvanceThruTick(NewTick, IsPreRoll);
+				}
+				else
+				{
+					SeekThruTick(NewTick);
+				}
+			}
+			else
+			{
+				if (NewTick > (LoopEndTick - 1) && NewTick < LoopStartTick)
+				{
+					// jump NewTick...
+					NewMs = LoopStartMs + (NewMs - LoopEndMs);
+					NewTick = (int32)(Owner->GetTempoMap().MsToTick(NewMs) + 0.5f);
+				}
+
+				if (CurrentTick < LoopEndTick)
+				{
+					// advance from current position to loop end...
+					AdvanceThruTick(LoopEndTick - 1, IsPreRoll);
+					// notify...
+					if ((FilterPassFlags & EFilterPassFlags::Loop) != EFilterPassFlags::None)
+					{
+						OnLoop(LoopStartTick, LoopEndTick);
+					}
+					// advance from loop start to new position...
+					SeekToTick(LoopStartTick);
+					AdvanceThruTick(NewTick, IsPreRoll);
+				}
+				else
+				{
+					if (Broadcast)
+					{
+						AdvanceThruTick(NewTick, IsPreRoll);
+					}
+					else
+					{
+						SeekThruTick(NewTick);
+					}
+				}
+			}
+		}
+	}
+	else
+	{
+		// simple... no loop...
+		if (Broadcast)
+		{
+			AdvanceThruTick(NewTick, IsPreRoll);
+		}
+		else
+		{
+			SeekThruTick(NewTick);
+		}
+	}
+	CurrentMs = NewMs;
+}
+
+void FMidiPlayCursor::SeekToTick(int32 Tick)
+{
+	const UMidiFile::FMidiTrackList& Tracks = Owner->Tracks();
+	int32 StartIndex = (WatchTrack < 0) ? 0 : WatchTrack;
+	int32 EndIndex = (WatchTrack < 0) ? Tracks.Num() : (WatchTrack + 1);
+	if (EndIndex > Tracks.Num())
+	{
+		EndIndex = Tracks.Num();
+	}
+	for (int32 i = StartIndex; i < EndIndex; ++i)
+	{
+		const FMidiTrack& Track = Tracks[i];
+		const FMidiEventList& Events = Track.GetEvents();
+		TrackNextEventIndexs[i] = 0;
+		while (Events.Num() > TrackNextEventIndexs[i] && Events[TrackNextEventIndexs[i]].GetTick() < Tick)
+		{
+			TrackNextEventIndexs[i]++;
+		}
+		if (Events.Num() == TrackNextEventIndexs[i])
+		{
+			TrackNextEventIndexs[i] = -1;
+		}
+	}
+	CurrentTick = Tick - 1;
+	CurrentMs = Owner->GetTempoMap().TickToMs(CurrentTick);
+}
+
+void FMidiPlayCursor::SeekThruTick(int32 Tick)
+{
+	const UMidiFile::FMidiTrackList& Tracks = Owner->Tracks();
+	int32 StartIndex = (WatchTrack < 0) ? 0 : WatchTrack;
+	int32 EndIndex = (WatchTrack < 0) ? Tracks.Num() : (WatchTrack + 1);
+	if (EndIndex > Tracks.Num())
+	{
+		EndIndex = Tracks.Num();
+	}
+	for (int32 i = StartIndex; i < EndIndex; ++i)
+	{
+		const FMidiTrack& Track = Tracks[i];
+		const FMidiEventList& Events = Track.GetEvents();
+		TrackNextEventIndexs[i] = 0;
+		while (Events.Num() > TrackNextEventIndexs[i] && Events[TrackNextEventIndexs[i]].GetTick() <= Tick)
+		{
+			TrackNextEventIndexs[i]++;
+		}
+		if (Events.Num() == TrackNextEventIndexs[i])
+		{
+			TrackNextEventIndexs[i] = -1;
+		}
+	}
+	CurrentTick = Tick;
+	CurrentMs = Owner->GetTempoMap().TickToMs(Tick);
+}
+
+bool FMidiPlayCursor::IsDone() const
+{
+	const UMidiFile::FMidiTrackList& Tracks = Owner->Tracks();
+	int32 StartIndex = (WatchTrack < 0) ? 0 : WatchTrack;
+	int32 EndIndex = (WatchTrack < 0) ? Tracks.Num() : (WatchTrack + 1);
+	if (EndIndex > Tracks.Num())
+	{
+		EndIndex = Tracks.Num();
+	}
+	for (int32 i = StartIndex; i < EndIndex; ++i)
+	{
+		if (TrackNextEventIndexs[i] != -1)
+		{
+			return false;
+		}
+	}
+	return true;
+}
+
+void FMidiPlayCursor::MoveToLoopStartIfAtNotOffset(int32 NewThruTick, float NewThruMs)
+{
+	if ((LookaheadType == ELookaheadType::Ticks && LookaheadTicks == 0) ||
+		(LookaheadType == ELookaheadType::Time && LookaheadMs == 0.0f))
+	{
+		//mCurrentTick = newThruTick;
+		//mCurrentMs   = newThruMs;
+
+		// This will set both mCurrentTick and mCurrentMs
+		SeekThruTick(NewThruTick);
+	}
+}
+
+namespace
+{
+	class FMidiEventHashTable
+	{
+	private:
+		struct Node
+		{
+			int32 TrackIndex = 0;
+			const FMidiEvent* E = nullptr;
+			Node* Next = nullptr;
+			uint32 Key;
+		};
+
+	public:
+		static FMidiEventHashTable& Get();
+
+		FMidiEventHashTable(int32 Size)
+			: NumNodes(Size)
+		{
+			Head = (Node*)FMemory::Malloc(sizeof(Node) * Size, 16);
+			Clear();
+		}
+
+		bool Push(int32 TrackIndex, const FMidiEvent* E)
+		{
+			ensure(FreeNodes);
+			if (!FreeNodes)
+			{
+				return false;
+			}
+
+			IsClear = false;
+
+			int32 BucketIndex;
+			uint32 Key;
+			MakeHashAndKey(TrackIndex, E, BucketIndex, Key);
+			// make sure this isn't in the table already!
+			Pop(BucketIndex, Key);
+
+			Node* N = FreeNodes;
+			FreeNodes = N->Next;
+			Count++;
+
+			N->TrackIndex = TrackIndex;
+			N->E = E;
+			N->Key = Key;
+			N->Next = Buckets[BucketIndex];
+			Buckets[BucketIndex] = N;
+
+			return true;
+		}
+
+		bool Pop(int32 TrackIndex, const FMidiEvent* E)
+		{
+			int32 BucketIndex;
+			uint32 Key;
+			MakeHashAndKey(TrackIndex, E, BucketIndex, Key);
+			return Pop(BucketIndex, Key);
+		}
+
+		bool Pop(int32 BucketIndex, uint32 Key)
+		{
+			Node* List = Buckets[BucketIndex];
+			if (!List)
+			{
+				return false;
+			}
+
+			Node* Freeable = nullptr;
+			if (List->Key == Key)
+			{
+				Freeable = List;
+				Buckets[BucketIndex] = List->Next;
+			}
+			else
+			{
+				while (List->Next && List->Next->Key != Key)
+				{
+					List = List->Next;
+				}
+				if (List->Next)
+				{
+					Freeable = List->Next;
+					List->Next = List->Next->Next;
+				}
+			}
+			if (Freeable)
+			{
+				Freeable->Next = FreeNodes;
+				FreeNodes = Freeable;
+				Count--;
+				return true;
+			}
+			return false;
+		}
+
+		Node* FindKeyMatch(int32 BucketIndex, uint32 Key)
+		{
+			Node* N = Buckets[BucketIndex];
+			while (N && N->Key != Key)
+			{
+				N = N->Next;
+			}
+			return N;
+		}
+
+		~FMidiEventHashTable()
+		{
+			FMemory::Free(Head);
+		}
+
+		using WorkFunction = std::function<void(int32, const FMidiEvent*)>;
+		void DoOnAll(WorkFunction WorkFunc)
+		{
+			for (int32 i = 0; i < kNumBuckets; ++i)
+			{
+				Node* List = Buckets[i];
+				while (List)
+				{
+					WorkFunc(List->TrackIndex, List->E);
+					List = List->Next;
+				}
+			}
+		}
+
+		void Clear()
+		{
+			if (IsClear)
+			{
+				return;
+			}
+			FreeNodes = Head;
+			for (int32 i = 0; i < NumNodes - 1; ++i)
+			{
+				FreeNodes[i].Next = &FreeNodes[i + 1];
+			}
+			FreeNodes[NumNodes - 1].Next = nullptr;
+			for (int32 i = 0; i < kNumBuckets; ++i)
+			{
+				Buckets[i] = nullptr;
+			}
+			Count = 0;
+			IsClear = true;
+		}
+
+	private:
+		int32 NumNodes = 0;
+		int32 Count = 0;
+		bool  IsClear = false;
+		Node* Head;
+		Node* FreeNodes;
+		static constexpr int32 kNumBuckets = 24;
+		Node* Buckets[kNumBuckets];
+
+		void MakeHashAndKey(int32 TrackIndex, const FMidiEvent* E, int32& BucketIndex, uint32& Key)
+		{
+			int32 Note = (int32)E->GetMsg().GetStdData1();
+			int32 MidiCh = (int32)E->GetMsg().GetStdChannel();
+			Key = (uint32)((TrackIndex << 16) | (MidiCh << 8) | Note);
+
+			BucketIndex = ((Note % 13) + MidiCh + (13 * TrackIndex)) % kNumBuckets;
+			checkf(BucketIndex >= 0 && BucketIndex < kNumBuckets, TEXT("What!? buckIndex = %d, numBuckets = %d, note = %d, midi ch = %d, trackIndex = %d"), BucketIndex, kNumBuckets, Note, MidiCh, TrackIndex);
+			// Maybe we need this ?!? bucketIndex = Abs((note % 13) + midiCh + (13 * trackIndex)) % kNumBuckets;
+			return;
+		}
+	};
+
+	// This next little wrapper lets us lazily new up a FMidiEventHashTable
+	// only on threads that actually need it. Otherwise the thread is only 
+	// wasting sizeof(void*) bytes for the unused hash table. 
+	class FThreadLocalHashTableHolder
+	{
+	public:
+		FThreadLocalHashTableHolder() : ThreadsTable(nullptr) {}
+		FMidiEventHashTable& Get() 
+		{
+			if (!ThreadsTable)
+			{
+				ThreadsTable = new FMidiEventHashTable(64);
+			}
+			ThreadsTable->Clear();
+			return *ThreadsTable;
+		}
+		~FThreadLocalHashTableHolder() { delete ThreadsTable; }
+	private:
+		FMidiEventHashTable* ThreadsTable = nullptr;
+	};
+
+	static thread_local FThreadLocalHashTableHolder ThisThreadsMidiHashTable;
+}
+
+void FMidiPlayCursor::AdvanceThruTick(int32 Tick, bool IsPreRoll)
+{
+	const UMidiFile::FMidiTrackList& Tracks = Owner->Tracks();
+	FMidiEventHashTable& HeldNoteOnEvents = ThisThreadsMidiHashTable.Get();
+
+	if (!IsPreRoll)
+	{
+		int32 StartIndex = (WatchTrack < 0) ? 0 : WatchTrack;
+		int32 EndIndex = (WatchTrack < 0) ? Tracks.Num() : (WatchTrack + 1);
+		if (EndIndex > Tracks.Num())
+		{
+			EndIndex = Tracks.Num();
+		}
+		for (int32 i = StartIndex; i < EndIndex && !UnregisterASAP; ++i)
+		{
+			const FMidiTrack& Track = Tracks[i];
+			const FMidiEventList& Events = Track.GetEvents();
+			while (TrackNextEventIndexs[i] != -1 && Tick >= Events[TrackNextEventIndexs[i]].GetTick() && !UnregisterASAP)
+			{
+				BroadcastEvent(i, Events[TrackNextEventIndexs[i]]);
+				TrackNextEventIndexs[i]++;
+				if (TrackNextEventIndexs[i] == Events.Num())
+				{
+					TrackNextEventIndexs[i] = -1;
+				}
+			}
+		}
+	}
+	else
+	{
+		// This is a much more complicated case...
+		// We only want to broadcast note-on messages that don't have a corresponding note-off message
+		// between the preroll start tick and the current tick.
+		int32 StartIndex = (WatchTrack < 0) ? 0 : WatchTrack;
+		int32 EndIndex = (WatchTrack < 0) ? Tracks.Num() : (WatchTrack + 1);
+		if (EndIndex > Tracks.Num())
+		{
+			EndIndex = Tracks.Num();
+		}
+		for (int32 i = StartIndex; i < EndIndex && !UnregisterASAP; ++i)
+		{
+			const FMidiTrack& Track = Tracks[i];
+			const FMidiEventList& Events = Track.GetEvents();
+			while (TrackNextEventIndexs[i] != -1 && Tick >= Events[TrackNextEventIndexs[i]].GetTick() && !UnregisterASAP)
+			{
+				const FMidiMsg& Msg = Events[TrackNextEventIndexs[i]].GetMsg();
+				bool Swallow = false;
+				if (Msg.MsgType() == FMidiMsg::EType::Std)
+				{
+					if (Msg.IsNoteOn())
+					{
+						if (Msg.GetStdData2() != 0)
+						{
+							HeldNoteOnEvents.Push(i, &Events[TrackNextEventIndexs[i]]);
+						}
+						else
+						{
+							HeldNoteOnEvents.Pop(i, &Events[TrackNextEventIndexs[i]]);
+						}
+						Swallow = true;
+					}
+					else if (Msg.IsNoteOff())
+					{
+						HeldNoteOnEvents.Pop(i, &Events[TrackNextEventIndexs[i]]);
+						Swallow = true;
+					}
+				}
+				if (!Swallow)
+				{
+					BroadcastEvent(i, Events[TrackNextEventIndexs[i]], true);
+				}
+				TrackNextEventIndexs[i]++;
+				if (TrackNextEventIndexs[i] == Events.Num())
+				{
+					TrackNextEventIndexs[i] = -1;
+				}
+			}
+		}
+	}
+
+	// A result of one of the broadcasts above may have resulted in this cursor being
+	// unregistered. So... check for that possibility...
+	if (UnregisterASAP)
+	{
+		return;
+	}
+
+	CurrentTick = Tick;
+	const FTempoMap& TempoMap = Owner->GetTempoMap();
+	CurrentMs = TempoMap.TickToMs(Tick);
+
+	// now we MAY have note on events which still need to be broadcast...
+	if ((FilterPassFlags & EFilterPassFlags::PreRollNoteOn) != EFilterPassFlags::None)
+	{
+		HeldNoteOnEvents.DoOnAll(
+			[&](int32 TrackIndex, const FMidiEvent* Event)
+			{
+				const FMidiMsg& Msg = Event->GetMsg();
+				int32 EventTick = Event->GetTick();
+				float EventMs = TempoMap.TickToMs(EventTick);
+				OnPreRollNoteOn(TrackIndex, EventTick, CurrentTick, CurrentMs - EventMs, Msg.GetStdStatus(), Msg.GetStdData1(), Msg.GetStdData2());
+			});
+	}
+}
+
+void FMidiPlayCursor::BroadcastEvent(int32 TrackIndex, const FMidiEvent& Event, bool IsPreroll)
+{
+	const FMidiMsg& Msg = Event.GetMsg();
+	//Call callbacks for each event that the cursor is interested in
+	switch (Msg.MsgType())
+	{
+	case FMidiMsg::EType::Std:
+		if ((FilterPassFlags & EFilterPassFlags::MidiMessage) != EFilterPassFlags::None)
+		{
+			OnMidiMessage(TrackIndex, Event.GetTick(), Msg.GetStdStatus(), Msg.GetStdData1(), Msg.GetStdData2(), IsPreroll);
+		}
+		break;
+	case FMidiMsg::EType::Tempo:
+		if ((FilterPassFlags & EFilterPassFlags::Tempo) != EFilterPassFlags::None)
+		{
+			OnTempo(TrackIndex, Event.GetTick(), Msg.GetMicrosecPerQuarterNote(), IsPreroll);
+		}
+		break;
+	case FMidiMsg::EType::Text:
+		if ((FilterPassFlags & EFilterPassFlags::Text) != EFilterPassFlags::None)
+		{
+			const FMidiTrack& Track = Owner->Tracks()[TrackIndex];
+			OnText(TrackIndex, Event.GetTick(), Msg.GetTextIndex(), Track.GetTextAtIndex(Msg.GetTextIndex()), Msg.GetTextType(), IsPreroll);
+		}
+		break;
+	case FMidiMsg::EType::TimeSig:
+		if ((FilterPassFlags & EFilterPassFlags::TimeSig) != EFilterPassFlags::None)
+		{
+			OnTimeSig(TrackIndex, Event.GetTick(), Msg.GetTimeSigNumerator(), Msg.GetTimeSigDenominator(), IsPreroll);
+		}
+		break;
+	default:
+		UE_LOG(LogMidi, Error, TEXT("Unknown midi message type %d on track %d at tick %d, file %s"), int(Msg.MsgType()), TrackIndex, Event.GetTick(), **Owner->GetMidiFileName());
+		break;
+	}
+}
+
+
+/*****************************************************************************
+
+Cursor situations...
+
+C = mCurrentTick
+N = newTick
+T = tracker->mCurrentTick
+
+LAGGING
+=======
+
+   Loop:         LS         LE
+				  |          |
+		1)  C  NT |          |
+		2)    C  N|T         |
+		3)      C | NT       |
+		4)        |C  NT     |
+		5)        |T     C  N|
+		6)        | NT     C |
+
+   Jump:         LE         LS
+				  |          |
+		1)  C  NT |          |
+		2)    C  N|          |T
+		3)      C |          | NT
+		4)        |          |C  NT
+
+LEADING
+=======
+   Loop:         LS         LE
+				  |          |
+		1)   T  CN|          |
+		2)    T  C|N         |
+		3)      T | CN       |
+		4)        |T  CN     |
+		5)        |N     T  C|
+		6)        | CN     T |
+
+   Jump:         LE         LS
+				  |          |
+		1)  T  CN |          |
+		2)    T  C|          |N
+		3)      T |          | CN
+		4)        |          |T  CN
+
+*****************************************************************************/
diff --git a/Engine/Plugins/Runtime/Harmonix/Source/HarmonixMidi/Private/HarmonixMidi/MidiPlayCursorMgr.cpp b/Engine/Plugins/Runtime/Harmonix/Source/HarmonixMidi/Private/HarmonixMidi/MidiPlayCursorMgr.cpp
new file mode 100644
--- /dev/null
+++ b/Engine/Plugins/Runtime/Harmonix/Source/HarmonixMidi/Private/HarmonixMidi/MidiPlayCursorMgr.cpp
@@ -0,0 +1,826 @@
+// Copyright Epic Games, Inc. All Rights Reserved.
+#include "HarmonixMidi/MidiPlayCursorMgr.h"
+#include "HarmonixMidi/MidiPlayCursor.h"
+
+#include "HarmonixMidi/MidiFile.h"
+#include "HarmonixMidi/MidiReader.h"
+#include "HarmonixMidi/MidiTrack.h"
+#include "HarmonixMidi/BarMap.h"
+#include "HarmonixMidi/TempoMap.h"
+
+FMidiPlayCursorMgr::FMidiPlayCursorMgr()
+	: SongMaps(&DefaultMaps)
+	, LengthMs(0.f)
+	, LengthTicks(0)
+	, DirectMappedTimeFollower(false)
+	, Loop(false)
+	, LoopOffsetTick(0.0f)
+	, LoopStartMs(0)
+	, LoopStartTick(0)
+	, LoopEndMs(0)
+	, LoopEndTick(0)
+	, MsSinceLowResUpdate(0.0f)
+	, HiResLoopedSinceLastLoResUpdate(false)
+	, InMidiChangeLock(false)
+	, CurrentAdvanceRate(1.f)
+{
+	// setup the default tempo map to have one entry...
+	DefaultMaps.GetTempoMap().AddTempoInfoPoint(MidiConstants::BPMToMidiTempo(120.0f), 0);
+	DefaultMaps.GetBarMap().AddTimeSignatureAtBarIncludingCountIn(0, 4, 4);
+}
+
+void FMidiPlayCursorMgr::Reset()
+{
+	FScopeLock LowResLock(&LowResCursorListLock);
+	FScopeLock HiResLock(&HiResCursorListLock);
+	UnregisterAllPlayCursors();
+	ResetTrackers();
+	MidiFileData = nullptr;
+	SongMaps = &DefaultMaps;
+	LengthMs = 0.f;
+	LengthTicks = 0;
+	DirectMappedTimeFollower = false;
+	Loop = false;
+	LoopOffsetTick = 0.0f;
+	LoopStartTick = 0;
+	LoopStartMs = 0.0f;
+	LoopEndTick = 0;
+	LoopEndMs = 0.0f;
+	MsSinceLowResUpdate = 0.0f;
+	HiResLoopedSinceLastLoResUpdate = false;
+	MidiFileData = nullptr;
+	InMidiChangeLock = false;
+}
+
+void FMidiPlayCursorMgr::AttachToTimeAuthority(const TSharedPtr<FMidiPlayCursorMgr>& InTimeAuthority)
+{
+	TimeAuthority = InTimeAuthority;
+}
+
+void FMidiPlayCursorMgr::AttachToMidiResource(TSharedPtr<FMidiFileData> InMidiFileData, bool ResetTrackersToStart, int32 PreRollBars)
+{
+	LockForMidiDataChanges();
+	MidiFileData = InMidiFileData;
+	if (MidiFileData)
+	{
+		SongMaps = &MidiFileData->SongMaps;
+	}
+	else
+	{
+		SongMaps = &DefaultMaps;
+	}
+	if (ResetTrackersToStart)
+	{
+		ResetTrackers();
+	}
+	MidiDataChangeComplete(EMidiChangePositionCorrectMode::MaintainTick, PreRollBars);
+}
+
+void FMidiPlayCursorMgr::DetachFromMidiResource()
+{
+	LockForMidiDataChanges();
+	MidiFileData = nullptr;
+	SongMaps = &DefaultMaps;
+	ResetTrackers();
+	MidiDataChangeComplete(EMidiChangePositionCorrectMode::MaintainTick);
+}
+
+FMidiSongPos FMidiPlayCursorMgr::CalculateSongPosWithOffsetMs(float DeltaMs) const
+{
+	FMidiSongPos OutSongPos;
+	float Ms = GetCurrentHiResMs() + DeltaMs;
+	if (DoesLoop())
+	{
+		float LoopLengthMs = LoopEndMs - LoopStartMs;
+		float MappedMs = LoopStartMs + FMath::Fmod(Ms - LoopStartMs, LoopLengthMs);
+		OutSongPos.SetByTime(MappedMs, GetSongMaps());
+	}
+	else
+	{
+		OutSongPos.SetByTime(Ms, GetSongMaps());
+	}
+
+	// currently only use the time authority for the tempo
+	if (TSharedPtr<FMidiPlayCursorMgr> TimeAuthorityPtr = TimeAuthority.Pin())
+	{
+		FMidiSongPos AuthoritySongPos = TimeAuthorityPtr->CalculateSongPosWithOffsetMs(DeltaMs);
+
+		OutSongPos.Tempo = AuthoritySongPos.Tempo;
+	}
+
+	return OutSongPos;
+}
+
+void FMidiPlayCursorMgr::DetermineLength()
+{
+	FScopeLock LowResLock(&LowResCursorListLock);
+	FScopeLock HiResLock(&HiResCursorListLock);
+	if (!MidiFileData)
+	{
+		LengthMs = 0.0f;
+		LengthTicks = 0;
+		LoopEndMs = 0.0f;
+		LoopEndTick = 0;
+		LoopStartMs = 0.0f;
+		LoopStartTick = 0;
+	}
+	else
+	{
+		//Determine file length in ms and ticks
+		LengthTicks = 0;
+		const TArray<FMidiTrack>& Tracks = MidiFileData->Tracks;
+		int32 NumTracks = Tracks.Num();
+		for (int32 i = 0; i < NumTracks; ++i)
+		{
+			LengthTicks = FMath::Max(LengthTicks, Tracks[i].GetEvents().Last().GetTick());
+		}
+		//Round file length up to the nearest bar
+		int32 Bar = FMath::CeilToInt32(SongMaps->GetBarMap().TickToFractionalBarIncludingCountIn(LengthTicks));
+		LengthTicks = SongMaps->GetBarMap().BarIncludingCountInToTick(Bar);
+		LengthMs = SongMaps->GetTempoMap().TickToMs(LengthTicks);
+
+		LoopEndMs = LengthMs;
+		LoopEndTick = LengthTicks;
+		LoopStartMs = 0.0f;
+		LoopStartTick = 0;
+	}
+}
+
+FMidiPlayCursorMgr::~FMidiPlayCursorMgr()
+{
+	UnregisterAllPlayCursors();
+}
+
+void FMidiPlayCursorMgr::RegisterHiResPlayCursor(FMidiPlayCursor* PlayCursor, float PreRollMs)
+{
+	if (PlayCursor->Owner)
+	{
+		PlayCursor->Owner->UnregisterPlayCursor(PlayCursor);
+	}
+	PlayCursor->UnregisterASAP = false;
+	FScopeLock HiResLock(&HiResCursorListLock);
+	PlayCursor->SetOwner(this, &HiResTracker, PreRollMs);
+	HiResTracker.AddCursor(PlayCursor);
+}
+
+void FMidiPlayCursorMgr::RegisterLowResPlayCursor(FMidiPlayCursor* PlayCursor, float PreRollMs)
+{
+	if (PlayCursor->Owner)
+	{
+		PlayCursor->Owner->UnregisterPlayCursor(PlayCursor);
+	}
+	PlayCursor->UnregisterASAP = false;
+	FScopeLock LowResLock(&LowResCursorListLock);
+	PlayCursor->SetOwner(this, &LowResTracker, PreRollMs);
+	LowResTracker.AddCursor(PlayCursor);
+}
+
+void FMidiPlayCursorMgr::RecalculatePreRollDueToCursorPosition(FMidiPlayCursor* PlayCursor)
+{
+	float PreRollMs = 0.0f;
+	if (PlayCursor->GetLookaheadType() == FMidiPlayCursor::ELookaheadType::Ticks)
+	{
+		PreRollMs = -GetTempoMap().TickToMs(-PlayCursor->GetLookaheadTicks());
+	}
+	else
+	{
+		PreRollMs = PlayCursor->GetLookaheadMs();
+	}
+	if (-PreRollMs < HiResTracker.CurrentMs)
+	{
+		// yup... earliest look ahead!
+		int32 NewTick = GetTempoMap().MsToTick(-PreRollMs);
+
+		// back up one tick...
+		NewTick--;
+		PreRollMs = GetTempoMap().TickToMs(NewTick);
+		HiResTracker.Reset(NewTick, PreRollMs, false);
+		LowResTracker.Reset(NewTick, PreRollMs, false);
+	}
+}
+
+void FMidiPlayCursorMgr::UnregisterPlayCursor(FMidiPlayCursor* PlayCursor, bool WarnOnFail)
+{
+	bool WasHiRes = false;
+	{
+		FScopeLock HiResLock(&HiResCursorListLock);
+		if (TraversingHiResCursors)
+		{
+			WasHiRes = HiResTracker.ContainsCursor(PlayCursor);
+		}
+		else
+		{
+			WasHiRes = HiResTracker.RemoveCursor(PlayCursor);
+		}
+		if (WasHiRes && TraversingHiResCursors)
+		{
+			UE_LOG(LogMidi, Error, TEXT("You cannot unregister a MidiPlayCursor while this manager is currently traversing it's cursor list!"));
+			return;
+		}
+	}
+	if (!WasHiRes)
+	{
+		// wasn't on hi-res list... try the low res...
+		FScopeLock LowResLock(&LowResCursorListLock);
+		bool WasLowRes = false;
+		if (TraversingLowResCursors)
+		{
+			WasLowRes = LowResTracker.ContainsCursor(PlayCursor);
+		}
+		else
+		{
+			WasLowRes = LowResTracker.RemoveCursor(PlayCursor);
+		}
+		if (!WasLowRes)
+		{
+			if (WarnOnFail)
+			{
+				UE_LOG(LogMidi, Warning, TEXT("Attempt to remove MidiPlayCursor from Manager that didn't own it!"));
+			}
+			return;
+		}
+		if (WasLowRes && TraversingLowResCursors)
+		{
+			UE_LOG(LogMidi, Error, TEXT("You cannot unregister a MidiPlayCursor while this manager is currently traversing it's cursor list!"));
+			return;
+		}
+	}
+	// it was removed from one or the other lists...
+	PlayCursor->SetOwner(nullptr, nullptr);
+}
+
+void FMidiPlayCursorMgr::UnregisterAllPlayCursors()
+{
+	FScopeLock LowResLock(&LowResCursorListLock);
+	FScopeLock HiResLock(&HiResCursorListLock);
+
+	if (TraversingHiResCursors || TraversingLowResCursors)
+	{
+		UE_LOG(LogMidi, Error, TEXT("You cannot unregister all MidiPlayCursors while this manager is currently traversing it's cursor list!"));
+		return;
+	}
+
+	TraversingHiResCursors = true;
+	for (auto it = HiResTracker.Cursors.begin(); it != HiResTracker.Cursors.end(); ++it)
+	{
+		it.GetNode()->ManagerIsDetaching();
+		it.GetNode()->SetOwner(nullptr, nullptr);
+	}
+	HiResTracker.Clear();
+	TraversingHiResCursors = false;
+	TraversingLowResCursors = true;
+	for (auto it = LowResTracker.Cursors.begin(); it != LowResTracker.Cursors.end(); ++it)
+	{
+		it.GetNode()->ManagerIsDetaching();
+		it.GetNode()->SetOwner(nullptr, nullptr);
+	}
+	LowResTracker.Clear();
+	TraversingLowResCursors = false;
+}
+
+bool FMidiPlayCursorMgr::HasLowResCursors() const
+{
+	return !LowResTracker.Cursors.IsEmpty();
+}
+
+void FMidiPlayCursorMgr::ResetTrackers()
+{
+	FScopeLock LowResLock(&LowResCursorListLock);
+	FScopeLock HiResLock(&HiResCursorListLock);
+	float Ms = GetTempoMap().TickToMs(-1);
+	LowResTracker.Reset(-1, Ms, false);
+	HiResTracker.Reset(-1, Ms, false);
+	MsSinceLowResUpdate = 0.0f;
+	HiResLoopedSinceLastLoResUpdate = false;
+}
+
+bool FMidiPlayCursorMgr::HasMidiFile() const
+{
+	return MidiFileData.IsValid();
+}
+
+int32 FMidiPlayCursorMgr::FindTrackIndexByName(const FString& name) const
+{
+	return MidiFileData.IsValid() ? MidiFileData->FindTrackIndexByName(name) : INDEX_NONE;
+}
+
+const FString* FMidiPlayCursorMgr::GetMidiFileName() const
+{
+	if (!MidiFileData)
+	{
+		return nullptr;
+	}
+
+	return &MidiFileData->MidiFileName;
+}
+
+const UMidiFile::FMidiTrackList& FMidiPlayCursorMgr::Tracks() const
+{
+	return MidiFileData.IsValid() ? MidiFileData->Tracks : DefaultTracks;
+}
+
+const FSongMaps& FMidiPlayCursorMgr::GetSongMaps() const
+{
+	check(SongMaps);
+	return *SongMaps;
+}
+
+const FTempoMap& FMidiPlayCursorMgr::GetTempoMap() const
+{
+	check(SongMaps);
+	return SongMaps->GetTempoMap();
+}
+
+const FBarMap& FMidiPlayCursorMgr::GetBarMap() const
+{
+	check(SongMaps);
+	return SongMaps->GetBarMap();
+}
+
+void FMidiPlayCursorMgr::GetCursorExtentsMs(float& Earliest, float& Latest) const
+{
+	Earliest = (HiResTracker.EarliestCursorMs > LowResTracker.EarliestCursorMs) ?
+		HiResTracker.EarliestCursorMs : LowResTracker.EarliestCursorMs;
+	Latest = (HiResTracker.LatestCursorMs < LowResTracker.LatestCursorMs) ?
+		HiResTracker.LatestCursorMs : LowResTracker.LatestCursorMs;
+}
+
+void FMidiPlayCursorMgr::GetCursorExtentsTicks(int32& Earliest, int32& Latest) const
+{
+	Earliest = (HiResTracker.EarliestCursorTick > LowResTracker.EarliestCursorTick) ?
+		HiResTracker.EarliestCursorTick : LowResTracker.EarliestCursorTick;
+	Latest = (HiResTracker.LatestCursorTick < LowResTracker.LatestCursorTick) ?
+		HiResTracker.LatestCursorTick : LowResTracker.LatestCursorTick;
+}
+
+bool FMidiPlayCursorMgr::CursorsAllInPhase() const
+{
+	if (!Loop)
+	{
+		return true;
+	}
+
+	FScopeLock LowResLock(&LowResCursorListLock);
+	FScopeLock HiResLock(&HiResCursorListLock);
+
+	// Check tick driven...
+	int32 EarliestTick, LatestTick;
+	GetCursorExtentsTicks(EarliestTick, LatestTick);
+	if ((HiResTracker.CurrentTick < LoopEndTick && HiResTracker.CurrentTick + EarliestTick > LoopEndTick) ||
+		(LowResTracker.CurrentTick < LoopEndTick && LowResTracker.CurrentTick + EarliestTick > LoopEndTick) ||
+		(HiResTracker.CurrentTick > LoopStartTick && HiResTracker.CurrentTick + LatestTick < LoopStartTick) ||
+		(LowResTracker.CurrentTick > LoopStartTick && LowResTracker.CurrentTick + LatestTick < LoopStartTick))
+	{
+		return false;
+	}
+
+	// check time driven...
+	float EarliestMs, LatestMs;
+	GetCursorExtentsMs(EarliestMs, LatestMs);
+	if ((HiResTracker.CurrentMs < LoopEndMs && HiResTracker.CurrentMs + EarliestMs > LoopEndMs) ||
+		(LowResTracker.CurrentMs < LoopEndMs && LowResTracker.CurrentMs + EarliestMs > LoopEndMs) ||
+		(HiResTracker.CurrentMs > LoopStartMs && HiResTracker.CurrentMs + LatestMs < LoopStartMs) ||
+		(LowResTracker.CurrentMs > LoopStartMs && LowResTracker.CurrentMs + LatestMs < LoopStartMs))
+	{
+		return false;
+	}
+
+	return true;
+}
+
+int32 FMidiPlayCursorMgr::GetFarthestAheadCursorTick() const
+{
+	int32 Hrt = HiResTracker.GetFarthestAheadCursorTick();
+	int32 Lrt = LowResTracker.GetFarthestAheadCursorTick();
+	return((Lrt > Hrt) ? Lrt : Hrt);
+}
+
+int32 FMidiPlayCursorMgr::GetFarthestBehindCursorTick() const
+{
+	int32 Hrt = HiResTracker.GetFarthestBehindCursorTick();
+	int32 Lrt = LowResTracker.GetFarthestBehindCursorTick();
+	return((Lrt < Hrt) ? Lrt : Hrt);
+}
+
+float FMidiPlayCursorMgr::GetBufferedMs() const
+{
+	return (HiResTracker.EarliestCursorMs > LowResTracker.EarliestCursorMs) ? HiResTracker.EarliestCursorMs : LowResTracker.EarliestCursorMs;
+}
+
+void FMidiPlayCursorMgr::SetLoop(int32 StartTick, int32 EndTick, bool IsDirectMappedFollower, bool IgnoringLookAhead)
+{
+	// NOTE:
+	// TODO: There is a missing test here that would check to see if the proposed loop end
+	// might fall in the middle of the current span of leading and lagging play cursors,
+	// which would result in ugly behavior!
+
+	FScopeLock LowResLock(&LowResCursorListLock);
+	FScopeLock HiResLock(&HiResCursorListLock);
+	const FTempoMap& TempoMap = GetTempoMap();
+	LoopStartTick = StartTick;
+	LoopStartMs = TempoMap.TickToMs(StartTick);
+	if (EndTick == kEndTick)
+	{
+		EndTick = LengthTicks;
+	}
+	LoopEndTick = EndTick;
+	LoopEndMs = TempoMap.TickToMs(EndTick);
+	Loop = true;
+	DirectMappedTimeFollower = IsDirectMappedFollower;
+	bool CursorsInPhase = CursorsAllInPhase();
+	if (!CursorsInPhase)
+	{
+		if (IgnoringLookAhead)
+		{
+			HiResTracker.Reset(HiResTracker.CurrentTick, HiResTracker.CurrentMs, false);
+			LowResTracker.Reset(LowResTracker.CurrentTick, LowResTracker.CurrentMs, false);
+		}
+		else
+		{
+			UE_LOG(LogMidi, Warning, TEXT("FMidiPlayCursorMgr::SetLoop : not ignoring look ahead, but cursors are not in phase!"));
+		}
+	}
+}
+
+void FMidiPlayCursorMgr::ClearLoop(bool IgnoringLookAhead)
+{
+	FScopeLock LowResLock(&LowResCursorListLock);
+	FScopeLock HiResLock(&HiResCursorListLock);
+	Loop = false;
+	bool CursorsInPhase = CursorsAllInPhase();
+	if (!CursorsInPhase)
+	{
+		if (IgnoringLookAhead)
+		{
+			HiResTracker.Reset(HiResTracker.CurrentTick, HiResTracker.CurrentMs, false);
+			LowResTracker.Reset(LowResTracker.CurrentTick, LowResTracker.CurrentMs, false);
+		}
+		else
+		{
+			UE_LOG(LogMidi, Warning, TEXT("FMidiPlayCursorMgr::ClearLoop : not ignoring look ahead, but cursors are not in phase!"));
+		}
+	}
+}
+
+void FMidiPlayCursorMgr::SeekTo(int32 Tick, int32 PreRollBars, bool IsRenderThread, bool IsLoop)
+{
+	FScopeLock LowResLock(&LowResCursorListLock);
+	FScopeLock HiResLock(&HiResCursorListLock);
+
+	// Back up one tick, as we are going to be setting the cursors'
+	// 'played through' position.
+	Tick--;
+	float NewPosMs = GetTempoMap().TickToMs(Tick);
+	float PreRollStartMs = NewPosMs;
+	int32 PreRollStartTick = Tick;
+	if (PreRollBars > 0)
+	{
+		const FBarMap& Map = GetBarMap();
+		float Bar = Map.TickToFractionalBarIncludingCountIn(Tick);
+		Bar -= (float)PreRollBars;
+
+		if (Bar < 0.0f)
+		{
+			PreRollStartTick = -1;
+		}
+		else
+		{
+			PreRollStartTick = Map.BarIncludingCountInToTick(Bar);
+		}
+
+		PreRollStartMs = GetTempoMap().TickToMs(PreRollStartTick);
+
+		if (PreRollStartTick > Tick)
+		{
+			PreRollStartTick = Tick;
+			PreRollStartMs = NewPosMs;
+		}
+	}
+
+	if (IsLoop)
+	{
+		HiResLoopedSinceLastLoResUpdate = true;
+	}
+
+	// The hi-res cursors can just slam to the new position.
+	// If this is resultOfDirectMappedLoop, the hi-res cursors
+	// were already advanced to the end of the loop by the parent.
+	HiResTracker.Reset(Tick, NewPosMs, PreRollStartTick, PreRollStartMs, false);
+	if (!IsRenderThread)
+	{
+		LowResTracker.Reset(Tick, NewPosMs, PreRollStartTick, PreRollStartMs, false);
+	}
+	else
+	{
+		LowResTracker.QueueReset(Tick, NewPosMs, PreRollStartTick, PreRollStartMs, false);
+	}
+}
+
+void FMidiPlayCursorMgr::MoveToLoopStart()
+{
+	FScopeLock HiResLock(&HiResCursorListLock);
+	{
+		if (ensureMsgf(Loop, TEXT("That's odd. Asked to move to the beginning of the loop... but there is no loop!")))
+		{
+			const FTempoMap& TempoMap = GetTempoMap();
+
+			int32   NewThruTick = LoopStartTick - 1;
+			float NewThruMs = TempoMap.TickToMs(NewThruTick);
+
+			// Move the hi-res cursor to the loop start
+			HiResTracker.MoveToLoopStart(NewThruTick, NewThruMs);
+
+			// The low res cursor will need to advance to the end of the loop, and then
+			// advance from loop-start to current position at some point in the future
+			// (during a low-res frame poll)
+			MsSinceLowResUpdate += LoopEndMs - LowResTracker.CurrentMs;
+			HiResLoopedSinceLastLoResUpdate = true;
+		}
+	}
+}
+
+bool FMidiPlayCursorMgr::IsDone() const
+{
+	if (Loop)
+	{
+		return false;
+	}
+
+	TraversingHiResCursors = true;
+	// This next line of ugliness is required because const iterating through a const linked list is broken
+	// and I don't have time to figure out what about the stack of templates is busted.
+	TIntrusiveDoubleLinkedList<FMidiPlayCursor>& NonConstHiResCursorList = const_cast<TIntrusiveDoubleLinkedList<FMidiPlayCursor>&>(HiResTracker.Cursors);
+	for (auto it = NonConstHiResCursorList.begin(); it != NonConstHiResCursorList.end(); it++)
+	{
+		if (!it.GetNode()->IsDone())
+		{
+			TraversingHiResCursors = false;
+			return false;
+		}
+	}
+	TraversingHiResCursors = false;
+	TraversingLowResCursors = true;
+	// This next line of ugliness is required because const iterating through a const linked list is broken
+	// and I don't have time to figure out what about the stack of templates is busted.
+	TIntrusiveDoubleLinkedList<FMidiPlayCursor>& NonConstLowResCursorList = const_cast<TIntrusiveDoubleLinkedList<FMidiPlayCursor>&>(LowResTracker.Cursors);
+	for (auto it = NonConstLowResCursorList.begin(); it != NonConstLowResCursorList.end(); it++)
+	{
+		if (!it.GetNode()->IsDone())
+		{
+			TraversingLowResCursors = false;
+			return false;
+		}
+	}
+	TraversingLowResCursors = false;
+	return true;
+}
+
+void FMidiPlayCursorMgr::LockForMidiDataChanges()
+{
+	if (ensureAlways(!InMidiChangeLock))
+	{
+		InMidiChangeLock = true;
+		LowResCursorListLock.Lock();
+		HiResCursorListLock.Lock();
+	}
+}
+
+void FMidiPlayCursorMgr::MidiDataChangeComplete(EMidiChangePositionCorrectMode PositionMode, int32 PreRollBars)
+{
+	if (ensureAlways(InMidiChangeLock))
+	{
+		// We will need to recalculate extents. Unfortunately, this will also
+		// blow away and previously set up loop points, so cache those...
+		int32 OriginalLoopStartTick = LoopStartTick;
+		int32 OriginalLoopEndTick = LoopEndTick;
+		// Now fix up for possible new lengths...
+		DetermineLength();
+		// Now fix up looping...
+		LoopStartTick = OriginalLoopStartTick;
+		LoopEndTick = FMath::Min(LoopEndTick, OriginalLoopEndTick);
+		LoopStartMs = GetTempoMap().TickToMs(LoopStartTick);
+		LoopEndMs = GetTempoMap().TickToMs(LoopEndTick);
+
+		if (PreRollBars == 0)
+		{
+			// Now fix up the hires cursors...
+			TraversingHiResCursors = true;
+			for (auto it = HiResTracker.Cursors.begin(); it != HiResTracker.Cursors.end(); ++it)
+			{
+				it.GetNode()->RecalcNextEventsDueToMidiChanges(PositionMode);
+			}
+			if (PositionMode == EMidiChangePositionCorrectMode::MaintainTick)
+			{
+				HiResTracker.CurrentMs = GetTempoMap().TickToMs(HiResTracker.CurrentTick);
+			}
+			else
+			{
+				HiResTracker.CurrentTick = GetTempoMap().MsToTick(HiResTracker.CurrentMs);
+			}
+			TraversingHiResCursors = false;
+			// Now the low res cursors...
+			TraversingLowResCursors = true;
+			for (auto it = LowResTracker.Cursors.begin(); it != LowResTracker.Cursors.end(); ++it)
+			{
+				it.GetNode()->RecalcNextEventsDueToMidiChanges(PositionMode);
+			}
+			if (PositionMode == EMidiChangePositionCorrectMode::MaintainTick)
+			{
+				LowResTracker.CurrentMs = GetTempoMap().TickToMs(LowResTracker.CurrentTick);
+			}
+			else
+			{
+				LowResTracker.CurrentTick = GetTempoMap().MsToTick(LowResTracker.CurrentMs);
+			}
+			// Now deal with the possibility that the low res cursors might be "out of phase" with the hi res
+			// cursors. This would happen if we are in the middle of a loop...
+			if (HiResLoopedSinceLastLoResUpdate)
+			{
+				if (HiResTracker.CurrentMs < LowResTracker.CurrentMs)
+				{
+					MsSinceLowResUpdate = (LoopEndMs - LowResTracker.CurrentMs) + (HiResTracker.CurrentMs - LoopStartMs);
+					check(MsSinceLowResUpdate >= 0.0f);
+				}
+				else
+				{
+					// yikes. how could this happen?
+					checkNoEntry();
+					HiResLoopedSinceLastLoResUpdate = false;
+				}
+			}
+			TraversingLowResCursors = false;
+		}
+		else
+		{
+			int32 SeekTick = PositionMode == EMidiChangePositionCorrectMode::MaintainTick ? HiResTracker.CurrentTick : GetTempoMap().MsToTick(HiResTracker.CurrentMs);
+			SeekTo(SeekTick, PreRollBars, true, false);
+		}
+		HiResCursorListLock.Unlock();
+		LowResCursorListLock.Unlock();
+		InMidiChangeLock = false;
+	}
+}
+
+FMusicTimestamp FMidiPlayCursorMgr::GetMusicTimestampAtMs(float Ms) const
+{
+	const int32 TickAtOffset = GetSongMaps().MsToTick(Ms);
+	return GetBarMap().TickToMusicTimestamp(TickAtOffset);
+}
+
+void FMidiPlayCursorMgr::AdvanceHiResToMs(float Ms, bool Broadcast)
+{
+	float DeltaMs = Ms - GetCurrentHiResMs();
+	if (DeltaMs <= 0.0f)
+	{
+		return;
+	}
+	AdvanceHiResByDeltaMs(DeltaMs, Broadcast);
+}
+
+void FMidiPlayCursorMgr::AdvanceHiResByDeltaMs(float Ms, bool Broadcast)
+{
+	FScopeLock HiResLock(&HiResCursorListLock);
+	HiResLoopedSinceLastLoResUpdate = AdvanceTrackerByDeltaMs(Ms, HiResTracker, Broadcast) || HiResLoopedSinceLastLoResUpdate;
+	MsSinceLowResUpdate += Ms;
+}
+
+void FMidiPlayCursorMgr::AdvanceHiResByDeltaTick(int32 NumTicks, bool Broadcast)
+{
+	FScopeLock HiResLock(&HiResCursorListLock);
+
+	int32 ThruTick = NumTicks + HiResTracker.CurrentTick;
+
+	float StartMs = HiResTracker.CurrentMs;
+	AdvanceTrackerThruTick(ThruTick, HiResTracker, false, Broadcast);
+	MsSinceLowResUpdate += HiResTracker.CurrentMs - StartMs;
+}
+
+void FMidiPlayCursorMgr::AdvanceHiResThruTick(int32 ThruTick, bool Broadcast, bool DontAdvancePastLoopEnd)
+{
+	FScopeLock HiResLock(&HiResCursorListLock);
+
+	// early out?...
+	if (ThruTick == HiResTracker.CurrentTick)
+	{
+		return;
+	}
+
+	if (Loop && DontAdvancePastLoopEnd && ThruTick >= LoopEndTick - 1)
+	{
+		ThruTick = LoopEndTick - 1;
+	}
+
+	if (ThruTick < HiResTracker.CurrentTick)
+	{
+		UE_LOG(LogMidi, Warning, TEXT("Asked to go back in time without a seek!"));
+		return;
+	}
+
+	float StartMs = HiResTracker.CurrentMs;
+	AdvanceTrackerThruTick(ThruTick, HiResTracker, false, Broadcast);
+	MsSinceLowResUpdate += HiResTracker.CurrentMs - StartMs;
+}
+
+bool FMidiPlayCursorMgr::AdvanceLowResCursors()
+{
+	FScopeLock LowResLock(&LowResCursorListLock);
+
+	LowResTracker.HandleQueuedReset();
+
+	float MsDiff = 0.0f;
+	bool  HiResLooped = false;
+	float HiResCurrentMs = 0.0f;
+	{
+		FScopeLock HiResLock(&HiResCursorListLock);
+
+		MsDiff = MsSinceLowResUpdate;
+		HiResCurrentMs = HiResTracker.CurrentMs;
+		HiResLooped = HiResLoopedSinceLastLoResUpdate;
+		MsSinceLowResUpdate = 0.0f;
+		HiResLoopedSinceLastLoResUpdate = false;
+	}
+
+	if (!HiResLooped)
+	{
+		MsDiff = HiResCurrentMs - LowResTracker.CurrentMs;
+		if (MsDiff <= 0.0f)
+		{
+			return false;
+		}
+	}
+	return AdvanceTrackerByDeltaMs(MsDiff, LowResTracker, true, true);
+}
+
+bool FMidiPlayCursorMgr::AdvanceTrackerByDeltaMs(float Ms, FMidiPlayCursorTracker& Tracker, bool IsLowRes, bool Broadcast)
+{
+	Tracker.ElapsedMs += Ms;
+	const FTempoMap& TempoMap = GetTempoMap();
+	float NewMs = Tracker.CurrentMs + Ms;
+	int32 NewTick = (int32)(TempoMap.MsToTick(NewMs) + 0.5f);
+	bool  Looped = false;
+
+	if (NewTick >= LoopEndTick && Tracker.CurrentTick < LoopEndTick && Loop) // && !mDirectMappedFollower) If this is a direct mapped follower, than this function only gets called for the lo-res cursor, which needs to do normal loop handling!
+	{
+		// loop the tick around...
+		NewTick = (NewTick - LoopEndTick) + LoopStartTick;
+		NewMs = TempoMap.TickToMs(NewTick);
+		Tracker.CurrentMs = NewMs;
+		Tracker.CurrentTick = NewTick;
+		Looped = true;
+	}
+	else
+	{
+		Tracker.CurrentTick = NewTick;
+		Tracker.CurrentMs = NewMs;
+	}
+
+	bool* TraversingFlag = IsLowRes ? &TraversingLowResCursors : &TraversingHiResCursors;
+	*TraversingFlag = true;
+	for (auto it = Tracker.Cursors.begin(); it != Tracker.Cursors.end();)
+	{
+		if (!it.GetNode()->Advance(IsLowRes))
+		{
+			auto DeadIt = it;
+			++it;
+			Tracker.Cursors.Remove(DeadIt.GetNode());
+			DeadIt.GetNode()->SetOwner(nullptr, nullptr);
+		}
+		else
+		{
+			++it;
+		}
+	}
+	*TraversingFlag = false;
+
+	return Looped;
+}
+
+void FMidiPlayCursorMgr::AdvanceTrackerThruTick(int32 ToTick, FMidiPlayCursorTracker& Tracker, bool IsLowRes, bool Broadcast)
+{
+	const FTempoMap& TempoMap = GetTempoMap();
+
+	float NewMs = TempoMap.TickToMs(ToTick);
+	Tracker.ElapsedMs += NewMs - Tracker.CurrentMs;
+	int32   NewTick = ToTick;
+
+	Tracker.CurrentTick = NewTick;
+	Tracker.CurrentMs = NewMs;
+
+	bool* TraversingFlag = IsLowRes ? &TraversingLowResCursors : &TraversingHiResCursors;
+	*TraversingFlag = true;
+	for (auto it = Tracker.Cursors.begin(); it != Tracker.Cursors.end();)
+	{
+		if (!it.GetNode()->Advance(IsLowRes))
+		{
+			auto DeadIt = it;
+			++it;
+			Tracker.Cursors.Remove(DeadIt.GetNode());
+			DeadIt.GetNode()->SetOwner(nullptr, nullptr);
+		}
+		else
+		{
+			++it;
+		}
+	}
+	*TraversingFlag = false;
+}
diff --git a/Engine/Plugins/Runtime/Harmonix/Source/HarmonixMidi/Private/HarmonixMidi/MidiPlayCursorTracker.cpp b/Engine/Plugins/Runtime/Harmonix/Source/HarmonixMidi/Private/HarmonixMidi/MidiPlayCursorTracker.cpp
new file mode 100644
--- /dev/null
+++ b/Engine/Plugins/Runtime/Harmonix/Source/HarmonixMidi/Private/HarmonixMidi/MidiPlayCursorTracker.cpp
@@ -0,0 +1,224 @@
+// Copyright Epic Games, Inc. All Rights Reserved.
+#include "HarmonixMidi/MidiPlayCursorTracker.h"
+#include "HarmonixMidi/MidiPlayCursor.h"
+
+FMidiPlayCursorTracker::FMidiPlayCursorTracker()
+	: CurrentTick(-1)
+	, CurrentMs(-FMidiPlayCursor::kSmallMs)
+	, ElapsedMs(0.0f)
+	, LoopCount(0)
+	, EarliestCursorTick(0)
+	, LatestCursorTick(0)
+	, EarliestCursorMs(0.0f)
+	, LatestCursorMs(0.0f)
+{
+
+}
+
+bool FMidiPlayCursorTracker::IsAtStart()
+{
+	return CurrentTick <= 0 && CurrentMs <= 0.0 && LoopCount == 0;
+}
+
+void FMidiPlayCursorTracker::AddCursor(FMidiPlayCursor* Cursor)
+{
+	Cursors.AddTail(Cursor);
+	RecalculateExtents();
+}
+
+bool FMidiPlayCursorTracker::ContainsCursor(FMidiPlayCursor* Cursor)
+{
+	for (auto it = Cursors.begin(); it != Cursors.end(); ++it)
+	{
+		if (it.GetNode() == Cursor)
+		{
+			return true;
+		}
+	}
+	return false;
+}
+
+bool FMidiPlayCursorTracker::RemoveCursor(FMidiPlayCursor* Cursor)
+{
+	if (ContainsCursor(Cursor))
+	{
+		Cursors.Remove(Cursor);
+		if ((Cursor->GetLookaheadType() == FMidiPlayCursor::ELookaheadType::Ticks &&
+			(Cursor->GetLookaheadTicks() < LatestCursorTick || Cursor->GetLookaheadTicks() > EarliestCursorTick)) ||
+			(Cursor->GetLookaheadMs() < LatestCursorMs || Cursor->GetLookaheadMs() > EarliestCursorMs))
+		{
+			RecalculateExtents();
+		}
+		return true;
+	}
+	return false;
+}
+
+void FMidiPlayCursorTracker::Clear()
+{
+	EarliestCursorTick = 0;
+	LatestCursorTick = 0;
+	EarliestCursorMs = 0.0f;
+	LatestCursorMs = 0.0f;
+	while (Cursors.PopTail());
+}
+
+void FMidiPlayCursorTracker::RecalculateExtents()
+{
+	EarliestCursorTick = 0;
+	LatestCursorTick = 0;
+	EarliestCursorMs = 0.0f;
+	LatestCursorMs = 0.0f;
+	for (auto it = Cursors.begin(); it != Cursors.end(); ++it)
+	{
+		FMidiPlayCursor* Cursor = it.GetNode();
+		if (Cursor->GetLookaheadType() == FMidiPlayCursor::ELookaheadType::Ticks)
+		{
+			int32 LookAheadTicks = Cursor->GetLookaheadTicks();
+			if (LookAheadTicks < 0 && LookAheadTicks < LatestCursorTick)
+			{
+				LatestCursorTick = LookAheadTicks;
+			}
+			else if (LookAheadTicks > 0 && LookAheadTicks > EarliestCursorTick)
+			{
+				EarliestCursorTick = LookAheadTicks;
+			}
+		}
+		else
+		{
+			float LookAheadMs = Cursor->GetLookaheadMs();
+			if (LookAheadMs < 0.0f && LookAheadMs < LatestCursorMs)
+			{
+				LatestCursorMs = LookAheadMs;
+			}
+			else if (LookAheadMs > 0 && LookAheadMs > EarliestCursorMs)
+			{
+				EarliestCursorMs = LookAheadMs;
+			}
+		}
+	}
+}
+
+int32 FMidiPlayCursorTracker::GetFarthestAheadCursorTick() const
+{
+	int32 Result = 0;
+	// This next line of ugliness is required because const iterating through a const linked list is broken
+	// and I don't have time to figure out what about the stack of templates is busted.
+	TIntrusiveDoubleLinkedList<FMidiPlayCursor>& NonConstCursorList = const_cast<TIntrusiveDoubleLinkedList<FMidiPlayCursor>&>(Cursors);
+	for (auto it = NonConstCursorList.begin(); it != NonConstCursorList.end(); ++it)
+	{
+		int32 Tick = it.GetNode()->GetCurrentTick();
+		if (Tick > Result)
+		{
+			Result = Tick;
+		}
+	}
+	return Result;
+}
+
+int32 FMidiPlayCursorTracker::GetFarthestBehindCursorTick() const
+{
+	int32 Result = 0x7FFFFFFF;
+	// This next line of ugliness is required because const iterating through a const linked list is broken
+	// and I don't have time to figure out what about the stack of templates is busted.
+	TIntrusiveDoubleLinkedList<FMidiPlayCursor>& NonConstCursorList = const_cast<TIntrusiveDoubleLinkedList<FMidiPlayCursor>&>(Cursors);
+	for (auto it = NonConstCursorList.begin(); it != NonConstCursorList.end(); ++it)
+	{
+		int32 Tick = it.GetNode()->GetCurrentTick();
+		if (Tick < Result)
+		{
+			Result = Tick;
+		}
+	}
+	return Result;
+}
+
+void FMidiPlayCursorTracker::MoveToLoopStart(int32 NewThruTick, float NewThruMs)
+{
+	CurrentTick = NewThruTick;
+	CurrentMs = NewThruMs;
+
+	for (auto it = Cursors.begin(); it != Cursors.end(); ++it)
+	{
+		it.GetNode()->MoveToLoopStartIfAtNotOffset(NewThruTick, NewThruMs);
+	}
+}
+
+void FMidiPlayCursorTracker::QueueReset(int32 NewCurrentTick, float NewCurrentMs, int32 PreRollStartTick, float PreRollStartMs, bool Broadcast)
+{
+	NewQueuedTick = NewCurrentTick;
+	NewQueuedMs = NewCurrentMs;
+	NewQueuedPreRollTick = PreRollStartTick;
+	NewQueuedPreRollMs = PreRollStartMs;
+	NewQueuedBroadcast = Broadcast;
+	HasQueuedReset = true;
+}
+
+void FMidiPlayCursorTracker::HandleQueuedReset()
+{
+	if (!HasQueuedReset)
+	{
+		return;
+	}
+	Reset(NewQueuedTick, NewQueuedMs, NewQueuedPreRollTick, NewQueuedPreRollMs, NewQueuedBroadcast);
+}
+
+void FMidiPlayCursorTracker::Reset(int32 NewCurrentTick, float NewCurrentMs, int32 PreRollStartTick, float PreRollStartMs, bool Broadcast)
+{
+	HasQueuedReset = false;
+	if (NewCurrentTick == PreRollStartTick)
+	{
+		Reset(NewCurrentTick, NewCurrentMs, Broadcast);
+		return;
+	}
+	CurrentTick = PreRollStartTick;
+	CurrentMs = PreRollStartMs;
+	LoopCount = 0;
+	for (auto it = Cursors.begin(); it != Cursors.end(); ++it)
+	{
+		it.GetNode()->Reset(!Broadcast);
+	}
+	CurrentTick = NewCurrentTick;
+	CurrentMs = NewCurrentMs;
+	for (auto it = Cursors.begin(); it != Cursors.end();)
+	{
+		if (!it.GetNode()->AdvanceAsPreRoll())
+		{
+			auto DeadIt = it;
+			++it;
+			Cursors.Remove(DeadIt.GetNode());
+			DeadIt.GetNode()->SetOwner(nullptr, nullptr);
+		}
+		else
+		{
+			++it;
+		}
+	}
+}
+
+void FMidiPlayCursorTracker::Reset(int32 NewCurrentTick, float NewCurrentMs, bool Broadcast)
+{
+	HasQueuedReset = false;
+	CurrentTick = NewCurrentTick;
+	CurrentMs = NewCurrentMs;
+	LoopCount = 0;
+	for (auto it = Cursors.begin(); it != Cursors.end(); ++it)
+	{
+		it.GetNode()->Reset(!Broadcast);
+	}
+}
+
+void FMidiPlayCursorTracker::ResetNewCursor(FMidiPlayCursor* Cursor, int32 PreRollStartTick, float PreRollStartMs, bool Broadcast)
+{
+	HasQueuedReset = false;
+	int32 CurrentTickSaver = CurrentTick;
+	float CurrentMsSaver = CurrentMs;
+	CurrentTick = PreRollStartTick;
+	CurrentMs = PreRollStartMs;
+	LoopCount = 0;
+	Cursor->Reset(!Broadcast);
+	CurrentTick = CurrentTickSaver;
+	CurrentMs = CurrentMsSaver;
+	Cursor->AdvanceAsPreRoll();
+}
+
