[Harmonix] Move Harmonix plugin into Engine/Plugins/Runtime to make it available to licensees in UE 5.4

diff --git a/Engine/Plugins/Runtime/Harmonix/Source/HarmonixMidi/Private/HarmonixMidi/MidiReader.cpp b/Engine/Plugins/Runtime/Harmonix/Source/HarmonixMidi/Private/HarmonixMidi/MidiReader.cpp
new file mode 100644
--- /dev/null
+++ b/Engine/Plugins/Runtime/Harmonix/Source/HarmonixMidi/Private/HarmonixMidi/MidiReader.cpp
@@ -0,0 +1,653 @@
+// Copyright Epic Games, Inc. All Rights Reserved.
+#include "HarmonixMidi/MidiReader.h"
+#include "HarmonixMidi/MidiConstants.h"
+#include "HarmonixMidi/MidiReceiver.h"
+#include "HarmonixMidi/VarLenNumber.h"
+#include "HarmonixMidi/MusicTimeSpecifier.h"
+#include "HarmonixMidi/BarMap.h"
+#include "HarmonixMidi/TempoMap.h"
+#include "Misc/Paths.h"
+
+// A 4-character IFF file chunk ID.  (Only "MThd" and "MTrk" are recognized
+// by the SMF file format.)
+const int32 kMidiChunkIDSize = 4;
+
+class FMidiChunkID
+{
+public:
+	FMidiChunkID(const char* InStr) { memcpy(Str, InStr, kMidiChunkIDSize); }
+	FMidiChunkID(FArchive& Archive) { Archive.Serialize(Str, 4); }
+
+	friend FArchive& operator<<(FArchive& Archive, FMidiChunkID& Id);
+
+	bool operator==(const FMidiChunkID& rhs) const
+	{
+		return !FCStringAnsi::Strncmp(Str, rhs.Str, kMidiChunkIDSize);
+	}
+
+	bool operator!=(const FMidiChunkID& rhs) const { return !operator==(rhs); }
+
+	static const FMidiChunkID kMThd; // "MThd": MIDI file header chunk
+	static const FMidiChunkID kMTrk; // "MTrk": MIDI track chunk
+
+private:
+	char Str[kMidiChunkIDSize]; // storage for the 4 bytes
+};
+
+const FMidiChunkID FMidiChunkID::kMThd("MThd");
+const FMidiChunkID	FMidiChunkID::kMTrk("MTrk");
+
+FArchive& operator<<(FArchive& Archive, FMidiChunkID& Id)
+{
+	Archive.Serialize(Id.Str, 4);
+	return Archive;
+}
+
+// Contains a chunk ID and a length.  The length is the number of bytes of
+// data in the chunk AFTER the header.
+
+class FMidiChunkHeader
+{
+public:
+	FMidiChunkHeader(FArchive& Archive)
+		: ID(Archive)
+		, Length(0)
+	{
+		Archive << Length;
+	}
+	FMidiChunkHeader(const FMidiChunkID& InId, uint32 InLength) :
+		ID(InId), Length(InLength) {}
+
+	// Access chunk-header information
+	uint32 GetLength() const { return Length; } // length of data following
+	const FMidiChunkID& GetID() const { return ID; } // the chunk's 4-byte ID
+
+	friend FArchive& operator<<(FArchive& Archive, FMidiChunkHeader& Header);
+
+private:
+	FMidiChunkID ID; // Chunk's 4-byte id
+	uint32 Length; // length of data part of this chunk (not including header)
+};
+
+FArchive& operator<<(FArchive& Archive, FMidiChunkHeader& Header)
+{
+	Archive << Header.ID;
+	Archive << Header.Length;
+	return Archive;
+}
+
+///////////////////////////////////////////////////////////////////////////////
+// FStdMidiFileReader
+FStdMidiFileReader::FStdMidiFileReader(void* Buffer, int32 BufferSize, const FString& InFileName, IMidiReceiver* InReceiver, int32 TicksPerQuarterNote, MidiConstants::EMidiTextEventEncoding InTextEncoding)
+	: Filename(InFileName)
+	, Receiver(InReceiver)
+	, TextEncoding(InTextEncoding)
+	, DestinationTicksPerQuarterNote(TicksPerQuarterNote)
+{
+	check(Receiver);
+	BufferArchive = MakeShared<FBufferReader>(Buffer, BufferSize, false);
+	InputArchive = BufferArchive;
+	Init();
+}
+
+FStdMidiFileReader::FStdMidiFileReader(const FString& FilePath, IMidiReceiver* InReceiver, int32 TicksPerQuarterNote, MidiConstants::EMidiTextEventEncoding InTextEncoding)
+	: Receiver(InReceiver)
+	, TextEncoding(InTextEncoding)
+	, DestinationTicksPerQuarterNote(TicksPerQuarterNote)
+{
+	check(Receiver);
+	Filename = FPaths::GetCleanFilename(FilePath);
+	IPlatformFile& PlatformFileApi = FPlatformFileManager::Get().GetPlatformFile();
+	IFileHandle* FileHandle = PlatformFileApi.OpenRead(*FilePath);
+	FileArchive = MakeShared<FArchiveFileReaderGeneric>(FileHandle, *Filename, FileHandle->Size());
+	InputArchive = FileArchive;
+	Init();
+}
+
+FStdMidiFileReader::FStdMidiFileReader(TSharedPtr<FArchive> Archive, const FString& InFilename, IMidiReceiver* InReceiver, int32 TicksPerQuarterNote, MidiConstants::EMidiTextEventEncoding InTextEncoding)
+	: Filename(InFilename)
+	, InputArchive(Archive)
+	, Receiver(InReceiver)
+	, TextEncoding(InTextEncoding)
+	, DestinationTicksPerQuarterNote(TicksPerQuarterNote)
+{
+	check(Receiver);
+	Init();
+}
+
+bool FStdMidiFileReader::ReadStdMidiFileForReceiver(void* Buffer, int32 BufferSize, const FString& FileName, IMidiReceiver* Receiver, int32 TicksPerQuarterNote, MidiConstants::EMidiTextEventEncoding InTextEncoding)
+{
+	bool Success = true;
+	FStdMidiFileReader Reader(Buffer, BufferSize, FileName, Receiver, TicksPerQuarterNote, InTextEncoding);
+
+	if (Reader.GetFailed())
+	{
+		UE_LOG(LogMidi, Error, TEXT("Unable to read MIDI file"));
+		Success = false;
+	}
+	else
+	{
+		Receiver->Reset();
+		Reader.ReadAllTracks();
+		Receiver->Finalize(Reader.LastTick);
+
+		if (Reader.GetFailed())
+		{
+			UE_LOG(LogMidi, Error, TEXT("opened but failed to parse MIDI file"));
+			Success = false;
+		}
+	}
+
+	return Success;
+}
+
+void FStdMidiFileReader::ReadAllTracks()
+{
+	// read from the top of the stream
+	InputArchive->Seek(0);
+
+	bool TrackReadSuccessfully = false;
+	do 
+	{
+		TrackReadSuccessfully = ReadTrack();
+	} while (!Failed && TrackReadSuccessfully);
+
+	Receiver->Finalize(LastTick);
+}
+
+bool FStdMidiFileReader::ReadSomeEvents(int32 NumEvents)
+{
+	for (int32 i = 0; i < NumEvents; ++i)
+	{
+		ReadNextEvent();
+		if ((State == EState::End) || Failed)
+		{
+			return true;
+		}
+	}
+	return false;
+}
+
+bool FStdMidiFileReader::ReadTrack()
+{
+	do
+	{
+		ReadNextEvent();
+	} while (State != EState::End && State != EState::NewTrack && !Failed);
+
+	return (State == EState::NewTrack);
+}
+
+bool FStdMidiFileReader::ReadEvents(int32 Count)
+{
+	for (int32 i = 0; i < Count; ++i)
+	{
+		if (State == EState::End || Failed)
+		{
+			return false;
+		}
+		ReadNextEvent();
+	}
+	return true;
+}
+
+// skips the processing of the current track and moves on to the next track.
+// Has no effect if called from within EndOfTrack().
+void FStdMidiFileReader::SkipCurrentTrack()
+{
+	if (State == EState::InTrack)
+	{
+		if (CurrentTrackIndex == NumTracks - 1)
+		{
+			State = EState::End;
+			Receiver->OnEndOfTrack(CurrentTick);
+			Receiver->OnAllTracksRead();
+		}
+		else
+		{
+			State = EState::NewTrack;
+			Receiver->OnEndOfTrack(CurrentTick);
+			InputArchive->Seek(TrackEndPos);
+		}
+	}
+}
+
+void FStdMidiFileReader::Init()
+{
+	BarMap = MakeShared<FBarMap>();
+	check(InputArchive);
+	InputArchive->ArForceByteSwapping = PLATFORM_LITTLE_ENDIAN == 1;
+	Receiver->SetMidiReader(this);
+}
+
+void FStdMidiFileReader::ReadNextEvent()
+{
+	ReadNextEventImpl();
+}
+
+void FStdMidiFileReader::ReadNextEventImpl()
+{
+	if (Failed)
+	{
+		return;
+	}
+
+	switch (State)
+	{
+	case EState::InTrack:
+		ReadEvent(*InputArchive);
+		break;
+	case EState::NewTrack:
+		ReadTrackHeader(*InputArchive);
+		break;
+	case EState::Start:
+		ReadFileHeader(*InputArchive);
+		break;
+	case EState::End:
+		break;
+	}
+}
+
+// Read the standard Midi file header chunk (MThd)
+void FStdMidiFileReader::ReadFileHeader(FArchive& Archive)
+{
+	check(State == EState::Start);
+
+	FMidiChunkHeader Header(Archive);
+
+	// Verify the MIDI header is properly formed.
+	ensureAlwaysMsgf(Header.GetID() == FMidiChunkID::kMThd && Header.GetLength() == 6, TEXT("%s: MIDI file header is corrupt"), *Filename);
+
+	// SMF format:
+	//   0: A single multi-channel track
+	//   1: 1 or more simultaneous tracks
+	//   2: 1 or more sequentially independent tracks
+	// we only support type 1.
+	int16 Format;
+	Archive << Format;
+	ensureMsgf(Format == 1 || Format == 0, TEXT("%s: Only type 0 or 1 MIDI files are supported; this file is type %d"), *Filename, Format);
+
+	Archive << NumTracks;
+
+	if (Format == 0)
+	{
+		// for midi type 0
+		Receiver->AddDummyConductorTrack();
+		++CurrentTrackIndex;
+		++NumTracks;  // we added a track
+	}
+
+	if (NumTracks <= 0)
+	{
+		UE_LOG(LogMidi, Error, TEXT("%s: MIDI file has no tracks"), *Filename);
+	}
+	else
+	{
+		Receiver->SetNumTracks(NumTracks);
+	}
+
+	// Hi bit set on ticksPerQuarter indicates a SMPTE division.
+	// We don't support that.
+	int16 TicksPerQuarter;
+	Archive << TicksPerQuarter;
+	if (TicksPerQuarter & 0x8000)
+	{
+		UE_LOG(LogMidi, Error, TEXT("%s: MIDI file uses SMPTE time division; this is not supported"), *Filename);
+	}
+	TickConversionFactor = float(DestinationTicksPerQuarterNote) / (float)TicksPerQuarter;
+
+	// If there are header failures, don't continue compilation. Just quit here!
+	if ((NumTracks == 0) || (Format != 1 && Format != 0) || (TicksPerQuarter & 0x8000))
+	{
+		Failed = true;
+	}
+	else
+	{
+		State = EState::NewTrack;
+	}
+}
+
+// Read the standard Midi track chunk header (MTrk)
+void FStdMidiFileReader::ReadTrackHeader(FArchive& Archive)
+{
+	check(State == EState::NewTrack);
+
+	FMidiChunkHeader TrackHeader(Archive);
+
+	if (TrackHeader.GetID() != FMidiChunkID::kMTrk)
+	{
+		UE_LOG(LogMidi, Error, TEXT("%s: MIDI track header for track %d is corrupt"), *Filename, CurrentTrackIndex);
+		Failed = true;
+		return;
+	}
+
+	TrackEndPos = int32(Archive.Tell() + TrackHeader.GetLength());
+
+	++CurrentTrackIndex; // update track num
+
+	// reset state for reading a new track
+	PrevStatus = 0;
+	CurrentFileTick = 0;
+	CurrentTick = 0;    // units of file's Tick
+	MidiListTick = -1;  // units of desired Tick (for sorting)
+	State = EState::InTrack;
+
+	// Let the midi receiver know that this is a new track.
+	Receiver->OnNewTrack(CurrentTrackIndex);
+}
+
+// Read an event from the track
+void FStdMidiFileReader::ReadEvent(FArchive& Archive)
+{
+	check(State == EState::InTrack);
+
+	uint8 Status, Data1 = 0;    // MIDI message bytes
+	bool RunningStatus;     // true if running status is detected
+
+	CurrentFileTick += Midi::VarLenNumber::Read(Archive); // accumulate tick
+	CurrentTick = int32(CurrentFileTick * TickConversionFactor);  // convert to kTicksPerQuarterNoteInt
+
+	if (CurrentTick > LastTick)
+	{
+		LastTick = CurrentTick;
+	}
+
+	if (CurrentTick != MidiListTick)
+	{
+		// we have a new tick; sort all midi from previous tick and send it out
+		ProcessMidiList();
+
+		if (State != EState::InTrack)
+		{
+			return;
+		}
+
+		MidiListTick = CurrentTick;
+	}
+
+	Archive << Status;
+	if (MidiConstants::IsStatus(Status))  // this byte is truly a status byte
+	{
+		RunningStatus = false;
+		if (!MidiConstants::IsSystem(Status))
+		{
+			// save the current status byte, in case the next message uses running
+			// status.  Note that this is only for midi events; running status can
+			// run across meta-events and system exclusive events.
+			PrevStatus = Status;
+		}
+	}
+	else
+	{
+		// whoops, that wasn't actually a status byte we read, but a data byte;
+		// the status byte comes from running status.
+		RunningStatus = true;
+		Data1 = Status;
+		Status = PrevStatus;
+	}
+
+	// Check if this is a system (ie, non-channel) message
+	if (MidiConstants::IsSystem(Status))
+	{
+		ReadSystemEvent(CurrentTick, Status, Archive);
+	}                   // All other messages are channel (non-system) messages
+	else
+	{
+		// Running status means that data1 has already been read above
+		if (!RunningStatus)
+		{
+			Archive << Data1;
+		}
+		ReadMidiEvent(CurrentTick, Status, Data1, Archive);
+	}
+}
+
+void FStdMidiFileReader::ReadMidiEvent(int32 Tick, uint8 Status, uint8 Data1, FArchive& Archive)
+{
+	uint8 Data2 = 0;
+	bool ValidMidiEvent = false;
+	switch (Status & MidiConstants::kMessageTypeMask)
+	{
+		// special case for NoteOn (3 byte message)
+	case MidiConstants::kNoteOn:
+		Archive << Data2;
+		ValidMidiEvent = true;
+		// convert note-on vel=0 -> note offs
+		if (Data2 == 0)
+		{
+			Status = MidiConstants::kNoteOff | (Status & MidiConstants::kChannelMask);
+		}
+		break;
+
+		// other 3 byte messages
+	case MidiConstants::kNoteOff:
+	case MidiConstants::kControl:
+	case MidiConstants::kPitch:
+	case MidiConstants::kPolyPres:
+		Archive << Data2;
+		ValidMidiEvent = true;
+		break;
+
+		// 2 byte messages
+	case MidiConstants::kProgram:
+	case MidiConstants::kChanPres:
+		Data2 = 0;
+		ValidMidiEvent = true;
+		break;
+
+	default:
+		UE_LOG(LogMidi, Error, TEXT("%s (%s): Cannot parse event %i"), *Filename, *CurrentTrackName, (Status & MidiConstants::kMessageTypeMask));
+		break;
+	}
+
+	if (ValidMidiEvent)
+	{
+		// At this time, we only care about Note On / Note Off events.
+		QueueChannelMsg(Tick, Status, Data1, Data2);
+	}
+}
+
+// Read system event from the stream.  Status byte has already been read.
+void FStdMidiFileReader::ReadSystemEvent(int32 Tick, uint8 Status, FArchive& Archive)
+{
+	switch (Status)
+	{
+	case MidiConstants::kFile_SysEx:
+	case MidiConstants::kFile_Escape:
+		{
+			// All system events start with length. We'll
+			// seek past that length and ignore these events
+			// entirely.
+			int32 PacketLength = Midi::VarLenNumber::Read(Archive);
+			Archive.Seek(Archive.Tell() + PacketLength);
+		}
+		break;
+	case MidiConstants::kFile_Meta:
+		{
+			uint8 type;
+			Archive << type;
+			ReadMetaEvent(Tick, type, Archive);
+		}
+		break;
+	default:
+		UE_LOG(LogMidi, Error, TEXT("%s (%s): Cannot parse system event %i"), *Filename, *CurrentTrackName, Status);
+		break;
+	}
+}
+
+// Read a Meta Event from stream.
+// Status byte and event type have already been read.
+void FStdMidiFileReader::ReadMetaEvent(int32 Tick, uint8 Type, FArchive& Archive)
+{
+	int32 Length = Midi::VarLenNumber::Read(Archive);
+	int32 StartPos = int32(Archive.Tell());
+
+	switch (Type)
+	{
+	case MidiConstants::kMeta_Text:
+	case MidiConstants::kMeta_Copyright:
+	case MidiConstants::kMeta_TrackName:
+	case MidiConstants::kMeta_Lyric:
+		{
+			unsigned char* AsUtf8;
+			if (TextEncoding == MidiConstants::EMidiTextEventEncoding::Latin1)
+			{
+				unsigned char* AsLatin = (unsigned char*)FMemory::Malloc(Length + 1);
+				Archive.Serialize(AsLatin, Length);
+				AsLatin[Length] = 0;
+				AsUtf8 = (unsigned char*)FMemory::Malloc(Length * 2 + 2);
+
+				unsigned char* in = AsLatin;
+				unsigned char* out = AsUtf8;
+				while (*in)
+				{
+					if (*in < 128)
+					{
+						*out++ = *in++;
+					}
+					else
+					{
+						*out++ = 0xc0 | *in >> 6;
+						*out++ = 0x80 | (*in++ & 0x3f);
+					}
+				}
+				*out = 0;
+
+				FMemory::Free(AsLatin);
+			}
+			else
+			{
+				AsUtf8 = (unsigned char*)FMemory::Malloc(Length + 1);
+				Archive.Serialize(AsUtf8, Length);
+				AsUtf8[Length] = 0;
+			}
+
+			FString AsFString = UTF8_TO_TCHAR(AsUtf8);
+
+			if (Type == MidiConstants::kMeta_TrackName)
+			{
+				// Save ourselves a copy for better error messaging...
+				CurrentTrackName = AsFString;
+			}
+
+			Receiver->OnText(Tick, AsFString, Type);
+			FMemory::Free(AsUtf8);
+		}
+		break;
+
+	case MidiConstants::kMeta_Tempo:
+		{
+			uint8 Byte1, Byte2, Byte3;
+			Archive << Byte1 << Byte2 << Byte3;
+			int32 Tempo = (Byte1 << 16) + (Byte2 << 8) + Byte3;
+			Receiver->OnTempo(Tick, Tempo);
+		}
+		break;
+
+	case MidiConstants::kMeta_EndOfTrack:
+		ProcessMidiList();  // in case there is anything left to send out.
+		Receiver->OnEndOfTrack(Tick);
+		if (CurrentTrackIndex == NumTracks - 1)
+		{
+			State = EState::End;
+			Receiver->OnAllTracksRead();
+		}
+		else
+		{
+			State = EState::NewTrack;
+		}
+		break;
+
+	case MidiConstants::kMeta_TimeSig:
+		{
+			uint8 Numerator, DenominatorExp;
+			Archive << Numerator << DenominatorExp;
+
+			if (DenominatorExp > 6)
+			{
+				UE_LOG(LogMidi, Error, TEXT("%s (%s): Time signature at %s has invalid denominator (2^%d); max is 64 (2^6)"), *Filename, *CurrentTrackName, *MidiTickFormat(Tick, BarMap.Get(), Midi::EMusicTimeStringFormat::Position), DenominatorExp);
+
+				// Continuing past here could lead to an integer overflow below.
+				break;
+			}
+
+			// time signature denominator is given as a power of 2
+			//  e.g. for 6/8 time, numerator=6 and denominator=3
+			int32 Denominator = int32(FMath::Pow(2.0f, DenominatorExp));
+
+			if (Numerator == 0)
+			{
+				UE_LOG(LogMidi, Error, TEXT("%s (%s): Time signature %d/%d at %s has invalid numerator (%d)"), *Filename, *CurrentTrackName, Numerator, Denominator,
+					*MidiTickFormat(Tick, BarMap.Get(), Midi::EMusicTimeStringFormat::Position), Numerator);
+			}
+
+			// before telling any receivers about the time signature we update our own local tempo map
+			// so we can print better error and warning messages...
+			FMusicTimestamp Timestamp;
+			check(Tick == 0 || BarMap->GetNumTimeSignaturePoints() > 0);
+			int32 BarIndex = BarMap->TickToBarIncludingCountIn(Tick);
+			if (!BarMap->AddTimeSignatureAtBarIncludingCountIn(BarIndex, Numerator, Denominator, true, false))
+			{
+				UE_LOG(LogMidi, Warning, TEXT("%s (%s): Time signature %d/%d at %s overlaps or conflicts with nearby time signatures"),
+					*Filename, *CurrentTrackName, Numerator, Denominator, *MidiTickFormat(Tick, BarMap.Get(), Midi::EMusicTimeStringFormat::Position));
+			}
+
+			// fill in our bar map
+			Receiver->OnTimeSignature(Tick, Numerator, Denominator, false);
+
+			// skip over next two bytes, "MIDI clocks per metronome click"
+			// and "32nd notes per MIDI quarter note"; they're not too useful
+			Archive.Seek(Archive.Tell() + 2);
+		}
+		break;
+
+	case MidiConstants::kMeta_ChannelPrefix:
+	case MidiConstants::kMeta_Port:
+	case MidiConstants::kMeta_KeySig:
+	case MidiConstants::kMeta_SMPTE:
+	case MidiConstants::kMeta_InstrumentName:
+	case MidiConstants::kMeta_Marker:
+	case MidiConstants::kMeta_CuePoint:
+		// these are valid events, but not currently supported by
+		// MidiReceiver; do nothing
+		break;
+
+	default:
+		UE_LOG(LogMidi, Warning, TEXT("%s (%s): Cannot parse meta event %i"), *Filename, *CurrentTrackName, Type);
+		break;
+	}
+
+	// 
+	if (State != EState::InTrack)
+	{
+		return;
+	}
+
+	// skip to end of meta event
+	Archive.Seek(StartPos + Length);
+}
+
+// Queue all midi that is on the same tick, for sorting. 
+// Note that tick is already in terms of desired ticks per quarter note.
+void FStdMidiFileReader::QueueChannelMsg(int32 Tick, uint8 Status, uint8 Data1, uint8 Data2)
+{
+	// stick this message in the MidiList
+	EventsOnSameTick.Add({Status, Data1, Data2});
+}
+
+// Check if MidiList has anything in it. If so, sort and send out.
+void FStdMidiFileReader::ProcessMidiList()
+{
+	EventsOnSameTick.StableSort(RawMidiLess());
+	for (auto& MidiItem : EventsOnSameTick)
+	{
+		Receiver->OnMidiMessage(MidiListTick, MidiItem.Status, MidiItem.Data1, MidiItem.Data2);
+
+		// state changed because SkipCurrentTrack was called. Stop processing.
+		if (State != EState::InTrack)
+		{
+			break;
+		}
+	}
+	EventsOnSameTick.Empty();
+}
diff --git a/Engine/Plugins/Runtime/Harmonix/Source/HarmonixMidi/Private/HarmonixMidi/MidiSongPos.cpp b/Engine/Plugins/Runtime/Harmonix/Source/HarmonixMidi/Private/HarmonixMidi/MidiSongPos.cpp
new file mode 100644
--- /dev/null
+++ b/Engine/Plugins/Runtime/Harmonix/Source/HarmonixMidi/Private/HarmonixMidi/MidiSongPos.cpp
@@ -0,0 +1,133 @@
+// Copyright Epic Games, Inc. All Rights Reserved.
+#include "HarmonixMidi/MidiSongPos.h"
+#include "HarmonixMidi/BarMap.h"
+#include "HarmonixMidi/MusicTimeSpecifier.h"
+#include "HarmonixMidi/MidiConstants.h"
+#include "HarmonixMidi/SongMaps.h"
+
+// -1 = less than, 1 = greater than, 0 = equal to
+int32 SongPosCmp(const FMidiSongPos& lhs, const FMidiSongPos& rhs)
+{
+	// can't use bar alone because of floating point. 
+	// so use bar first and then "rounded" beat.
+		
+	// bar
+	if (lhs.Timestamp.Bar < rhs.Timestamp.Bar)
+	{
+		return -1;
+	}
+	else if (lhs.Timestamp.Bar > rhs.Timestamp.Bar)
+	{
+		return 1;
+	}
+
+	// beat
+	if (MidiConstants::RoundToStandardBeatPrecision(lhs.Timestamp.Beat, lhs.TimeSigDenominator) ==
+		MidiConstants::RoundToStandardBeatPrecision(rhs.Timestamp.Beat, rhs.TimeSigDenominator))
+	{
+		return 0;
+	}
+	if (lhs.Timestamp.Beat < rhs.Timestamp.Beat)
+	{
+		return -1;
+	}
+	return 1;
+}
+
+// operators
+bool FMidiSongPos::operator<(const FMidiSongPos& rhs) const
+{
+	return SongPosCmp(*this, rhs) == -1;
+}
+
+bool FMidiSongPos::operator<=(const FMidiSongPos& rhs) const
+{
+	int32 cmp = SongPosCmp(*this, rhs);
+	return cmp == -1 || cmp == 0;
+}
+
+bool FMidiSongPos::operator>(const FMidiSongPos& rhs) const
+{
+	return SongPosCmp(*this, rhs) == 1;
+}
+
+bool FMidiSongPos::operator>=(const FMidiSongPos& rhs) const
+{
+	int32 cmp = SongPosCmp(*this, rhs);
+	return cmp == 1 || cmp == 0;
+}
+
+bool FMidiSongPos::operator==(const FMidiSongPos& rhs) const
+{
+	return SongPosCmp(*this, rhs) == 0;
+}
+
+void FMidiSongPos::SetByTime(float InElapsedMs, float InBpm, int32 InTimeSigNumerator, int32 InTimeSigDenominator, int32 StartBar)
+{
+	// first some numbers we will need...
+	float QuarterNotesPerSecond = InBpm / 60.0f;
+	float BeatsPerSecond       = QuarterNotesPerSecond * (InTimeSigDenominator / 4);
+	float TotalBeats           = BeatsPerSecond * (InElapsedMs / 1000.0f);
+	float BeatsPerBar          = (float)InTimeSigNumerator; // simple assumption given lack of SongMaps!
+	int32 BeatsOfCountIn       = (1 - StartBar) * InTimeSigNumerator;
+	float CountInSeconds       = (float)BeatsOfCountIn / BeatsPerSecond;
+
+	SecondsIncludingCountIn = InElapsedMs / 1000.0f;
+	SecondsFromBarOne       = SecondsIncludingCountIn - CountInSeconds;
+	TimeSigNumerator        = InTimeSigNumerator;
+	TimeSigDenominator      = InTimeSigDenominator;
+	Tempo                   = InBpm;
+	BarsIncludingCountIn    = TotalBeats / TimeSigNumerator;
+	BeatsIncludingCountIn   = TotalBeats;
+	Timestamp.Bar = FMath::FloorToInt32(BarsIncludingCountIn) + StartBar;
+	Timestamp.Beat = FMath::Fractional(BarsIncludingCountIn) * InTimeSigNumerator + 1.0f;
+}
+
+void FMidiSongPos::SetByTime(float InMs, const FSongMaps& Maps)
+{
+	float Tick = Maps.GetTempoMap().MsToTick(InMs);
+	SetByTick(Tick, Maps);
+}
+
+void FMidiSongPos::SetByTick(float InTick, const FSongMaps& Maps)
+{
+	const FTimeSignature* TimeSig = Maps.GetTimeSignatureAtTick(InTick);
+	float TickMs			= Maps.TickToMs(InTick);
+	SecondsIncludingCountIn = TickMs / 1000.0f;
+	SecondsFromBarOne       = SecondsIncludingCountIn - Maps.GetCountInSeconds();
+	TimeSigNumerator        = TimeSig ? TimeSig->Numerator : 4;
+	TimeSigDenominator      = TimeSig ? TimeSig->Denominator : 4;
+	Tempo                   = Maps.GetTempoAtTick(InTick);
+	BarsIncludingCountIn    = Maps.GetBarIncludingCountInAtTick(InTick);
+
+	const FSongSection* SongSection = !Maps.GetSectionMap().IsEmpty() ? Maps.GetSectionAtMs(TickMs) : nullptr;
+	CurrentSongSection		= SongSection ? FSongSection(SongSection->Name, SongSection->StartTick, SongSection->LengthTicks) : FSongSection();
+	
+	const FBeatMap& BeatMap = Maps.GetBeatMap();
+	int32 BeatPointIndex     = BeatMap.GetPointIndexForTick(FMath::FloorToInt32(InTick));
+	if (BeatPointIndex >= 0)
+	{
+		const FBeatMapPoint& BeatPoint = BeatMap.GetBeatPointInfo(BeatPointIndex);
+		float TickInBeat     = InTick - (float)BeatPoint.StartTick;
+		float FractionalPart = TickInBeat / (float)BeatPoint.LengthTicks;
+		BeatsIncludingCountIn     = (float)BeatPointIndex + FractionalPart;
+		BeatType             = BeatPoint.Type;
+	}
+	else
+	{
+		// use basic time signature based info from bar map...
+		const FBarMap& BarMap = Maps.GetBarMap();
+		int32 BarPointIndex = BarMap.GetPointIndexForTick(FMath::FloorToInt32(InTick));
+		if (BarPointIndex == INDEX_NONE)
+		{
+			ensure(InTick < 0.0f);
+			BarPointIndex = 0;
+		}
+		const FTimeSignaturePoint& TimeSigPoint = BarMap.GetTimeSignaturePoint(BarPointIndex);
+		float BarAtTimeSig = BarsIncludingCountIn - (float)TimeSigPoint.BarIndex;
+		float FractPart = FMath::Fractional(BarAtTimeSig);
+		BeatsIncludingCountIn = TimeSigPoint.BeatIndex + (BarAtTimeSig * TimeSigPoint.TimeSignature.Numerator);
+		BeatType = FMath::IsNearlyEqual(Timestamp.Beat, 1.0f) ? EMusicalBeatType::Downbeat : EMusicalBeatType::Normal;
+	}
+	Timestamp = Maps.GetBarMap().TickToMusicTimestamp(InTick);
+}
diff --git a/Engine/Plugins/Runtime/Harmonix/Source/HarmonixMidi/Private/HarmonixMidi/MusicTimeSpan.cpp b/Engine/Plugins/Runtime/Harmonix/Source/HarmonixMidi/Private/HarmonixMidi/MusicTimeSpan.cpp
new file mode 100644
--- /dev/null
+++ b/Engine/Plugins/Runtime/Harmonix/Source/HarmonixMidi/Private/HarmonixMidi/MusicTimeSpan.cpp
@@ -0,0 +1,214 @@
+// Copyright Epic Games, Inc. All Rights Reserved.
+
+#include "HarmonixMidi/MusicTimeSpan.h"
+#include "HarmonixMidi/SongMaps.h"
+#include "HarmonixMidi/MidiSongPos.h"
+#include "HarmonixMidi/MidiConstants.h"
+
+float FMusicalTimeSpan::CalcPositionInSpan(const FMidiSongPos& Position, const FSongMaps& Maps) const
+{
+	if (Position.TimeSigDenominator == 0)
+	{
+		return 0.0f;
+	}
+
+	if (Offset != 0)
+	{
+		return CalcPositionInSpanWithOffset(Position, Maps);
+	}
+
+	return CalcPositionInSpanNoOffset(Position, Maps);
+}
+
+float FMusicalTimeSpan::CalcPositionInSpanWithOffset(const FMidiSongPos& Position, const FSongMaps& Maps) const
+{
+	FMidiSongPos OffsetPosition;
+
+	int32 UnadjustedTick = Maps.MsToTick(Position.SecondsIncludingCountIn * 1000.0f);
+	int32 AdjustedTick = -1;
+	switch (OffsetUnits)
+	{
+	case EMusicTimeSpanOffsetUnits::Ms:
+		{
+			OffsetPosition.SetByTime(Position.SecondsIncludingCountIn * 1000.0f - Offset, Maps);
+		}
+		break;
+	case EMusicTimeSpanOffsetUnits::Bars:
+		{
+			float FractionalBar = Position.BarsIncludingCountIn + Maps.GetBarMap().GetStartBar() - Offset;
+			int32 Tick = Maps.GetBarMap().FractionalBarIncludingCountInToTick(FractionalBar);
+			OffsetPosition.SetByTick((float)Tick, Maps);
+		}
+		break;
+	case EMusicTimeSpanOffsetUnits::Beats:
+		{
+			float Beat = Position.BeatsIncludingCountIn - Offset;
+			OffsetPosition.SetByTime(Maps.GetMsAtBeat(Beat), Maps);
+		}
+		break;
+	case EMusicTimeSpanOffsetUnits::ThirtySecondNotes:
+		AdjustedTick = UnadjustedTick - MidiConstants::kTicksPerQuarterNoteInt * Offset / 8;
+		break;
+	case EMusicTimeSpanOffsetUnits::SixteenthNotes:
+		AdjustedTick = UnadjustedTick - MidiConstants::kTicksPerQuarterNoteInt * Offset / 4;
+		break;
+	case EMusicTimeSpanOffsetUnits::EighthNotes:
+		AdjustedTick = UnadjustedTick - MidiConstants::kTicksPerQuarterNoteInt * Offset / 2;
+		break;
+	case EMusicTimeSpanOffsetUnits::QuarterNotes:
+		AdjustedTick = UnadjustedTick - MidiConstants::kTicksPerQuarterNoteInt * Offset;
+		break;
+	case EMusicTimeSpanOffsetUnits::HalfNotes:
+		AdjustedTick = UnadjustedTick - MidiConstants::kTicksPerQuarterNoteInt * Offset * 2;
+		break;
+	case EMusicTimeSpanOffsetUnits::WholeNotes:
+		AdjustedTick = UnadjustedTick - MidiConstants::kTicksPerQuarterNoteInt * Offset * 4;
+		break;
+	case EMusicTimeSpanOffsetUnits::DottedSixteenthNotes:
+		AdjustedTick = UnadjustedTick - MidiConstants::kTicksPerQuarterNoteInt * Offset * 3 / 8;
+		break;
+	case EMusicTimeSpanOffsetUnits::DottedEighthNotes:
+		AdjustedTick = UnadjustedTick - MidiConstants::kTicksPerQuarterNoteInt * Offset * 3 / 4;
+		break;
+	case EMusicTimeSpanOffsetUnits::DottedQuarterNotes:
+		AdjustedTick = UnadjustedTick - MidiConstants::kTicksPerQuarterNoteInt * Offset * 3 / 2;
+		break;
+	case EMusicTimeSpanOffsetUnits::DottedHalfNotes:
+		AdjustedTick = UnadjustedTick - MidiConstants::kTicksPerQuarterNoteInt * Offset * 3;
+		break;
+	case EMusicTimeSpanOffsetUnits::DottedWholeNotes:
+		AdjustedTick = UnadjustedTick - MidiConstants::kTicksPerQuarterNoteInt * Offset * 6;
+		break;
+	case EMusicTimeSpanOffsetUnits::SixteenthNoteTriplets:
+		AdjustedTick = UnadjustedTick - MidiConstants::kTicksPerQuarterNoteInt * Offset / 6;
+		break;
+	case EMusicTimeSpanOffsetUnits::EighthNoteTriplets:
+		AdjustedTick = UnadjustedTick - MidiConstants::kTicksPerQuarterNoteInt * Offset / 3;
+		break;
+	case EMusicTimeSpanOffsetUnits::QuarterNoteTriplets:
+		AdjustedTick = UnadjustedTick - MidiConstants::kTicksPerQuarterNoteInt * Offset * 2 / 3;
+		break;
+	case EMusicTimeSpanOffsetUnits::HalfNoteTriplets:
+		AdjustedTick = UnadjustedTick - MidiConstants::kTicksPerQuarterNoteInt * Offset * 4 / 3;
+		break;
+	default:
+		checkNoEntry();
+		break;
+	}
+
+	if (AdjustedTick > 0)
+	{
+		OffsetPosition.SetByTick(AdjustedTick, Maps);
+	}
+
+	return CalcPositionInSpanNoOffset(OffsetPosition, Maps);
+}
+
+float FMusicalTimeSpan::CalcPositionInSpanNoOffset(const FMidiSongPos& Position, const FSongMaps& Maps) const
+{
+	if (LengthUnits == EMusicTimeSpanLengthUnits::Bars || LengthUnits == EMusicTimeSpanLengthUnits::Beats)
+	{
+		return CalculateEnclosingVariableSizeSpanExtents(Position, Maps);
+	}
+	return CalculateEnclosingFixedSizeSpanExtents(Position, Maps);
+}
+
+
+float FMusicalTimeSpan::CalcPositionInSpan(float Ms, const FSongMaps& Maps) const
+{
+	FMidiSongPos Position;
+	Position.SetByTime(Ms, Maps);
+	return CalcPositionInSpan(Position, Maps);
+}
+
+float FMusicalTimeSpan::CalculateEnclosingVariableSizeSpanExtents(const FMidiSongPos& Position, const FSongMaps& Maps) const
+{
+	switch (LengthUnits)
+	{
+	case EMusicTimeSpanLengthUnits::Bars:
+		{
+			float Progress = FMath::Fractional(Position.BarsIncludingCountIn / Length);
+			if (Progress < 0.0f)
+			{
+				Progress += 1.0f;
+			}
+			return Progress;
+		}
+		break;
+	case EMusicTimeSpanLengthUnits::Beats:
+		{
+			float Progress = FMath::Fractional(Position.BeatsIncludingCountIn / Length);
+			if (Progress < 0.0f)
+			{
+				Progress += 1.0f;
+			}
+			return Progress;
+	}
+		break;
+	default:
+		checkNoEntry();
+		break;
+	}
+	return 0.0f;
+}
+
+float FMusicalTimeSpan::CalculateEnclosingFixedSizeSpanExtents(const FMidiSongPos& Position, const FSongMaps& Maps) const
+{
+	int32 Tick = Maps.MsToTick(Position.SecondsIncludingCountIn * 1000.0f);
+	int32 GridUnitTicks = 1;
+	switch (LengthUnits)
+	{
+	case EMusicTimeSpanLengthUnits::ThirtySecondNotes:
+		GridUnitTicks = Maps.GetTicksPerQuarterNote() / 8;
+		break;
+	case EMusicTimeSpanLengthUnits::SixteenthNotes:
+		GridUnitTicks = Maps.GetTicksPerQuarterNote() / 4;
+		break;
+	case EMusicTimeSpanLengthUnits::EighthNotes:
+		GridUnitTicks = Maps.GetTicksPerQuarterNote() / 2;
+		break;
+	case EMusicTimeSpanLengthUnits::QuarterNotes:
+		GridUnitTicks = Maps.GetTicksPerQuarterNote();
+		break;
+	case EMusicTimeSpanLengthUnits::HalfNotes:
+		GridUnitTicks = Maps.GetTicksPerQuarterNote() * 2;
+		break;
+	case EMusicTimeSpanLengthUnits::WholeNotes:
+		GridUnitTicks = Maps.GetTicksPerQuarterNote() * 4;
+		break;
+	case EMusicTimeSpanLengthUnits::DottedSixteenthNotes:
+		GridUnitTicks = Maps.GetTicksPerQuarterNote() * 3 / 8;
+		break;
+	case EMusicTimeSpanLengthUnits::DottedEighthNotes:
+		GridUnitTicks = Maps.GetTicksPerQuarterNote() * 3 / 4;
+		break;
+	case EMusicTimeSpanLengthUnits::DottedQuarterNotes:
+		GridUnitTicks = Maps.GetTicksPerQuarterNote() * 3 / 2;
+		break;
+	case EMusicTimeSpanLengthUnits::DottedHalfNotes:
+		GridUnitTicks = Maps.GetTicksPerQuarterNote() * 3;
+		break;
+	case EMusicTimeSpanLengthUnits::DottedWholeNotes:
+		GridUnitTicks = Maps.GetTicksPerQuarterNote() * 6;
+		break;
+	case EMusicTimeSpanLengthUnits::SixteenthNoteTriplets:
+		GridUnitTicks = Maps.GetTicksPerQuarterNote() / 6;
+		break;
+	case EMusicTimeSpanLengthUnits::EighthNoteTriplets:
+		GridUnitTicks = Maps.GetTicksPerQuarterNote() / 3;
+		break;
+	case EMusicTimeSpanLengthUnits::QuarterNoteTriplets:
+		GridUnitTicks = Maps.GetTicksPerQuarterNote() * 2 / 3;
+		break;
+	case EMusicTimeSpanLengthUnits::HalfNoteTriplets:
+		GridUnitTicks = Maps.GetTicksPerQuarterNote() * 4 / 3;
+		break;
+	default:
+		checkNoEntry();
+		break;
+	}
+	int32 CellSizeTicks = GridUnitTicks * Length;
+	float InCell = (float) Tick / CellSizeTicks;
+	float Progress = FMath::Fractional(InCell);
+	return Progress < 0.0f ? 1.0f + Progress : Progress;
+}
diff --git a/Engine/Plugins/Runtime/Harmonix/Source/HarmonixMidi/Private/HarmonixMidi/SmoothedMidiPlayCursor.cpp b/Engine/Plugins/Runtime/Harmonix/Source/HarmonixMidi/Private/HarmonixMidi/SmoothedMidiPlayCursor.cpp
new file mode 100644
--- /dev/null
+++ b/Engine/Plugins/Runtime/Harmonix/Source/HarmonixMidi/Private/HarmonixMidi/SmoothedMidiPlayCursor.cpp
@@ -0,0 +1,166 @@
+// Copyright Epic Games, Inc. All Rights Reserved.
+#include "HarmonixMidi/SmoothedMidiPlayCursor.h"
+#include "Logging/LogMacros.h"
+
+DEFINE_LOG_CATEGORY_STATIC(LogSmoothedMidiPlayCursor, Log, Log);
+
+namespace SmoothedPlayCursor
+{
+	static constexpr float kMassiveJump = 60.0f;
+	static constexpr float kMinorErrorThreshold = 5.0f;
+	static constexpr float kMinorCorrectionFactor = 0.002;
+	static constexpr float kMajorErrorThreshold = 20.0f;
+	static constexpr float kMajorCorrectionFactor = 2.0 * kMinorCorrectionFactor;
+}
+
+FSmoothedMidiPlayCursor::FSmoothedMidiPlayCursor()
+{
+	SetupMsLookahead(-SmoothingLatencyMs, FMidiPlayCursor::ESyncOptions::NoBroadcastNoPreRoll);
+	SmoothingTimer.SetSpeed(1.0f);
+	SmoothingTimer.Start();
+	SmoothingTimer.Reset(0.0f);
+}
+
+void FSmoothedMidiPlayCursor::SetSpeed(float NewSpeed)
+{
+	if (BaseSpeed == NewSpeed)
+		return;
+	BaseSpeed = NewSpeed;
+	SetupMsLookahead(-SmoothingLatencyMs * BaseSpeed, FMidiPlayCursor::ESyncOptions::NoBroadcastNoPreRoll);
+	if (Tracker)
+		Tracker->RecalculateExtents();
+	SmoothingTimer.SetSpeed(BaseSpeed);
+}
+
+void FSmoothedMidiPlayCursor::Reset(bool ForceNoBroadcast)
+{
+	FMidiPlayCursor::Reset(ForceNoBroadcast);
+
+	ErrorTracker.Reset();
+	SmoothingLatencyMs = 30.0f;
+	SmoothedMs = 0.0f;
+	SmoothedTick = 0.0f;
+	LoopedThisPass = false;
+	BaseSpeed = 1.0f;
+	CurrentSongPos.Reset();
+}
+
+void FSmoothedMidiPlayCursor::AdvanceByTicks(bool processLoops /*= true*/, bool broadcast /*= true*/, bool isPreRoll /*= false*/)
+{
+	FMidiPlayCursor::AdvanceByTicks(processLoops, broadcast, isPreRoll);
+	SyncSmoothingTimer();
+	UpdateSongPosition();
+}
+
+void FSmoothedMidiPlayCursor::AdvanceByMs(bool processLoops /*= true*/, bool broadcast /*= true*/, bool isPreRoll /*= false*/)
+{
+	FMidiPlayCursor::AdvanceByMs(processLoops, broadcast, isPreRoll);
+	SyncSmoothingTimer();
+	UpdateSongPosition();
+}
+
+void FSmoothedMidiPlayCursor::OnLoop(int loopStartTick, int loopEndTick)
+{
+	LoopedThisPass = true;
+}
+
+void FSmoothedMidiPlayCursor::SyncSmoothingTimer()
+{
+	if (!Owner)
+	{
+		return;
+	}
+
+	if (!SmoothingTimer.IsRunning())
+	{
+		SmoothingTimer.Start();
+	}
+
+	SetSpeed(Owner->GetCurrentAdvanceRate());
+
+	float RawMs = CurrentMs;
+	float SmoothMs = SmoothingTimer.Ms();
+	float Error = (SmoothMs - RawMs) * (1.0f / (float)SmoothingTimer.GetSpeed());
+	ErrorTracker.Push(Error);
+	float MinRecentError = ErrorTracker.Min();
+	float AbsError = FMath::Abs(MinRecentError);
+
+	// if we have a massive jump just slam the thing to the current "raw" time and be done with it...
+	if (FMath::Abs(Error) > SmoothedPlayCursor::kMassiveJump)
+	{
+		UE_LOG(LogSmoothedMidiPlayCursor, Verbose, TEXT("Smoothing: Massive Error = %fms. Slamming."), Error);
+		SmoothingTimer.SetSpeed(BaseSpeed);
+		SmoothingTimer.Reset(RawMs);
+		SmoothedMs = RawMs;
+		SmoothedMsDelta = 0.0f;
+		SmoothedTick = Owner->GetTempoMap().MsToTick(SmoothedMs);
+		ErrorTracker.Reset();
+		return;
+	}
+
+	if (AbsError < SmoothedPlayCursor::kMinorErrorThreshold)
+	{
+		if (SmoothingTimer.GetSpeed() != BaseSpeed)
+		{
+			UE_LOG(LogSmoothedMidiPlayCursor, Verbose, TEXT("Smoothing: Error = %fms. Running normal speed."), MinRecentError);
+			SmoothingTimer.SetSpeed(BaseSpeed);
+		}
+	}
+	else if (AbsError < SmoothedPlayCursor::kMajorErrorThreshold)
+	{
+		float NewSpeedFactor = BaseSpeed + (MinRecentError < 0.0f ? SmoothedPlayCursor::kMinorCorrectionFactor : -SmoothedPlayCursor::kMinorCorrectionFactor);
+		if (SmoothingTimer.GetSpeed() != NewSpeedFactor)
+		{
+			UE_LOG(LogSmoothedMidiPlayCursor, Verbose, TEXT("Smoothing: Small error = %f. Adjusting speedfactor to %f."), MinRecentError, NewSpeedFactor);
+			SmoothingTimer.SetSpeed(NewSpeedFactor);
+		}
+	}
+	else
+	{
+		float NewSpeedFactor = BaseSpeed + (MinRecentError < 0.0f ? SmoothedPlayCursor::kMajorCorrectionFactor : -SmoothedPlayCursor::kMajorCorrectionFactor);
+		if (SmoothingTimer.GetSpeed() != NewSpeedFactor)
+		{
+			UE_LOG(LogSmoothedMidiPlayCursor, Verbose, TEXT("Smoothing: Significant error = %f. Adjusting speedfactor to %f."), MinRecentError, NewSpeedFactor);
+			SmoothingTimer.SetSpeed(NewSpeedFactor);
+		}
+	}
+
+	// keep track of the difference in Ms between the current time and the smoothed time
+	// update the MsDelta before applying any looping so that it's 
+	SmoothedMsDelta = SmoothMs - RawMs;
+
+	// we need to see if the smoothed position is a position that we have never played!
+	float LoopStartMs = Owner->GetLoopStartMs();
+	float LoopEndMs = Owner->GetLoopEndMs();
+	if (LoopStartMs < LoopEndMs)
+	{
+		// possible loop back...
+		if (SmoothMs > LoopEndMs)
+		{
+			// yikes, we smoothed past the end of the loop!
+			SmoothMs = LoopStartMs + (SmoothMs - LoopEndMs);
+		}
+	}
+	else if (LoopStartMs > LoopEndMs)
+	{
+		// possible jump forward...
+		if (SmoothMs >= LoopEndMs && SmoothMs < LoopStartMs)
+		{
+			// yikes, we smoothed into the gap!
+			SmoothMs = LoopStartMs + (SmoothMs - LoopEndMs);
+		}
+	}
+	SmoothedMs = SmoothMs;
+	SmoothedTick = Owner->GetTempoMap().MsToTick(SmoothedMs);
+	LoopedThisPass = false;
+}
+
+void FSmoothedMidiPlayCursor::UpdateSongPosition()
+{
+	if (!Owner)
+	{
+		return;
+	}
+
+	CurrentSongPos = Owner->CalculateSongPosWithOffsetMs(SmoothedMsDelta);
+}
diff --git a/Engine/Plugins/Runtime/Harmonix/Source/HarmonixMidi/Private/HarmonixMidi/SongMapReceiver.h b/Engine/Plugins/Runtime/Harmonix/Source/HarmonixMidi/Private/HarmonixMidi/SongMapReceiver.h
new file mode 100644
--- /dev/null
+++ b/Engine/Plugins/Runtime/Harmonix/Source/HarmonixMidi/Private/HarmonixMidi/SongMapReceiver.h
@@ -0,0 +1,61 @@
+// Copyright Epic Games, Inc. All Rights Reserved.
+#pragma once
+#include "HarmonixMidi/MidiReceiver.h"
+
+struct FSongMaps;
+
+class FSongMapReceiver : public IMidiReceiver
+{
+public:
+	FSongMapReceiver(FSongMaps* Maps);
+	virtual ~FSongMapReceiver()	{}
+
+	void Reset() override;
+	void Finalize(int32 InLastFileTick) override;
+
+	// Track begin/end:
+	virtual void OnNewTrack(int32 NewTrackIndex) override;
+	virtual void OnEndOfTrack(int32 InLastTick) override;
+	virtual void OnAllTracksRead() override {};
+
+	// Standard 1- or 2-byte MIDI message:
+	virtual void OnMidiMessage(int32 Tick, uint8 Status, uint8 Data1, uint8 Data2) override;
+	virtual void OnText(int32 Tick, const FString& Str, uint8 Type) override;
+	virtual void OnTempo(int32 Tick, int32 Tempo) override;
+	virtual void OnTimeSignature(int32 Tick, int32 Numerator, int32 Denominator, bool FailOnError = true) override;
+
+	int32 GetLastTick() { return LastTick; }
+
+private:
+	void ReadSectionTrackText(int32 Tick, const FString& Str);
+	void ReadChordTrackText(int32 Tick, const FString& Str);
+	void OnTrackName(const FString& TrackName);
+	void OnStartSectionTrack();
+	void OnStartChordTrack();
+	void OnStartBeatTrack();
+
+	FString FmtTick(int32 Tick) const;
+
+	enum class EMidiNoteAssignments {
+		Invalid = -1,
+		DetailedGridPitch = 11,
+		StrongBeatPitch = 12,
+		NormalBeatPitch = 13,
+	};
+
+	enum class EMidiTrack {
+		UnknownTrack,
+		FirstTrack,
+		BeatTrack,
+		SectionTrack,
+		ChordTrack,
+	};
+
+	EMidiTrack CurrentTrack = EMidiTrack::UnknownTrack;
+	int32 TrackIndex        = -1;
+	int32 LastBeatTick     = -1;
+	bool bHaveBeatFailure  = false;
+	int32 LastTick          = 0;
+	FSongMaps* SongMaps     = nullptr;
+	EMidiNoteAssignments LastBeatType = EMidiNoteAssignments::Invalid;
+};
diff --git a/Engine/Plugins/Runtime/Harmonix/Source/HarmonixMidi/Private/HarmonixMidi/SongMaps.cpp b/Engine/Plugins/Runtime/Harmonix/Source/HarmonixMidi/Private/HarmonixMidi/SongMaps.cpp
new file mode 100644
--- /dev/null
+++ b/Engine/Plugins/Runtime/Harmonix/Source/HarmonixMidi/Private/HarmonixMidi/SongMaps.cpp
@@ -0,0 +1,685 @@
+// Copyright Epic Games, Inc. All Rights Reserved.
+#include "HarmonixMidi/SongMaps.h"
+#include "HarmonixMidi/SongMapReceiver.h"
+#include "HarmonixMidi/MidiFile.h"
+#include "HarmonixMidi/MidiReader.h"
+#include "Misc/RuntimeErrors.h"
+
+FSongMaps::FSongMaps()
+{
+}
+
+void FSongMaps::Init(int32 InTicksPerQuarterNote)
+{
+	TicksPerQuarterNote = InTicksPerQuarterNote;
+	TempoMap.SetTicksPerQuarterNote(InTicksPerQuarterNote);
+	BarMap.SetTicksPerQuarterNote(InTicksPerQuarterNote);
+	BeatMap.SetTicksPerQuarterNote(InTicksPerQuarterNote);
+	SectionMap.SetTicksPerQuarterNote(InTicksPerQuarterNote);
+	ChordMap.SetTicksPerQuarterNote(InTicksPerQuarterNote);
+}
+
+void FSongMaps::Copy(const FSongMaps& Other, int32 StartTick, int32 EndTick)
+{
+	EmptyAllMaps();
+
+	TicksPerQuarterNote = Other.TicksPerQuarterNote;
+
+	TempoMap.Copy(Other.GetTempoMap(), StartTick, EndTick);
+	BarMap.Copy(Other.GetBarMap(), StartTick, EndTick);
+	BeatMap.Copy(Other.GetBeatMap(), StartTick, EndTick);
+	SectionMap.Copy(Other.GetSectionMap(), StartTick, EndTick);
+	ChordMap.Copy(Other.GetChordMap(), StartTick, EndTick);
+
+	int32 LastTick = EndTick == -1 ? Other.LengthData.LastTick : EndTick;
+	TempoMap.Finalize(LastTick);
+	BarMap.Finalize(LastTick);
+	BeatMap.Finalize(LastTick);
+	SectionMap.Finalize(LastTick);
+	ChordMap.Finalize(LastTick);
+
+	memset(&LengthData, 0, sizeof(LengthData));
+	LengthData.LastTick = LastTick;
+	LengthData.LengthTicks = LengthData.LastTick + 1;
+	FMusicTimestamp Timestamp;
+	Timestamp = BarMap.TickToMusicTimestamp(LengthData.LastTick);
+	LengthData.LengthBars = FMath::IsNearlyZero(Timestamp.Beat-1.0f,1.0E-4) ? Timestamp.Bar - 1 : Timestamp.Bar;
+	if (LengthData.LengthBars == Timestamp.Bar)
+	{
+		LengthData.LengthTicks = BarMap.MusicTimestampToTick({LengthData.LengthBars + BarMap.GetStartBar(), 1.0f}); 
+		LengthData.LastTick = LengthData.LengthTicks - 1;
+	}
+}
+
+bool FSongMaps::LoadFromStdMidiFile(const FString& FilePath)
+{
+	FSongMapReceiver MapReceiver(this);
+	FStdMidiFileReader Reader(FilePath, &MapReceiver);
+	return ReadWithReader(Reader);
+}
+
+bool FSongMaps::LoadFromStdMidiFile(void* Buffer, int32 BufferSize, const FString& Filename)
+{
+	FSongMapReceiver MapReceiver(this);
+	FStdMidiFileReader Reader(Buffer, BufferSize, Filename, &MapReceiver);
+	return ReadWithReader(Reader);
+}
+
+bool FSongMaps::LoadFromStdMidiFile(TSharedPtr<FArchive> Archive, const FString& Filename)
+{
+	FSongMapReceiver MapReceiver(this);
+	FStdMidiFileReader Reader(Archive, Filename, &MapReceiver);
+	return ReadWithReader(Reader);
+}
+
+float FSongMaps::TickToMs(int32 Tick) const
+{
+	if (TempoMap.IsEmpty())
+	{
+		UE_LOG(LogMidi, Log, TEXT("SongMaps does not contain a Tempo Map."));
+		return 0.0f;
+	}
+	return TempoMap.TickToMs(Tick);
+}
+
+int32 FSongMaps::MsToTick(float Ms) const
+{
+	if (TempoMap.IsEmpty())
+	{
+		UE_LOG(LogMidi, Log, TEXT("SongMaps does not contain a Tempo Map."));
+		return 0;
+	}
+	return FMath::FloorToInt(TempoMap.MsToTick(Ms));
+}
+
+float FSongMaps::GetCountInSeconds() const
+{
+	int32 BarOneBeatOneTick = BarMap.MusicTimestampToTick({1,1.0f});
+	return TempoMap.TickToSeconds(BarOneBeatOneTick);
+}
+
+bool FSongMaps::FinalizeRead(IMidiReader* Reader)
+{
+	int32 LastTick = Reader->GetLastTick();
+
+	if (Reader->IsFailed())
+	{
+		EmptyAllMaps();
+		return false;
+	}
+
+	TempoMap.Finalize(LastTick);
+	BarMap.Finalize(LastTick);
+	BeatMap.Finalize(LastTick);
+	SectionMap.Finalize(LastTick);
+	ChordMap.Finalize(LastTick);
+
+	memset(&LengthData, 0, sizeof(LengthData));
+	LengthData.LastTick = LastTick;
+	LengthData.LengthTicks = LengthData.LastTick + 1;
+	FMusicTimestamp Timestamp;
+	Timestamp = BarMap.TickToMusicTimestamp(LengthData.LastTick);
+	LengthData.LengthBars = FMath::IsNearlyZero(Timestamp.Beat - 1.0f, 1.0E-4) ? Timestamp.Bar - 1 : Timestamp.Bar;
+	if (LengthData.LengthBars == Timestamp.Bar)
+	{
+		LengthData.LengthTicks = BarMap.MusicTimestampToTick({ LengthData.LengthBars + BarMap.GetStartBar(), 1.0f });
+		LengthData.LastTick = LengthData.LengthTicks - 1;
+	}
+	return true;
+}
+
+bool FSongMaps::ReadWithReader(FStdMidiFileReader& Reader)
+{
+	Reader.ReadAllTracks();
+	return FinalizeRead(&Reader);
+}
+
+FArchive& operator<<(FArchive& Archive, FSongLengthData& LengthData)
+{
+	Archive << LengthData.LengthTicks;  
+	Archive << LengthData.LengthBars;   
+	Archive << LengthData.LastTick;
+	return Archive;
+} 
+
+void FSongMaps::EmptyAllMaps()
+{
+	TempoMap.Empty();
+	BarMap.Empty();
+	BeatMap.Empty();
+	SectionMap.Empty();
+	ChordMap.Empty();
+	TrackNames.Empty();
+}
+
+float FSongMaps::GetSongLengthMs() const
+{
+	return TickToMs(GetSongLengthData().LengthTicks);
+}
+
+int32 FSongMaps::GetSongLengthBeats() const
+{
+	return BeatMap.GetNumMapPoints();
+}
+
+int32 FSongMaps::GetSongLengthBars() const
+{
+	return GetSongLengthData().LengthBars;
+}
+
+void FSongMaps::StringLengthToMT(const FString& LengthString, int32& OutBars, int32& OutTicks)
+{
+	if (BarMap.IsEmpty())
+	{
+		UE_LOG(LogMidi, Log, TEXT("SongMaps does not contain a Bar Map."));
+		return;
+	}
+
+	// ASSUMES WE DON'T HAVE A COMPLICATED TIME SIGNATURE.
+	// IF WE DID, THEN WE WOULD NEVER GET HERE, AS WE WOULD
+	// HAVE DETERMINED THE LENGTH FROM THE MIDI FILE!
+	
+	// ALSO: LengthString, being a 'length', is specifying 0 based bars and beats!
+	//       And, the return values are lengths, so also 0 based.
+	const TCHAR* Walk = *LengthString;
+	OutBars = 0;
+	int32 Beats = 0;
+	OutTicks = 0;
+
+	// get bars...
+	int32 Count = 0;
+	while (*Walk != ':' && *Walk != 0 && Count < 3)
+	{
+		OutBars = (OutBars * 10) + ((int32)(*Walk) - (int32)'0');
+		Walk++;
+		Count++;
+	}
+	if (*Walk == ':')
+	{
+		// get beats...
+		Count = 0;
+		Walk++;
+		while (*Walk != ':' && *Walk != 0 && Count < 3)
+		{
+			Beats = (Beats * 10) + ((int32)(*Walk) - (int32)'0');
+			Walk++;
+			Count++;
+		}
+		if (*Walk == ':')
+		{
+			// get ticks...
+			Count = 0;
+			Walk++;
+			while (*Walk != ':' && *Walk != 0 && Count < 3)
+			{
+				OutTicks = (OutTicks * 10) + ((int32)(*Walk) - (int32)'0');
+				Walk++;
+				Count++;
+			}
+		}
+	}
+	const FTimeSignaturePoint* TimeSignaturePoint = &BarMap.GetTimeSignaturePoint(0);
+	if (!TimeSignaturePoint)
+	{
+		UE_LOG(LogMidi, Log, TEXT("No Time Signature found in SongMaps."));
+		return;
+	}
+	OutBars += (Beats == 0 && OutTicks == 0) ? 0 : 1;
+	int32 TicksPerBeat   = TicksPerQuarterNote / (TimeSignaturePoint->TimeSignature.Denominator / 4);
+	int32 TicksPerBar = TicksPerBeat * TimeSignaturePoint->TimeSignature.Numerator;
+	OutTicks = (TicksPerBar * OutBars)
+		+ (TicksPerBeat * Beats)
+		+ OutTicks;
+}
+
+FString FSongMaps::GetTrackName(int32 Index) const
+{
+	if (Index < 0 || Index >= TrackNames.Num())
+		return FString();
+
+	return TrackNames[Index];
+}
+
+///////////////////////////////////////////////////////////////////////////////////
+// TEMPO
+const FTempoInfoPoint* FSongMaps::GetTempoInfoForMs(float Ms) const
+{
+	if (TempoMap.IsEmpty())
+	{
+		UE_LOG(LogMidi, Log, TEXT("SongMaps does not contain a Tempo Map."));
+		return nullptr;
+	}
+	int32 Tick = MsToTick(Ms);
+	return TempoMap.GetTempoPointAtTick(Tick);
+}
+
+const FTempoInfoPoint* FSongMaps::GetTempoInfoForTick(int32 Tick) const
+{
+	if (TempoMap.IsEmpty())
+	{
+		UE_LOG(LogMidi, Log, TEXT("SongMaps does not contain a Tempo Map."));
+		return nullptr;
+	}
+	return TempoMap.GetTempoPointAtTick(Tick);
+}
+
+float FSongMaps::GetTempoAtMs(float Ms) const
+{
+	if (TempoMap.IsEmpty())
+	{
+		UE_LOG(LogMidi, Log, TEXT("SongMaps does not contain a Tempo Map."));
+		return 0.0f;
+	}
+	int32 Tick = MsToTick(Ms);
+	return TempoMap.GetTempoAtTick(Tick);
+}
+
+float FSongMaps::GetTempoAtTick(int32 Tick) const
+{
+	if (TempoMap.IsEmpty())
+	{
+		UE_LOG(LogMidi, Log, TEXT("SongMaps does not contain a Tempo Map."));
+		return 0.0f;
+	}
+	return TempoMap.GetTempoAtTick(Tick);
+}
+
+///////////////////////////////////////////////////////////////////////////////////
+// BEAT
+const FBeatMapPoint* FSongMaps::GetBeatAtMs(float Ms) const
+{
+	if (BeatMap.IsEmpty())
+	{
+		UE_LOG(LogMidi, Log, TEXT("SongMaps does not contain a Beat Map."));
+		return nullptr;
+	}
+	int32 Tick = MsToTick(Ms);
+	return BeatMap.GetPointInfoForTick(Tick);
+}
+
+float FSongMaps::GetMsAtBeat(float Beat) const
+{
+	if (!BeatMap.IsEmpty() && Beat < BeatMap.GetNumMapPoints())
+	{
+		float Tick = BeatMap.GetFractionalTickAtBeat(Beat);
+		return TickToMs(Tick);
+	}
+
+	if (BarMap.IsEmpty())
+	{
+		return 0.0f;
+	}
+
+	const FTimeSignaturePoint& Point1 = BarMap.GetTimeSignaturePoint(0);
+	if (BarMap.GetNumTimeSignaturePoints() == 1)
+	{
+		float Tick = Beat * TicksPerQuarterNote / (Point1.TimeSignature.Denominator / 4);
+		return TickToMs(Tick);
+	}
+
+	const FTimeSignaturePoint* FromPoint = &Point1;
+	for (int i = 1; i < BarMap.GetNumTimeSignaturePoints(); ++i)
+	{
+		const FTimeSignaturePoint& Point2 = BarMap.GetTimeSignaturePoint(i);
+		if (Beat < Point2.BeatIndex)
+		{
+			break;
+		}
+		FromPoint = &Point2;
+	}
+
+	float BeatsAtTimeSignature = Beat - (float)FromPoint->BeatIndex;
+	float Bar = (float)FromPoint->BarIndex + (BeatsAtTimeSignature / (float)FromPoint->TimeSignature.Numerator);
+	return TickToMs(BarMap.FractionalBarIncludingCountInToTick(Bar));
+}
+
+const FBeatMapPoint* FSongMaps::GetBeatAtTick(int32 Tick) const
+{
+	if (BeatMap.IsEmpty())
+	{
+		UE_LOG(LogMidi, Log, TEXT("SongMaps does not contain a Beat Map."));
+		return nullptr;
+	}
+	return BeatMap.GetPointInfoForTick(Tick);
+}
+
+float FSongMaps::GetMsPerBeatAtMs(float Ms) const
+{
+	int32 Tick = MsToTick(Ms);
+	return GetMsPerBeatAtTick(Tick);
+}
+
+float FSongMaps::GetMsPerBeatAtTick(int32 Tick) const
+{
+	const FBeatMapPoint* BeatInfo = GetBeatAtTick(Tick);
+	if (!BeatInfo)
+	{
+		// In midi, tempo is always quarter notes per minute. Without a beat map,
+		// a beat is the divisor in the time signature, which might not be a 
+		// quarter note. So we will convert here...
+		float QuarterNotesPerMinute = GetTempoAtTick(Tick);
+		const FTimeSignature* TimeSignature = GetTimeSignatureAtTick(Tick);
+		float BeatsPerMinute = QuarterNotesPerMinute / ((float)TimeSignature->Denominator / 4.0f /* quarter note / 4 == 1 */);
+		return 60000.0f / BeatsPerMinute;
+	}
+	return  TickToMs(BeatInfo->StartTick + BeatInfo->LengthTicks) - TickToMs(BeatInfo->StartTick);
+}
+
+float FSongMaps::GetFractionalBeatAtMs(float Ms) const
+{
+	int32 Tick = MsToTick(Ms);
+	return GetFractionalBeatAtTick(Tick);
+}
+
+float FSongMaps::GetFractionalBeatAtTick(int32 Tick) const
+{
+	int32 BeatIndex = GetBeatIndexAtTick(Tick);
+	if (BeatIndex < 0 || BeatMap.IsEmpty())
+	{
+		return 1.0f; // 1 based position!
+	}
+
+	const FBeatMapPoint* BeatInfo = &BeatMap.GetBeatPointInfo(BeatIndex);
+	check(BeatInfo);
+	int32 TickInBeat = Tick - BeatInfo->StartTick;
+	float FractionalPart = (float)TickInBeat / (float)BeatInfo->LengthTicks;
+	return (float)BeatIndex + FractionalPart + 1.0f; // +1 for musical position
+}
+
+int32 FSongMaps::GetBeatIndexAtMs(float Ms) const
+{
+	int32 Tick = MsToTick(Ms);
+	return GetBeatIndexAtTick(Tick);
+}
+
+int32 FSongMaps::GetBeatIndexAtTick(int32 Tick) const
+{
+	if (BeatMap.IsEmpty())
+	{
+		UE_LOG(LogMidi, Log, TEXT("SongMaps does not contain a Beat Map."));
+		return -1;
+	}
+	return BeatMap.GetPointIndexForTick(Tick);
+}
+
+EMusicalBeatType FSongMaps::GetBeatTypeAtMs(float Ms) const
+{
+	int32 Tick = MsToTick(Ms);
+	return GetBeatTypeAtTick(Tick);
+}
+
+EMusicalBeatType FSongMaps::GetBeatTypeAtTick(int32 Tick) const
+{
+	if (BeatMap.IsEmpty())
+	{
+		UE_LOG(LogMidi, Log, TEXT("SongMaps does not contain a Beat Map."));
+		return EMusicalBeatType::Normal;
+	}
+	return BeatMap.GetBeatTypeAtTick(Tick);
+}
+
+float FSongMaps::GetBeatInPulseBarAtMs(float Ms) const
+{
+	int32 Tick = MsToTick(Ms);
+	return GetBeatInPulseBarAtTick(Tick);
+}
+
+float FSongMaps::GetBeatInPulseBarAtTick(int32 Tick) const
+{
+	if (BeatMap.IsEmpty())
+	{
+		UE_LOG(LogMidi, Log, TEXT("SongMaps does not contain a Beat Map."));
+		return 0.0f;
+	}
+	return BeatMap.GetBeatInPulseBarAtTick(Tick);
+}
+
+float FSongMaps::GetNumBeatsInPulseBarAtMs(float Ms) const
+{
+	int32 Tick = MsToTick(Ms);
+	return GetNumBeatsInPulseBarAtTick(Tick);
+}
+
+float FSongMaps::GetNumBeatsInPulseBarAtTick(int32 Tick) const
+{
+	if (BeatMap.IsEmpty())
+	{
+		UE_LOG(LogMidi, Log, TEXT("SongMaps does not contain a Beat Map."));
+		return 0.0f;
+	}
+	return BeatMap.GetNumBeatsInPulseBarAt(Tick);
+}
+
+///////////////////////////////////////////////////////////////////////////////////
+// Time Signature
+
+const FTimeSignature* FSongMaps::GetTimeSignatureAtMs(float Ms) const
+{
+	int32 Tick = MsToTick(Ms);
+	return GetTimeSignatureAtTick(Tick);
+}
+
+const FTimeSignature* FSongMaps::GetTimeSignatureAtTick(int32 Tick) const
+{
+	if (BarMap.IsEmpty())
+	{
+		UE_LOG(LogMidi, Log, TEXT("SongMaps does not contain a Bar Map."));
+		return nullptr;
+	}
+	return &BarMap.GetTimeSignatureAtTick(Tick);
+}
+
+const FTimeSignature* FSongMaps::GetTimeSignatureAtBar(int32 Bar) const
+{
+	if (Bar < 1)
+	{
+		UE_LOG(LogMidi, Warning, TEXT("Bar < 1 (%d) specified as a musical position! Bars are '1' based in musical positions. Using bar 1!"), Bar);
+	}
+
+	if (BarMap.IsEmpty())
+	{
+		UE_LOG(LogMidi, Log, TEXT("SongMaps does not contain a Bar Map."));
+		return nullptr;
+	}
+	return &BarMap.GetTimeSignatureAtBar(Bar);
+}
+
+float FSongMaps::GetBarIncludingCountInAtMs(float Ms) const
+{
+	int32 Tick = MsToTick(Ms);
+	return GetBarIncludingCountInAtTick(Tick);
+}
+
+float FSongMaps::GetBarIncludingCountInAtTick(int32 Tick) const
+{
+	if (BarMap.IsEmpty())
+	{
+		UE_LOG(LogMidi, Log, TEXT("SongMaps does not contain a Bar Map."));
+		return 1.0f; // 1 based music positions!
+	}
+	return BarMap.TickToFractionalBarIncludingCountIn(Tick);
+}
+
+float FSongMaps::GetMsPerBarAtMs(float Ms) const
+{
+	int32 Tick = MsToTick(Ms);
+	return GetMsPerBarAtTick(Tick);
+}
+
+float FSongMaps::GetMsPerBarAtTick(int32 Tick) const
+{
+	if (TempoMap.IsEmpty())
+	{
+		UE_LOG(LogMidi, Log, TEXT("SongMaps does not contain a Tempo Map."));
+		return 0.0f;
+	}
+	if (BarMap.IsEmpty())
+	{
+		UE_LOG(LogMidi, Log, TEXT("SongMaps does not contain a Bar Map."));
+		return 0.0f;
+	}
+	float Bpm = TempoMap.GetTempoAtTick(Tick); // quarter noted per minute
+	const FTimeSignature* TimeSignature = GetTimeSignatureAtTick(Tick);
+	float QuarterNotesInBar;
+	if (TimeSignature)
+	{
+		QuarterNotesInBar = (float)TimeSignature->Numerator / ((float)TimeSignature->Denominator / 4.0f);
+	}
+	else
+	{
+		QuarterNotesInBar = 4.0f;
+	}
+	float MsPerQuarterNote = 1000.0f / (Bpm / 60.0f);
+	return QuarterNotesInBar * MsPerQuarterNote;
+}
+
+void FSongMaps::SetLengthTotalBars(int32 Bars)
+{
+	check(Bars >= 0);
+	memset(&LengthData, 0, sizeof(LengthData));
+	LengthData.LengthTicks = BarMap.BarIncludingCountInToTick(Bars);
+	LengthData.LastTick = LengthData.LengthTicks - 1;
+	LengthData.LengthBars = Bars;
+}
+
+///////////////////////////////////////////////////////////////////////////////////
+// Sections
+
+const FSongSection* FSongMaps::GetSectionAtMs(float Ms) const
+{
+	int32 Tick = MsToTick(Ms);
+	return GetSectionAtTick(Tick);
+}
+
+const FSongSection* FSongMaps::GetSectionAtTick(int32 Tick) const
+{
+	if (SectionMap.IsEmpty())
+	{
+		UE_LOG(LogMidi, Log, TEXT("SongMaps does not contain a Section Map."));
+		return nullptr;
+	}
+	return SectionMap.TickToSection(Tick);
+}
+
+const FSongSection* FSongMaps::GetSectionWithName(const FString& Name) const
+{
+	if (SectionMap.IsEmpty())
+	{
+		UE_LOG(LogMidi, Log, TEXT("SongMaps does not contain a Section Map."));
+		return nullptr;
+	}
+	return SectionMap.FindSectionInfo(Name);
+}
+
+FString FSongMaps::GetSectionNameAtMs(float Ms) const
+{
+	int32 Tick = MsToTick(Ms);
+	return GetSectionNameAtTick(Tick);
+}
+
+FString FSongMaps::GetSectionNameAtTick(int32 Tick) const
+{
+	if (SectionMap.IsEmpty())
+	{
+		UE_LOG(LogMidi, Log, TEXT("SongMaps does not contain a Section Map."));
+		return FString();
+	}
+	return SectionMap.GetSectionNameAtTick(Tick);
+}
+
+float FSongMaps::GetSectionLengthMsAtMs(float Ms) const
+{
+	int32 Tick = MsToTick(Ms);
+	return GetSectionLengthMsAtTick(Tick);
+}
+
+float FSongMaps::GetSectionStartMsAtMs(float Ms) const
+{
+	const int32 Tick = MsToTick(Ms);
+	const FSongSection* SectionAtTick = GetSectionAtTick(Tick);
+	if (!SectionAtTick)
+	{
+		return 0.0f;
+	}
+
+	return TickToMs(SectionAtTick->StartTick);
+}
+
+float FSongMaps::GetSectionEndMsAtMs(float Ms) const
+{
+	const int32 Tick = MsToTick(Ms);
+	const FSongSection* SectionAtTick = GetSectionAtTick(Tick);
+	if (!SectionAtTick)
+	{
+		return 0.0f;
+	}
+
+	return TickToMs(SectionAtTick->EndTick());
+}
+
+float FSongMaps::GetSectionLengthMsAtTick(int32 Tick) const
+{
+	const FSongSection* SectionAtTick = GetSectionAtTick(Tick);
+	if (!SectionAtTick)
+	{
+		return 0.0f;
+	}
+
+	float StartMs = TickToMs(SectionAtTick->StartTick);
+	float EndMs = TickToMs(SectionAtTick->EndTick());
+	return EndMs - StartMs;
+}
+
+///////////////////////////////////////////////////////////////////////////////////
+// Chords
+
+const FChordMapPoint* FSongMaps::GetChordAtMs(float Ms) const
+{
+	int32 Tick = MsToTick(Ms);
+	return GetChordAtTick(Tick);
+}
+
+const FChordMapPoint* FSongMaps::GetChordAtTick(float Tick) const
+{
+	if (ChordMap.IsEmpty())
+	{
+		UE_LOG(LogMidi, Log, TEXT("SongMaps does not contain a Chord Map."));
+		return nullptr;
+	}
+	return ChordMap.GetPointInfoForTick(Tick);
+}
+
+FName FSongMaps::GetChordNameAtMs(float Ms) const
+{
+	int32 Tick = MsToTick(Ms);
+	return GetChordNameAtTick(Tick);
+}
+
+FName FSongMaps::GetChordNameAtTick(int32 Tick) const
+{
+	if (ChordMap.IsEmpty())
+	{
+		UE_LOG(LogMidi, Log, TEXT("SongMaps does not contain a Chord Map."));
+		return FName();
+	}
+	return ChordMap.GetChordNameAtTick(Tick);
+}
+
+float FSongMaps::GetChordLengthMsAtMs(float Ms) const
+{
+	int32 Tick = MsToTick(Ms);
+	return GetChordLengthMsAtTick(Tick);
+}
+
+float FSongMaps::GetChordLengthMsAtTick(int32 Tick) const
+{
+	const FChordMapPoint* ChordInfo = GetChordAtTick(Tick);
+	if (!ChordInfo)
+	{
+		return 0.0f;
+	}
+	float ChordStartMs = TickToMs(ChordInfo->StartTick);
+	float ChordEndMs = TickToMs(ChordInfo->EndTick());
+	return ChordEndMs - ChordStartMs;
+}
+
diff --git a/Engine/Plugins/Runtime/Harmonix/Source/HarmonixMidi/Public/HarmonixMidi/BeatMap.h b/Engine/Plugins/Runtime/Harmonix/Source/HarmonixMidi/Public/HarmonixMidi/BeatMap.h
new file mode 100644
--- /dev/null
+++ b/Engine/Plugins/Runtime/Harmonix/Source/HarmonixMidi/Public/HarmonixMidi/BeatMap.h
@@ -0,0 +1,144 @@
+// Copyright Epic Games, Inc. All Rights Reserved.
+#pragma once
+#include "HarmonixMidi/MusicMapBase.h"
+#include "HarmonixMidi/MidiConstants.h"
+
+#include "BeatMap.generated.h"
+
+UENUM(BlueprintType)
+enum class EMusicalBeatType : uint8
+{
+	Downbeat = 0,
+	Strong = 1,
+	Normal = 2,
+};
+
+/**
+ * A point in the music representing a "beat".
+ * 
+ * Type may be 'Downbeat', 'Strong', or 'Normal' beat. 
+ * 
+ * This is very useful for odd time signatures like 5/8 where the music's
+ * beat might be on the 1st and 3rd eighth notes or the 1st and 4th eighth notes
+ */
+USTRUCT()
+struct HARMONIXMIDI_API FBeatMapPoint : public FMusicMapTimespanBase
+{
+	GENERATED_BODY()
+
+public:
+	static constexpr bool DefinedAsRegions = true;
+
+	FBeatMapPoint()
+		: PulseBar(-1)
+		, Type(EMusicalBeatType::Normal)
+	{}
+
+	FBeatMapPoint(EMusicalBeatType InType, int32 InStartTick, int32 InLengthTicks = 1)
+		: FMusicMapTimespanBase(InStartTick, InLengthTicks)
+		, PulseBar(0)
+		, Type(InType)
+	{}
+
+	UPROPERTY()
+	int32 PulseBar;
+	UPROPERTY()
+	EMusicalBeatType Type;
+};
+
+/**
+ * Pulse Bars are groupings of beats where the first beat in the group has been marked up as a
+ * 'Downbeat' type.
+ */
+USTRUCT()
+struct HARMONIXMIDI_API FPulseBar
+{
+	GENERATED_BODY()
+
+public:
+	FPulseBar()
+		: StartTick(0)
+		, LengthTicks(0)
+		, FirstIncludedBeatIndex(0)
+		, LastIncludedBeatIndex(0)
+	{}
+	FPulseBar(int32 InStartTick, int32 InLengthTicks, int32 InFirstIncludedBeatIndex, int32 InLastIncludedBeatIndex)
+		: StartTick(InStartTick)
+		, LengthTicks(InLengthTicks)
+		, FirstIncludedBeatIndex(InFirstIncludedBeatIndex)
+		, LastIncludedBeatIndex(InLastIncludedBeatIndex)
+	{}
+
+	UPROPERTY()
+	int32 StartTick;
+	UPROPERTY()
+	int32 LengthTicks;
+	UPROPERTY()
+	int32 FirstIncludedBeatIndex;
+	UPROPERTY()
+	int32 LastIncludedBeatIndex;
+};
+
+/**
+ * A map of 'beats' in a piece of music.
+ */
+USTRUCT()
+struct HARMONIXMIDI_API FBeatMap
+{
+	GENERATED_BODY()
+
+public:
+	FBeatMap()
+		: TicksPerQuarterNote(MidiConstants::kTicksPerQuarterNoteInt)
+	{}
+
+	void Empty();
+	void Copy(const FBeatMap& Other, int32 StartTick = 0, int32 EndTick = -1);
+	bool IsEmpty() const;
+	int32 GetNumMapPoints() const;
+
+	/** Called by the midi file importer before map points are added to this map */
+	void SetTicksPerQuarterNote(int32 InTicksPerQuarterNote)
+	{
+		TicksPerQuarterNote = InTicksPerQuarterNote;
+	}
+
+	void AddBeat(EMusicalBeatType Type, int32 Tick, bool SortNow = true);
+
+	const FBeatMapPoint* GetPointInfoForTick(int32 Tick) const;
+
+	int32 GetPointIndexForTick(int32 Tick) const;
+
+	float GetFractionalBeatAtTick(float Tick) const;
+
+	float GetFractionalTickAtBeat(float InBeat) const;
+
+	EMusicalBeatType GetBeatTypeAtTick(float Tick) const;
+
+	float GetBeatInPulseBarAtTick(float tick) const;
+
+	float GetNumBeatsInPulseBarAt(float Tick) const;
+
+	/** Is the specified "from count in" beat the first beat in its bar? */
+	bool IsDownbeat(float Beat) const;
+
+	/** Retrieve the first downbeat AFTER the specified "from count in" beat */
+	int32 FindDownbeatIndexAfterBeat(float Beat) const;
+
+	/** Retrieve the first downbeat ON OR AFTER the specified "from count in"  beat */
+	int32 FindDownbeatIndexAtOrAfterBeat(float Beat) const;
+
+	const FBeatMapPoint& GetBeatPointInfo(int32 index) const;
+	
+	void Finalize(int32 LastTick);
+
+protected:
+	UPROPERTY()
+	int32 TicksPerQuarterNote;
+	UPROPERTY()
+	TArray<FBeatMapPoint> Points;
+
+	UPROPERTY()
+	TArray<FPulseBar> Bars;
+};
+
diff --git a/Engine/Plugins/Runtime/Harmonix/Source/HarmonixMidi/Public/HarmonixMidi/Blueprint/MidiNote.h b/Engine/Plugins/Runtime/Harmonix/Source/HarmonixMidi/Public/HarmonixMidi/Blueprint/MidiNote.h
new file mode 100644
--- /dev/null
+++ b/Engine/Plugins/Runtime/Harmonix/Source/HarmonixMidi/Public/HarmonixMidi/Blueprint/MidiNote.h
@@ -0,0 +1,63 @@
+// Copyright Epic Games, Inc. All Rights Reserved.
+
+#pragma once
+
+#include "MidiNote.generated.h"
+
+/**
+* Helper struct for getting and assigning MidiNote values
+* Uses custom detail customization and custom pins for convenience in Editor
+* 
+* has a one property
+* uint8 NoteNumber
+* ranging from Midi::kMinNote to Midi::kMaxNote [0 - 127]
+*/
+USTRUCT(BlueprintType)
+struct HARMONIXMIDI_API FMidiNote
+{
+	GENERATED_BODY()
+
+	FMidiNote() : NoteNumber(60) { };
+
+	FMidiNote(uint8 inValue) : NoteNumber(inValue) { };
+	FMidiNote(int inValue) : NoteNumber(inValue) { };
+	FMidiNote(int8 inValue) : NoteNumber(inValue) { };
+
+	operator uint8() { return NoteNumber; }
+	operator int8() { return NoteNumber; }
+	operator int() { return NoteNumber; }
+
+	bool operator==(const FMidiNote& rhs) const { return NoteNumber == rhs.NoteNumber; }
+	bool operator!=(const FMidiNote& rhs) const { return NoteNumber != rhs.NoteNumber; }
+	bool operator< (const FMidiNote& rhs) const { return NoteNumber <  rhs.NoteNumber; }
+	bool operator> (const FMidiNote& rhs) const { return NoteNumber >  rhs.NoteNumber; }
+	bool operator<=(const FMidiNote& rhs) const { return NoteNumber <= rhs.NoteNumber; }
+	bool operator>=(const FMidiNote& rhs) const { return NoteNumber >= rhs.NoteNumber; }
+	
+	
+
+	UPROPERTY(EditAnywhere, BlueprintReadWrite, Category = "Midi")
+	uint8 NoteNumber = 60;
+
+	FText GetDisplayText() const;
+
+	FName GetDisplayName() const;
+
+	// note names formated enharmonic stile sharp with octave number (eg. C#4)
+	FString ToString() const;
+
+	// takes a note name (eg C#4) and makes a FMidiNote struct from it
+	static FMidiNote FromString(const FString& NoteName);
+
+	static uint8 NoteNumberFromString(const FString& NoteName);
+
+#ifdef WITH_EDITOR
+
+	FString ToEditorString() const;
+
+	static FMidiNote FromEditorString(const FString& EditorName);
+
+	static uint8 NoteNumberFromEditorString(const FString& EditorName);
+#endif
+};
+
diff --git a/Engine/Plugins/Runtime/Harmonix/Source/HarmonixMidi/Public/HarmonixMidi/MidiConstants.h b/Engine/Plugins/Runtime/Harmonix/Source/HarmonixMidi/Public/HarmonixMidi/MidiConstants.h
new file mode 100644
--- /dev/null
+++ b/Engine/Plugins/Runtime/Harmonix/Source/HarmonixMidi/Public/HarmonixMidi/MidiConstants.h
@@ -0,0 +1,239 @@
+// Copyright Epic Games, Inc. All Rights Reserved.
+#pragma once
+#include "Logging/LogMacros.h"
+
+DECLARE_LOG_CATEGORY_EXTERN(LogMidi, Log, All);
+
+class HARMONIXMIDI_API MidiConstants
+{
+public:
+	enum class EMidiTextEventEncoding
+	{
+		Latin1,
+		UTF8
+	};
+
+	// helper enum for common time subdivisions
+	enum class ESubdivisions
+	{
+		Downbeat      = 1,
+		Half          = 2,
+		ThreeOverFour = 3,
+		Quarter       = 4,
+		Eighth        = 8,
+		Triplet       = 12,
+		Sixteenth     = 16,
+		ThirtySecond  = 32,
+	};
+
+	/////////////////////////////////////////////////////////////////////////////// 
+	// Constants and functions for MIDI 1.0 standard
+	//
+	// Useful web resources:
+	//  http://www.midi.org/about-midi/table1.shtml
+	//  http://www.borg.com/~jglatt/tech/midispec.htm
+
+	static const uint8 kNumChannels    = 16;
+	static const uint8 kNotesPerOctave = 12;
+	static const uint8 kMinNote        = 0;
+	static const uint8 kMaxNote        = 127;
+	static const uint8 kMaxNumNotes    = 128;
+	static const uint8 kMiddleC        = 60;
+	static const uint8 kMinVelocity    = 0;
+	static const uint8 kMaxVelocity    = 127;
+
+	static const int32   kTicksPerQuarterNoteInt  = 960;
+	static constexpr float kTicksPerQuarterNote = 960.0f;
+	static constexpr float kQuarterNotesPerTick = 1.0f / kTicksPerQuarterNote;
+
+	//////////////////////////////////////////////////////////////////////////
+	// Constants for the SMF (Standard Midi File) format.
+
+	// Codes for special handling of System Exclusive messages:
+	static const uint8 kFile_Escape      = 0xf7;
+	static const uint8 kFile_SysEx       = 0xf0; 
+
+	// Status code for SMF meta-events:
+	static const uint8 kFile_Meta        = 0xff;
+
+	// Meta-event IDs:
+	static const uint8 kMeta_Text            = 0x01;
+	static const uint8 kMeta_Copyright       = 0x02;
+	static const uint8 kMeta_TrackName       = 0x03;
+	static const uint8 kMeta_InstrumentName  = 0x04;
+	static const uint8 kMeta_Lyric           = 0x05;
+	static const uint8 kMeta_Marker          = 0x06;
+	static const uint8 kMeta_CuePoint        = 0x07;
+	static const uint8 kMeta_ChannelPrefix   = 0x20;
+	static const uint8 kMeta_Port            = 0x21; // obsolete
+	static const uint8 kMeta_EndOfTrack      = 0x2f;
+	static const uint8 kMeta_Tempo           = 0x51;
+	static const uint8 kMeta_SMPTE           = 0x54;
+	static const uint8 kMeta_TimeSig         = 0x58;
+	static const uint8 kMeta_KeySig          = 0x59;
+	static const uint8 kMeta_Special         = 0x7f;
+
+	// if FMidiMsg type is "Runtime" these are the possible status byte.
+	// note: runtime messages are not serialized!
+	static const uint8 kRuntimeAllNotesOffStatus = 0x01;  // Receiver should feel free to allow notes to "adsr release"
+	static const uint8 kRuntimeAllNotesKillStatus = 0x02; // Receiver is expected to immediately free resources associated with any sounding notes and stop playing ASAP!
+
+	static FString GetTextTypeName(uint8 TextType);
+
+	///////////////////////////////////////////////////////////////////////////////
+	// Status bytes indicate the type of message (3 bytes long except where noted).
+
+	// Masks for disassembling MIDI Status bytes
+	static const uint8 kStatusBitMask   = 0x80; // hi bit always set on Status byte
+	static const uint8 kMessageTypeMask = 0xf0; // mask for message type
+	static const uint8 kChannelMask     = 0x0f; // mask for channel
+	static const uint8 kRealTimeMask    = 0xf8; // indicates real-time Status message
+
+	static const uint8 kNoteOff  = 0x80; // Note Off
+	static const uint8 kNoteOn   = 0x90; // Note On
+	static const uint8 kPolyPres = 0xa0; // Polyphonic Key Pressure 
+	static const uint8 kControl  = 0xb0; // Control Change
+	static const uint8 kProgram  = 0xc0; // Program Change (2 bytes)
+	static const uint8 kChanPres = 0xd0; // Channel Pressure (2 bytes)
+	static const uint8 kPitch    = 0xe0; // Pitch Wheel Change
+	static const uint8 kSystem   = 0xf0; // System
+
+
+	static inline bool  IsStatus(uint8 Byte)     { return (Byte & kStatusBitMask) > 0;  }
+	static inline uint8 GetChannel(uint8 Status) { check(IsStatus(Status));	check(!IsSystem(Status)); return Status & kChannelMask; }
+	static inline uint8 GetType(uint8 Status)    { check(IsStatus(Status)); return Status & kMessageTypeMask; }
+
+	static inline bool  IsNoteOff(uint8 Status)  { return GetType(Status) == kNoteOff;  }
+	static inline bool  IsNoteOn(uint8 Status)   { return GetType(Status) == kNoteOn;   }
+	static inline bool  IsPolyPres(uint8 Status) { return GetType(Status) == kPolyPres; }
+	static inline bool  IsControl(uint8 Status)  { return GetType(Status) == kControl;  }
+	static inline bool  IsProgram(uint8 Status)  { return GetType(Status) == kProgram;  }
+	static inline bool  IsChanPres(uint8 Status) { return GetType(Status) == kChanPres; }
+	static inline bool  IsPitch(uint8 Status)    { return GetType(Status) == kPitch;    }
+	static inline bool  IsSystem(uint8 Status)   { return GetType(Status) == kSystem;   }
+
+	// turns MIDI style tempo (microseconds per quarter note) into BPM...
+	static inline float MidiTempoToBPM(int32 UsPerQuarterNote) { return UsPerQuarterNote== 0 ? 0 : 60000000.0f/(float)UsPerQuarterNote; }
+	static inline int32 BPMToMidiTempo(float Bpm)       { return Bpm == 0.0 ? 0 : (int32)(60000000.0f / (float)Bpm); }
+
+	static float RoundToStandardBeatPrecision(float InBeat, int TimeSIgnatureDenominator);
+
+	///////////////////////////////////////////////////////////////////////////////
+	// Controllers
+	// 
+	// In a Control Change message, data byte 1 is the controller ID, and data
+	// byte 2 is the controller value.
+	enum EControllerID
+	{ // from midi specification...
+	   // NOTE: If you add any here also add them to the GetControllerName function!
+		BankSelection = 0,
+		ModWheel = 1,
+		Breath = 2,
+
+		PortamentoTime = 5,
+		DataCoarse = 6,
+		Volume = 7,
+		Balance = 8,
+		PanRight = 10,
+		Expression = 11,
+
+		BitCrushWetMix = 14,
+		BitCrushLevel = 15,
+		BitCrushSampleHold = 16,
+
+		LFO0Frequency = 22,
+		LFO1Frequency = 23,
+		LFO0Depth = 24,
+		LFO1Depth = 25,
+
+		DelayTime = 26,
+		DelayDryGain = 27,
+		DelayWetGain = 28,
+		DelayFeedback = 29,
+
+		CoarsePitchBend = 30,
+		SampleStartTime = 31,
+		DataFine = 38,
+
+		SubStreamVol1 = 52,
+		SubStreamVol2 = 53,
+		SubStreamVol3 = 54,
+		SubStreamVol4 = 55,
+		SubStreamVol5 = 56,
+		SubStreamVol6 = 57,
+		SubStreamVol7 = 58,
+		SubStreamVol8 = 59,
+
+		DelayEQEnabled = 60,
+		DelayEQType = 61,
+		DelayEQFreq = 62,
+		DelayEQQ = 63,
+
+		Hold = 64,
+		PortamentoSwitch = 65,
+		Sustenuto = 66,
+		SoftPedal = 67,
+		Legato = 68,
+		Hold2 = 69,
+		FilterQ = 71,
+		Release = 72,
+		Attack = 73,
+		FilterFrequency = 74,
+
+		TimeStretchEnvelopeOrder = 79,
+
+		DelayLFOBeatSync = 80,
+		DelayLFOEnabled = 81,
+		DelayLFORate = 82,
+		DelayLFODepth = 83,
+
+		DelayStereoType = 84,
+		DelayPanLeft = 85,
+		DelayPanRight = 86,
+
+		RPNFine = 100,
+		RPNCourse = 101,
+		AllSoundOff = 120,
+		Reset = 121,
+		AllNotesOff = 123
+	};
+
+	static FString GetControllerName(EControllerID ControllerId);
+	static void    GetControllerNames(TArray<FString>& Names);
+
+	// allows choice between various naming conventions for enharmonic notes
+	enum class ENoteNameEnharmonicStyle
+	{
+		Sharp,
+		Flat,
+		SharpAndFlat
+	};
+
+	static const int8 GetNoteNumberFromNoteName(const char* InName);
+	static const char* GetNoteNameFromNoteNumber(uint8 MidiNoteNumber, ENoteNameEnharmonicStyle Style = ENoteNameEnharmonicStyle::SharpAndFlat);
+	static int8 GetNoteOctaveFromNoteNumber(uint8 MidiNoteNumber);
+
+	///////////////////////////////////////////////////////////////////////////////
+	// RPNs (registered paramater numbers)
+	static const uint8 kRPN_BendRangeC   = 0; 
+	static const uint8 kRPN_BendRangeF   = 0; 
+
+
+	///////////////////////////////////////////////////////////////////////////////
+	// System messages:
+
+	// System Common:
+	static const uint8 kSys_MTC         = 0xf1; // 2 byte message
+	static const uint8 kSys_SongPos     = 0xf2; // 3 byte message
+	static const uint8 kSys_SongSelect  = 0xf3; // 2 byte message
+	static const uint8 kSys_TuneRequest = 0xf6; // 1 byte message
+	static const uint8 kSys_EOX         = 0xf7; // 1 byte message
+
+	// System RealTime messages (all 1 byte)
+	static const uint8 kSys_TimingClock = 0xf8;
+	static const uint8 kSys_Start       = 0xfa;
+	static const uint8 kSys_Continue    = 0xfb;
+	static const uint8 kSys_Stop        = 0xfc;
+	static const uint8 kSys_ActiveSense = 0xfe;
+	static const uint8 kSys_Reset       = 0xff;
+};
diff --git a/Engine/Plugins/Runtime/Harmonix/Source/HarmonixMidi/Public/HarmonixMidi/MidiFile.h b/Engine/Plugins/Runtime/Harmonix/Source/HarmonixMidi/Public/HarmonixMidi/MidiFile.h
new file mode 100644
--- /dev/null
+++ b/Engine/Plugins/Runtime/Harmonix/Source/HarmonixMidi/Public/HarmonixMidi/MidiFile.h
@@ -0,0 +1,228 @@
+// Copyright Epic Games, Inc. All Rights Reserved.
+#pragma once
+#include "HarmonixMidi/TempoMap.h"
+#include "HarmonixMidi/MidiMsg.h"
+#include "HarmonixMidi/BarMap.h"
+#include "HarmonixMidi/SongMaps.h"
+#include "HarmonixMidi/MidiTrack.h"
+#include "IAudioProxyInitializer.h"
+
+
+#include "MidiFile.generated.h"
+
+struct FSongMaps;
+class UAssetImportData;
+
+class FMidiFileProxy;
+using FMidiFileProxyPtr = TSharedPtr<FMidiFileProxy, ESPMode::ThreadSafe>;
+
+//midi file length conform options,rounding file length to a whole bar up or down 
+enum class EMidiFileLengthConformOption : uint8
+{
+	RoundDown,
+	RoundUp,
+	Nearest
+};
+
+USTRUCT(BlueprintType)
+struct HARMONIXMIDI_API FMidiFileData
+{
+	GENERATED_BODY()
+
+	FMidiFileData()
+		: TicksPerQuarterNote(MidiConstants::kTicksPerQuarterNoteInt)
+		, LastEventTick(0)
+	{}
+	UPROPERTY(BlueprintReadOnly, Category = "MidiFile")
+	FString MidiFileName;
+	UPROPERTY(BlueprintReadOnly, Category = "MidiFile")
+	int32 TicksPerQuarterNote;
+	UPROPERTY(BlueprintReadOnly, Category = "MidiFile")
+	FSongMaps SongMaps;
+	UPROPERTY(BlueprintReadOnly, Category = "MidiFile")
+	TArray<FMidiTrack> Tracks;
+	UPROPERTY(BlueprintReadOnly, Category = "MidiFile")
+	int32 LastEventTick;
+
+	int32 FindTrackIndexByName(const FString& TrackName);
+	
+	void PostSerialize(const FArchive& Ar);
+
+	/**
+	* Adds a tempo change to the midi data at the given tick for the given track idx
+	* 
+	* asserts valid TrackIdx
+	*/
+	void AddTempoChange(int32 TrackIdx, int32 Tick, float TempoBPM);
+	
+	/**
+	* Adds a Time Signature change to the next bar boundary from the input tick for the given track
+	* 
+	* Time signature changes can only happen at the beginning of a bar
+	* So input Tick will be rounded UP to the next bar boundary.
+	* 
+	* asserts valid TrackIdx
+	*/
+	void AddTimeSigChange(int32 TrackIdx, int32 Tick, int32 TimeSigNum, int32 TimeSigDenom);
+
+};
+
+template<>
+struct TStructOpsTypeTraits<FMidiFileData> : public TStructOpsTypeTraitsBase2<FMidiFileData>
+{
+	enum
+	{
+		WithPostSerialize = true,
+	};
+};
+
+/**
+ * An FMidFile is primarily a container for FMidiTracks. 
+ * 
+ * This class can handle loading and saving standard midi files, as well
+ * as serializing itself to standard Unreal Engine FArchives.
+ */
+UCLASS(BlueprintType, Category="Music")
+class HARMONIXMIDI_API UMidiFile : public UObject, public IAudioProxyDataFactory
+{
+	GENERATED_BODY()
+
+public:
+	using FMidiTrackList = TArray<FMidiTrack>;
+
+	UMidiFile();
+	virtual void GetAssetRegistryTags(FAssetRegistryTagsContext Context) const override;
+	virtual void PostInitProperties() override;
+
+	/** A method for importing a standard midi file, with the option of providing a pointer to an FSongMaps instance that will be populated during the load. */
+	void LoadStdMidiFile(const FString& FilePath, int32 DesiredTicksPerQuarterNote = MidiConstants::kTicksPerQuarterNoteInt,
+		MidiConstants::EMidiTextEventEncoding InTextEncoding = MidiConstants::EMidiTextEventEncoding::UTF8);
+	/** A method for importing a standard midi file, with the option of providing a pointer to an FSongMaps instance that will be populated during the load. */
+	void LoadStdMidiFile(void* Buffer, int32 BufferSize, const FString& FileName, int32 DesiredTicksPerQuarterNote = MidiConstants::kTicksPerQuarterNoteInt,
+		MidiConstants::EMidiTextEventEncoding InTextEncoding = MidiConstants::EMidiTextEventEncoding::UTF8);
+	/** A method for importing a standard midi file, with the option of providing a pointer to an FSongMaps instance that will be populated during the load. */
+	void LoadStdMidiFile(TSharedPtr<FArchive> Archive, const FString& Filename, int32 DesiredTicksPerQuarterNote = MidiConstants::kTicksPerQuarterNoteInt,
+		MidiConstants::EMidiTextEventEncoding InTextEncoding = MidiConstants::EMidiTextEventEncoding::UTF8);
+
+	/** A method for exporting the midi track data to a standard midi file. */
+	void SaveStdMidiFile(const FString& FilePath);
+	/** A method for exporting the midi track data to a standard midi file. */
+	void SaveStdMidiFile(TSharedPtr<FArchive> Archive, const FString& Filename);
+
+	FMidiTrack* AddTrack(const FString& Name);
+
+	void Empty();
+
+	void SetConductorTrack(const FTempoMap*, const FBarMap*);
+
+	void SortAllTracks();
+
+	int32 GetNumTracks() const { return TheMidiData.Tracks.Num(); }
+	const FMidiTrack* GetTrack(int32 Index) const { return &TheMidiData.Tracks[Index]; }
+	FMidiTrack* GetTrack(int32 Index) { return &TheMidiData.Tracks[Index]; }
+	const FMidiTrack* FindTrackByName(const FString& TrackName) const;
+	int32 FindTrackIndexByName(const FString& TrackName) const;
+
+	// Find the text event and return the tick or -1 if not found.
+	int32 FindTextEvent(const FString& EventText, const FString* TrackName = nullptr);
+	// Find the text event and return the tick or -1 if not found.
+	int32 FindTextEvent(const FString& EventText, int32 TrackIndex);
+
+	// Find all the ticks at which the text event occurs
+	TArray<int32> FindAllTextEvents(const FString& EventText, const FString* TrackName = nullptr);
+	void FindAllTextEvents(const FString& EventText, int32 TrackIndex, TArray<int32> OutIndexes);
+
+	FMidiTrackList& GetTracks() { return TheMidiData.Tracks; }
+	const FMidiTrackList& GetTracks() const { return TheMidiData.Tracks; }
+
+	/**
+	 * This function must be called if any changes are made to any of the tracks of this midi file. It will 
+	 * do various internal tasks necessary to assure the midi data is consistent and "playable".
+	 */
+	void TracksChanged();
+
+	int32  GetLastEventTick() const { return TheMidiData.LastEventTick; }
+
+#if WITH_EDITORONLY_DATA
+	// Import data for this MidiFileAsset
+	// updated during Import in MidiFile
+	UPROPERTY(VisibleAnywhere, Instanced, Category = ImportSettings)
+	TObjectPtr<UAssetImportData> AssetImportData;
+	FString GetImportedSrcFilePath() const;
+
+	//The Start Bar of a Midi File
+	UPROPERTY(EditDefaultsOnly,BlueprintReadOnly,Category = "Midi File Start Bar")
+	int32 StartBar = 1;
+
+	//StartBar UPROPERTY callback 
+	virtual void PostEditChangeProperty(FPropertyChangedEvent& PropertyChangedEvent) override;
+#endif
+
+	int32 GetStartBar() { return TheMidiData.SongMaps.GetBarMap().GetStartBar(); };
+
+	//Midi File Conform Option, keep track if the files' length are conformed during import
+	UPROPERTY(BlueprintReadOnly, Category = "MidiFile")
+	bool bLengthRoundedDown = false;
+
+	UPROPERTY(BlueprintReadOnly, Category = "MidiFile")
+	bool bLengthRoundedUp = false;
+
+	UPROPERTY(BlueprintReadOnly, Category = "MidiFile")
+	bool bLengthRoundedToNearest = false;
+
+	/**
+	* Takes a file conform option (currently the last integer bar or the last integer beat),determine whether or not 
+	* the file needs to be conformed/are already conformed.
+	* @param A midi file conform option (last bar or last beat)
+	* @return A bool indicating the file should or should not be conformed
+	*/
+	bool ShouldConformMidiFileLength(EMidiFileLengthConformOption Option);
+
+	/**
+	* Conform Midi File length(bars) to an integer number of bars:
+	* LengthTicks and LengthBars in SongLengthData are rounded up to an integer number of bars
+	* if the midi file has a fractional bar length or ends on a fractional beat, 
+	* this function takes a midi conform option (Round Down, and Round Up,Round to Nearest):
+	* Round Down: modifies/moves the raw midi events in every midi track of a midi file to the previous integer bar
+	*			  and remove excessive note on/note off pairs, control change events, pitch bend events, etc. at the tick they are moved to
+	* Round Up:	since file length is automatically rounded up, this option does NOT do or change anything to the Midi events in file
+	* Round to Nearest: rounding the file length to the nearest integer (either round down or round up), depending on the file's fractional length
+	* e.g. if a midi file has length 5.875 bars, Round Down conforms it to 5 bars, Round Up conforms it to 6 bars,
+	* and Round To Nearest conforms it to 6 bars (5.875 > 5.5)
+	*/
+	void ConformMidiFileLength(EMidiFileLengthConformOption Option);
+
+	const FSongMaps* GetSongMaps() const { return &TheMidiData.SongMaps; }
+	FSongMaps* GetSongMaps() { return &TheMidiData.SongMaps; }
+
+	// IAudioProxyDataFactory
+	virtual TSharedPtr<Audio::IProxyData> CreateProxyData(const Audio::FProxyDataInitParams& InitParams) override;
+
+	// UObject
+	void BeginDestroy() override;
+
+protected:
+	UPROPERTY()
+	FMidiFileData TheMidiData;
+
+
+	TSharedPtr<FMidiFileData> RenderableCopyOfMidiFileData;
+};
+
+class HARMONIXMIDI_API FMidiFileProxy : public Audio::TProxyData<FMidiFileProxy>
+{
+public:
+	IMPL_AUDIOPROXY_CLASS(FMidiFileProxy);
+
+	explicit FMidiFileProxy(TSharedPtr<FMidiFileData>& Data)
+		: MidiFileDataPtr(Data)
+	{}
+
+	TSharedPtr<FMidiFileData> GetMidiFile()
+	{
+		return MidiFileDataPtr;
+	}
+
+private:
+	TSharedPtr<FMidiFileData> MidiFileDataPtr;
+};
diff --git a/Engine/Plugins/Runtime/Harmonix/Source/HarmonixMidi/Public/HarmonixMidi/MidiMsg.h b/Engine/Plugins/Runtime/Harmonix/Source/HarmonixMidi/Public/HarmonixMidi/MidiMsg.h
new file mode 100644
--- /dev/null
+++ b/Engine/Plugins/Runtime/Harmonix/Source/HarmonixMidi/Public/HarmonixMidi/MidiMsg.h
@@ -0,0 +1,163 @@
+// Copyright Epic Games, Inc. All Rights Reserved.
+#pragma once
+#include "HarmonixMidi/MidiConstants.h"
+
+#include "MidiMsg.generated.h"
+
+class FMidiWriter;
+struct FMidiTrack;
+
+using FMidiTextRepository = TArray<FString>;
+
+#pragma pack(push,1)
+
+/**
+	* A class representing a single standard midi message. This class is tiny and tightly packed
+	* to keep the midi data footprint small on disk and in memory. 
+	*/
+USTRUCT()
+struct HARMONIXMIDI_API FMidiMsg
+{
+	GENERATED_BODY()
+
+public:
+	enum class EType : uint8
+	{
+		Std     = 1,
+		Tempo   = 2,
+		TimeSig = 4,
+		Text    = 8,
+		Runtime = 16
+	};
+
+	static FMidiMsg CreateNoteOn(int32 Channel, int32 Note, int32 Velocity);
+	static FMidiMsg CreateNoteOff(int32 Channel, int32 Note);
+	static FMidiMsg CreateText(uint16 InTextIndex, uint8 InTextType);
+	static float GetPitchBendFromData(int8 Data1, int8 Data2);
+
+	EType MsgType() const   { return Type;                   }
+	bool  IsStd() const     { return Type == EType::Std;     }
+	bool  IsRuntime() const { return Type == EType::Runtime; }
+	bool  IsText() const    { return Type == EType::Text;    }
+	bool  IsNoteMessage() const { return IsNoteOn() || IsNoteOff() || IsAllNotesOff() || IsAllNotesKill(); }
+	bool  IsNoteOn() const  { return Type == EType::Std && MidiConstants::IsNoteOn(Status);  }
+	bool  IsNoteOff() const { return Type == EType::Std && (MidiConstants::IsNoteOff(Status) || (MidiConstants::IsNoteOn(Status) && Data2 == 0)); }
+	bool  IsAllNotesOff() const { return Type == EType::Runtime && Status == MidiConstants::kRuntimeAllNotesOffStatus; }
+	bool  IsAllNotesKill() const { return Type == EType::Runtime && Status == MidiConstants::kRuntimeAllNotesKillStatus; }
+
+	/** Construct a standard (std) midi message */
+	FMidiMsg(uint8 InStatus, uint8 InData1, uint8 InData2);
+	uint8 GetStdStatus()     const { check(Type == EType::Std); return Status; }
+	uint8 GetStdData1()      const { check(Type == EType::Std); return Data1; }
+	uint8 GetStdData2()      const { check(Type == EType::Std); return Data2; }
+	uint8 GetStdChannel()    const { check(Type == EType::Std); return MidiConstants::GetChannel(Status); }
+	uint8 GetStdStatusType() const { check(Type == EType::Std); return MidiConstants::GetType(Status); }
+	float GetPitchBendFromData() const;
+
+	/** Construct a midi tempo message */
+	FMidiMsg(int32 MicrosecPerQuarterNote);
+	int32 GetMicrosecPerQuarterNote() const { check(Type == EType::Tempo); return int32((uint32)MicsPerQuarterNoteH << 16 | (uint32)MicsPerQuarterNoteL); }
+
+	/** Construct a midi time signature message */
+	FMidiMsg(uint8 Numerator, uint8 Denominator);
+	uint8 GetTimeSigNumerator()  const { check(Type == EType::TimeSig); return Numerator;   }
+	uint8 GetTimeSigDenominator()  const { check(Type == EType::TimeSig); return Denominator; }
+
+	//Text
+	// Note: messages of this type are constructed using the subclass below.
+	uint16 GetTextIndex()   const { check(Type == EType::Text); return TextIndex; }
+	uint8  GetTextType()    const { check(Type == EType::Text); return TextType;  }
+
+	void SetNoteOnVelocity(uint8 Velocity);
+
+	void WriteStdMidi(int32 Tick, FMidiWriter& writer, const FMidiTrack& track) const;
+		
+	static FString ToString(const FMidiMsg& Message, const FMidiTrack* Track = nullptr);
+
+	bool Serialize(FArchive& Archive);
+
+	explicit FORCEINLINE FMidiMsg(EForceInit)
+		: Type(EType::Std)
+		, Status(0)
+		, Data1(0)
+		, Data2(0)
+	{
+	}
+
+	// Runtime Messages
+	static FMidiMsg CreateAllNotesOff()
+	{
+		FMidiMsg NewMsg;
+		NewMsg.Type = EType::Runtime;
+		NewMsg.Status = MidiConstants::kRuntimeAllNotesOffStatus;
+		return NewMsg;
+	}
+	static FMidiMsg CreateAllNotesKill()
+	{
+		FMidiMsg NewMsg;
+		NewMsg.Type = EType::Runtime;
+		NewMsg.Status = MidiConstants::kRuntimeAllNotesKillStatus;
+		return NewMsg;
+	}
+
+	FMidiMsg()
+		: Type(EType::Std)
+		, Status(0)
+		, Data1(0)
+		, Data2(0)
+	{}
+
+
+	EType Type;
+	// Use a union instead of subclassing MidiMsg so that MidiMsg is 
+	// constant size and therefore you can have a vector<MidiMsg>.
+	//
+	// All structs in this unit should be 3 bytes long. 
+	//
+	// All members of those structs (including padding) must be 
+	// initialized so that no random data appears when these are
+	// serialized!
+	union
+	{
+		//Std
+		struct
+		{
+			uint8 Status;
+			uint8 Data1;
+			uint8 Data2;
+		};
+		// Tempo
+		struct
+		{
+			uint8  MicsPerQuarterNoteH;
+			uint16 MicsPerQuarterNoteL;
+		};
+		// Time Sig
+		struct
+		{
+			uint8  Numerator;
+			uint8  Denominator;
+			uint8  ts_pad;
+		};
+		//Text
+		struct
+		{
+			uint8  TextType;
+			uint16 TextIndex;
+		};
+	};
+};
+
+static_assert(sizeof(FMidiMsg) == 4);
+
+template <>
+struct TStructOpsTypeTraits<FMidiMsg> : public TStructOpsTypeTraitsBase2<FMidiMsg>
+{
+	enum
+	{
+		WithSerializer = true,
+		WithNoInitConstructor = true,
+	};
+};
+
+#pragma pack(pop)
diff --git a/Engine/Plugins/Runtime/Harmonix/Source/HarmonixMidi/Public/HarmonixMidi/MidiPlayCursor.h b/Engine/Plugins/Runtime/Harmonix/Source/HarmonixMidi/Public/HarmonixMidi/MidiPlayCursor.h
new file mode 100644
--- /dev/null
+++ b/Engine/Plugins/Runtime/Harmonix/Source/HarmonixMidi/Public/HarmonixMidi/MidiPlayCursor.h
@@ -0,0 +1,145 @@
+// Copyright Epic Games, Inc. All Rights Reserved.
+#pragma once
+
+#include "Containers/IntrusiveDoubleLinkedList.h"
+#include "HarmonixMidi/MidiFile.h"
+#include "HarmonixMidi/MidiMsg.h"
+#include "HarmonixMidi/MidiReceiver.h"
+#include "HarmonixMidi/MidiPlayCursorMgr.h"
+#include "Misc/EnumClassFlags.h"
+
+struct FMidiPlayCursorTracker;
+
+class  UMidiFile;
+
+class HARMONIXMIDI_API FMidiPlayCursor : public TIntrusiveDoubleLinkedListNode<FMidiPlayCursor>
+{
+public:
+	using FMidiPlayCurorListType = TIntrusiveDoubleLinkedList<FMidiPlayCursor>;
+
+	FMidiPlayCursor();
+	virtual ~FMidiPlayCursor();
+
+	virtual void Reset(bool ForceNoBroadcast = false);
+	virtual bool IsDone() const;
+
+	bool Advance(bool IsLowRes = true);
+	bool AdvanceAsPreRoll();
+
+	virtual void AdvanceByTicks(bool ProcessLoops = true, bool Broadcast = true, bool IsPreRoll = false);
+	virtual void AdvanceByMs(bool ProcessLoops = true, bool Broadcast = true, bool IsPreRoll = false);
+
+	int32 GetLookaheadTicks() { return LookaheadTicks; }
+	float GetLookaheadMs() { return LookaheadMs; }
+
+	enum class EFilterPassFlags : uint32
+	{
+		None = 0,
+		MidiMessage = 1 << 0,
+		Text = 1 << 1,
+		Tempo = 1 << 2,
+		TimeSig = 1 << 3,
+		Reset = 1 << 4,
+		Loop = 1 << 5,
+		PreRollNoteOn = 1 << 6,
+		All = 0xFFFFFFFF
+	};
+
+	enum class ELookaheadType
+	{
+		Ticks = 0,
+		Time = 1
+	};
+
+	enum class ESyncOptions
+	{
+		NoBroadcastNoPreRoll,
+		BroadcastImmediately,
+		PreRollIfAtStartOrBroadcast,
+		PreRollIfAtStartOrNoBroadcast
+	};
+
+	void SetupMsLookahead(float Ms, ESyncOptions Opts);
+	void SetupTickLookahead(int32 Ticks, ESyncOptions Opts);
+
+	ELookaheadType GetLookaheadType() { return LookaheadType; }
+
+	int32  GetCurrentTick() { return CurrentTick; }
+
+	void SetMessageFilter(EFilterPassFlags InFilterPassFlags) { FilterPassFlags = InFilterPassFlags; }
+	void SetWatchTrack(int32 TrackIndex) { WatchTrack = TrackIndex; }
+	virtual void ManagerIsDetaching() {}
+
+	const FMidiPlayCursorMgr* GetOwner() const { return Owner; }
+	FMidiPlayCursorMgr* GetOwner() { return Owner; }
+
+	static const float kSmallMs;
+
+protected:
+	//////////////////////////////////////////////////////////////////////////
+	// Implement the following in your subclass to process callbacks
+	//////////////////////////////////////////////////////////////////////////
+	// Standard 1- or 2-byte MIDI message:
+	virtual void OnMidiMessage(int32 TrackIndex, int32 Tick, uint8 Status, uint8 Data1, uint8 Data2, bool IsPreroll = false) {}
+	// Text, Copyright, TrackName, InstName, Lyric, Marker, CuePoint meta-event:
+	//  "type" is the type of meta-event (constants defined in
+	//  MidiFileConstants.h)
+	virtual void OnText(int32 TrackIndex, int32 Tick, int32 TextIndex, const FString& Str, uint8 Type, bool IsPreroll = false) {}
+	// Tempo Change meta-event:
+	//  tempo is in microseconds per quarter-note
+	virtual void OnTempo(int32 TrackIndex, int32 Tick, int32 Tempo, bool IsPreroll = false) {}
+	// Time Signature meta-event:
+	//   time signature is numer/denom
+	virtual void OnTimeSig(int32 TrackIndex, int32 Tick, int32 Numerator, int32 Denominator, bool IsPreroll = false) {}
+	// Called when the look ahead amount changes or midi playback resets...
+	virtual void OnReset() {}
+	// Called when the cursor is about to loop back (or forward) from the loop end point to the loop start point
+	virtual void OnLoop(int32 LoopStartTick, int32 LoopEndTick) {}
+	// Called when a noteOn happens during a pre-roll. Example usage...
+	// A cursor playing a synth that supports note-ons with start offsets into playback may
+	// pass this note on to the synth with the preRollMs as the start offset.
+	virtual void OnPreRollNoteOn(int32 TrackIndex, int32 EventTick, int32 Tick, float PreRollMs, uint8 Status, uint8 Data1, uint8 Data2) {}
+
+protected:
+	float CurrentMs = 0.0f; // We've broadcast all events up through this Ms
+	int32 CurrentTick = 0;  // We've broadcast all events up to and including this tick
+	int32 LoopCount = 0;
+	bool  UnregisterASAP = false;
+	FMidiPlayCursorMgr* Owner = nullptr;
+	FMidiPlayCursorTracker* Tracker = nullptr;
+
+protected:
+	virtual void SeekToTick(int32 Tick);
+	virtual void SeekThruTick(int32 Tick);
+	virtual void AdvanceThruTick(int32 Tick, bool IsPreRoll);
+
+private:
+	friend struct FMidiPlayCursorTracker;
+
+	void MoveToLoopStartIfAtNotOffset(int32 NewThruTick, float NewThruMs);
+
+	void PrepareLookAheadTicks(bool ForceNoBroadcast = false);
+	void PrepareLookAheadMs(bool ForceNoBroadcast = false);
+
+	void DoAdvanceForLaggingTickCursor(bool Broadcast, bool IsPreRoll = false);
+	void DoAdvanceForLeadingTickCursor(bool Broadcast, bool ProcessLoops, bool IsPreRoll = false);
+	void DoAdvanceForLaggingMsCursor(bool Broadcast, bool IsPreRoll = false);
+	void DoAdvanceForLeadingMsCursor(bool Broadcast, bool ProcessLoops, bool IsPreRoll = false);
+
+	TArray<int32> TrackNextEventIndexs;
+	ELookaheadType LookaheadType;
+	ESyncOptions   SyncOpts;
+	int32  LookaheadTicks;             // positive is ahead, negative is behind
+	float  LookaheadMs;
+	int32  WatchTrack;
+	EFilterPassFlags FilterPassFlags;
+
+	friend class FMidiPlayCursorMgr;
+	void SetOwner(FMidiPlayCursorMgr* NewOwner, FMidiPlayCursorTracker* NewTracker, float PreRollMs = -1.0f);
+
+	void BroadcastEvent(int32 TrackIndex, const FMidiEvent& MidiEvent, bool IsPreroll = false);
+
+	void RecalcNextEventsDueToMidiChanges(FMidiPlayCursorMgr::EMidiChangePositionCorrectMode PositionMode);
+};
+
+ENUM_CLASS_FLAGS(FMidiPlayCursor::EFilterPassFlags)
