[Harmonix] Move Harmonix plugin into Engine/Plugins/Runtime to make it available to licensees in UE 5.4

diff --git a/Engine/Plugins/Runtime/Harmonix/Source/HarmonixMidi/Public/HarmonixMidi/MidiPlayCursorMgr.h b/Engine/Plugins/Runtime/Harmonix/Source/HarmonixMidi/Public/HarmonixMidi/MidiPlayCursorMgr.h
new file mode 100644
--- /dev/null
+++ b/Engine/Plugins/Runtime/Harmonix/Source/HarmonixMidi/Public/HarmonixMidi/MidiPlayCursorMgr.h
@@ -0,0 +1,164 @@
+// Copyright Epic Games, Inc. All Rights Reserved.
+#pragma once
+
+#include "HarmonixMidi/MidiFile.h"
+#include "HarmonixMidi/MidiPlayCursorTracker.h"
+#include "HarmonixMidi/SongMaps.h"
+#include "HarmonixMidi/MidiSongPos.h"
+
+class FMidiPlayCursor;
+
+class HARMONIXMIDI_API FMidiPlayCursorMgr
+{
+public:
+	FMidiPlayCursorMgr();
+	virtual ~FMidiPlayCursorMgr();
+
+	void Reset();
+
+	void AttachToTimeAuthority(const TSharedPtr<FMidiPlayCursorMgr>& InTimeAuthority);
+	void DetachFromTimeAuthority() { TimeAuthority = nullptr; }
+
+	void AttachToMidiResource(TSharedPtr<FMidiFileData> MidiDataProxy, bool ResetTrackersToStart = true, int32 PreRollBars = 0);
+	void DetachFromMidiResource();
+
+	FMidiSongPos CalculateSongPosWithOffsetMs(float Ms) const;
+
+	//////////////////////////////////////////////////////////////////////////
+	// Play cursor management
+	//////////////////////////////////////////////////////////////////////////
+	enum class ESetupOptions
+	{
+		NoBroadcastNoPreRoll,
+		BroadcastImmediately,
+		PreRollIfAtStartOrBroadcast,
+		PreRollIfAtStartOrNoBroadcast
+	};
+	void  RegisterHiResPlayCursor(FMidiPlayCursor* PlayCursor, float PreRollMs = -1.0f);
+	void  RegisterLowResPlayCursor(FMidiPlayCursor* PlayCursor, float PreRollMs = -1.0f);
+	void  UnregisterPlayCursor(FMidiPlayCursor* PlayCursor, bool WarnOnFail = true);
+	void  UnregisterAllPlayCursors();
+	bool  HasLowResCursors() const;
+	void  RecalculatePreRollDueToCursorPosition(FMidiPlayCursor* PlayCursor);
+	void  GetCursorExtentsMs(float& Earliest, float& Latest) const;
+	void  GetCursorExtentsTicks(int32& Earliest, int32& Latest) const;
+	bool  CursorsAllInPhase() const;
+	int32 GetFarthestAheadCursorTick() const;
+	int32 GetFarthestBehindCursorTick() const;
+	float GetBufferedMs() const;
+
+	// Do we loop when we hit the end of the file?
+	// Will function incorrectly if you change this after passing the end
+	static const int32 kEndTick = -1;
+	void SetLoop(int32 StartTick, int32 EndTick, bool IsDirectMappedFollower, bool IgnoringLookAhead);
+	void ClearLoop(bool IgnoringLookAhead);
+
+	void SeekTo(int32 Tick, int32 PreRollBars, bool IsRenderThread, bool IsLoop);
+	void MoveToLoopStart();
+
+	// NOTE: This next function DOES NOT affect how quickly the play cursor manager advances!
+	// The manager is advanced as desired with calls to the Advance___ functions. This simply 
+	// sets a member to inform the manager and the manager's cursors how quickly the driver of
+	// this manager is advancing time as this might be useful information for some cursors
+	// (eg. smoothing cursors)
+	void  InformOfCurrentAdvanceRate(float Rate) { CurrentAdvanceRate = Rate; }
+	float GetCurrentAdvanceRate() const { return CurrentAdvanceRate; }
+
+	//If broadcast is true, call callbacks for each midi event that we pass
+	//If broadcast is false, silently update the internals of each play cursor
+	void AdvanceHiResToMs(float Ms, bool Broadcast);
+	void AdvanceHiResByDeltaMs(float Ms, bool Broadcast);
+	void AdvanceHiResThruTick(int32 ThruTick, bool Broadcast, bool DontAdvancePastLoopEnd = false);
+	void AdvanceHiResByDeltaTick(int32 NumTicks, bool Broadcast);
+
+	// returns whether the cursors looped
+	bool AdvanceLowResCursors();
+
+	void ResetTrackers(); //reset to the beginning of the song
+
+	float GetLoopStartMs() const { return LoopStartMs; }
+	int32 GetLoopStartTick() const { return LoopStartTick; }
+	float GetLoopEndMs() const { return LoopEndMs; }
+	int32 GetLoopEndTick() const { return LoopEndTick; }
+	bool  DoesLoop() const { return Loop; }
+	bool  IsDirectMappedTimeFollower() const { return DirectMappedTimeFollower; }
+	int32 GetLengthTicks() const { return LengthTicks; }
+
+	bool   HasMidiFile() const;
+	int32  FindTrackIndexByName(const FString& name) const;
+	const FString* GetMidiFileName() const;
+
+	const UMidiFile::FMidiTrackList& Tracks() const;
+
+	const FSongMaps& GetSongMaps() const;
+	const FTempoMap& GetTempoMap() const;
+	const FBarMap& GetBarMap() const;
+
+	int32 GetCurrentHiResTick() const   { return HiResTracker.CurrentTick;  }
+	float GetCurrentHiResMs() const     { return HiResTracker.CurrentMs;    }
+	float GetElapsedHiResMs() const     { return HiResTracker.ElapsedMs;    }
+	int32 GetCurrentLowResTick() const  { return LowResTracker.CurrentTick; }
+	float GetCurrentLowResMs() const    { return LowResTracker.CurrentMs;   }
+	float GetElapsedLowResMs() const    { return LowResTracker.ElapsedMs;   }
+
+	bool IsDone() const;
+
+	void LockForMidiDataChanges();
+
+	enum class EMidiChangePositionCorrectMode
+	{
+		MaintainTick,
+		MaintainTime
+	};
+
+	void MidiDataChangeComplete(EMidiChangePositionCorrectMode positionMode, int32 PreRollBars = 0);
+	
+	/**
+	 * @brief Get the music timestamp at an absolute time in milliseconds
+	 * @param Ms - The time in milliseconds
+	 * @return The music timestamp
+	 */
+	FMusicTimestamp GetMusicTimestampAtMs(float Ms) const;
+
+protected:
+
+	friend class FMidiPlayCursor;
+
+	// REFACTOR NEEDED... the are too many places that need to know if a
+	// given tracker is hires or lowres. So the tracker needs to know which it is!
+	FMidiPlayCursorTracker HiResTracker;
+	FMidiPlayCursorTracker LowResTracker;
+
+private:
+	bool AdvanceTrackerByDeltaMs(float Ms, FMidiPlayCursorTracker& Tracker, bool IsLowRes, bool Broadcast = true);
+	void AdvanceTrackerThruTick(int32 ToTick, FMidiPlayCursorTracker& Tracker, bool IsLowRes, bool Broadcast);
+
+	void DetermineLength();
+
+	FSongMaps                 DefaultMaps;
+	UMidiFile::FMidiTrackList DefaultTracks;
+
+	TWeakPtr<FMidiPlayCursorMgr> TimeAuthority;
+
+	TSharedPtr<FMidiFileData> MidiFileData;
+	const FSongMaps* SongMaps;     // <-- This is a pointer to the actual midi song map data... allows for fast access
+
+	mutable FCriticalSection LowResCursorListLock;
+	mutable FCriticalSection HiResCursorListLock;
+	mutable bool TraversingHiResCursors  = false;
+	mutable bool TraversingLowResCursors = false;
+	float LengthMs;
+	int32 LengthTicks;
+	bool  DirectMappedTimeFollower;
+	bool  Loop;
+	float LoopOffsetTick;
+	float LoopStartMs;
+	int32 LoopStartTick;
+	float LoopEndMs;
+	int32 LoopEndTick;
+	float MsSinceLowResUpdate;
+	bool  HiResLoopedSinceLastLoResUpdate;
+	bool  InMidiChangeLock;
+	float CurrentAdvanceRate;
+};
+
diff --git a/Engine/Plugins/Runtime/Harmonix/Source/HarmonixMidi/Public/HarmonixMidi/MidiPlayCursorTracker.h b/Engine/Plugins/Runtime/Harmonix/Source/HarmonixMidi/Public/HarmonixMidi/MidiPlayCursorTracker.h
new file mode 100644
--- /dev/null
+++ b/Engine/Plugins/Runtime/Harmonix/Source/HarmonixMidi/Public/HarmonixMidi/MidiPlayCursorTracker.h
@@ -0,0 +1,49 @@
+// Copyright Epic Games, Inc. All Rights Reserved.
+#pragma once
+
+#include "Containers/IntrusiveDoubleLinkedList.h"
+
+class FMidiPlayCursor;
+
+struct HARMONIXMIDI_API FMidiPlayCursorTracker
+{
+	int32   CurrentTick; // We've broadcast all events up through this tick
+	float   CurrentMs;   // We've broadcast all events up through this Ms
+	float   ElapsedMs;
+	int32   LoopCount;
+
+	int32   EarliestCursorTick;
+	int32   LatestCursorTick;
+	float   EarliestCursorMs;
+	float   LatestCursorMs;
+
+	TIntrusiveDoubleLinkedList<FMidiPlayCursor> Cursors;
+
+	FMidiPlayCursorTracker();
+
+	bool IsAtStart();
+	void AddCursor(FMidiPlayCursor* Cursor);
+	bool ContainsCursor(FMidiPlayCursor* Cursor);
+	bool RemoveCursor(FMidiPlayCursor* Cursor);
+	void Clear();
+
+	void RecalculateExtents();
+	int32 GetFarthestAheadCursorTick() const;
+	int32 GetFarthestBehindCursorTick() const;
+
+	void MoveToLoopStart(int32 NewThruTick, float NewThruMs);
+
+	bool HasQueuedReset = true;
+	int32 NewQueuedTick = 0;
+	float NewQueuedMs = 0.0f;
+	int32 NewQueuedPreRollTick = 0;
+	float NewQueuedPreRollMs = 0.0f;
+	bool  NewQueuedBroadcast = false;
+
+	void QueueReset(int32 NewCurrentTick, float NewCurrentMs, int32 PreRollStartTick, float PreRollStartMs, bool Broadcast);
+	void HandleQueuedReset();
+	void Reset(int32 NewCurrentTick, float NewCurrentMs, int32 PreRollStartTick, float PreRollStartMs, bool Broadcast);
+	void ResetNewCursor(FMidiPlayCursor* Cursor, int32 PreRollStartTick, float PreRollStartMs, bool Broadcast);
+	void Reset(int32 NewCurrentTick, float NewCurrentMs, bool Broadcast);
+};
+
diff --git a/Engine/Plugins/Runtime/Harmonix/Source/HarmonixMidi/Public/HarmonixMidi/MidiReader.h b/Engine/Plugins/Runtime/Harmonix/Source/HarmonixMidi/Public/HarmonixMidi/MidiReader.h
new file mode 100644
--- /dev/null
+++ b/Engine/Plugins/Runtime/Harmonix/Source/HarmonixMidi/Public/HarmonixMidi/MidiReader.h
@@ -0,0 +1,159 @@
+// Copyright Epic Games, Inc. All Rights Reserved.
+#pragma once
+#include "Serialization/BufferReader.h"
+#include "HAL/FileManagerGeneric.h"
+#include "HarmonixMidi/BarMap.h"
+#include "HarmonixMidi/MidiConstants.h"
+
+/**
+* MidiReaders are helper classes for reading/parsing Standard MIDI Files.
+* they read midi bytes, turn them in to more appropriate data representations,
+* and then passes the data to an IMidiReceivers
+ */
+
+class IMidiReceiver;
+
+/**
+ * A base class for Midi Readers so that IMidiReceivers can have a generic
+ * "back pointer" to the reader that is calling them.
+ */
+class IMidiReader
+{
+public:
+	virtual ~IMidiReader() {}
+	/** Called by a receiver when it wants to format a tick into an MBT for error messaging */
+	virtual const TSharedPtr<FBarMap> GetBarMap() const = 0; 
+	virtual const FString& GetCurrentTrackName() const = 0;
+	virtual int32 GetCurrentTrackIndex() const = 0;
+	virtual const FString& GetFilename() const = 0;
+	virtual int32 GetLastTick() const = 0;
+	/** Called by a receiver when it wants to skip to the next track */
+	virtual void SkipCurrentTrack() = 0;           
+	virtual bool IsFailed() const = 0;
+};
+
+/**
+ * A concrete IMidiReader implementation that can read standard midi files.
+ * 
+ * You rarely have to interact with this class, as FMidiFile uses it internally when 
+ * you use its LoadStdMidiFile functions. 
+ * 
+ * @see FMidiFile
+ */
+class FStdMidiFileReader : public  IMidiReader
+{
+public:
+
+	FStdMidiFileReader(TSharedPtr<FArchive> Archive, const FString& InFilename, IMidiReceiver* InReceiver, 
+			int32 TicksPerQuarterNote = MidiConstants::kTicksPerQuarterNoteInt, MidiConstants::EMidiTextEventEncoding InTextEncoding = MidiConstants::EMidiTextEventEncoding::UTF8);
+	FStdMidiFileReader(const FString& FilePath, IMidiReceiver* Receiver,
+			int32 TicksPerQuarterNote = MidiConstants::kTicksPerQuarterNoteInt, MidiConstants::EMidiTextEventEncoding InTextEncoding = MidiConstants::EMidiTextEventEncoding::UTF8);
+	FStdMidiFileReader(void* Buffer, int32 BufferSize, const FString& FileName, IMidiReceiver* Receiver,
+			int32 TicksPerQuarterNote = MidiConstants::kTicksPerQuarterNoteInt, MidiConstants::EMidiTextEventEncoding InTextEncoding = MidiConstants::EMidiTextEventEncoding::UTF8);
+
+	static bool ReadStdMidiFileForReceiver(void* Buffer, int32 BufferSize, const FString& FileName, IMidiReceiver* Receiver,
+			int32 TicksPerQuarterNote = MidiConstants::kTicksPerQuarterNoteInt, MidiConstants::EMidiTextEventEncoding InTextEncoding = MidiConstants::EMidiTextEventEncoding::UTF8);
+
+	virtual ~FStdMidiFileReader() {}
+
+	void ReadAllTracks();                // Read all tracks
+	bool ReadSomeEvents(int32 NumEvents);
+	bool ReadTrack();                    // Read the next track, return false if no more
+	bool ReadEvents(int32 Count);          // Read some events, return false if no more
+	void SkipCurrentTrack() override;    // call this at any point to skip to the next track
+
+	// used by error reporting system
+	virtual const FString& GetFilename() const override { return Filename; }
+	virtual int32 GetCurrentTrackIndex() const override { return CurrentTrackIndex; }
+	virtual const FString& GetCurrentTrackName() const override { return CurrentTrackName; }
+	virtual const TSharedPtr<FBarMap> GetBarMap() const override { return BarMap; }
+
+	bool GetFailed() const { return Failed; }
+
+	virtual int32 GetLastTick() const override { return LastTick; }
+
+	virtual bool IsFailed() const override { return Failed; }
+
+private:
+	void Init();
+	void ReadNextEvent();
+	void ReadNextEventImpl();
+
+	void ReadFileHeader(FArchive& Archive);
+	void ReadTrackHeader(FArchive& Archive);
+	void ReadEvent(FArchive& Archive);
+	void ReadMidiEvent(int32 Tick, uint8 Status, uint8 Data1, FArchive& Archive);
+	void ReadSystemEvent(int32 Tick, uint8 Status, FArchive& Archive);
+	void ReadMetaEvent(int32 Tick, uint8 Type, FArchive& Archive);
+
+	// functions for sorting the midi before sending it out to the receivers
+	void QueueChannelMsg(int32 Tick, uint8 Status, uint8 Data1, uint8 Data2);
+	void ProcessMidiList();
+
+	enum class EState
+	{
+		Start,     // at the very beginning (before MThd chunk)
+		NewTrack,  // at the beginning of a track (before MTrk chunk)
+		InTrack,   // in the middle of reading a track
+		End        // done
+	};
+
+	FString Filename;
+	TSharedPtr<FArchive> InputArchive;
+	TSharedPtr<FBufferReader> BufferArchive;
+	TSharedPtr<FArchiveFileReaderGeneric> FileArchive;
+
+	IMidiReceiver* Receiver = nullptr;
+	MidiConstants::EMidiTextEventEncoding TextEncoding = MidiConstants::EMidiTextEventEncoding::Latin1;
+	EState State = EState::Start;
+
+	int32 DestinationTicksPerQuarterNote = MidiConstants::kTicksPerQuarterNoteInt;
+	float TickConversionFactor = 1.0f; // factor to convert from file's tick-per-quarternote to kTicksPerQuarterNote 
+
+	int16 NumTracks = 0;          // Number of tracks.
+	int32 TrackEndPos = 0;        // end byte position of the current track
+	int32 CurrentFileTick = 0;    // Current tick in units of File's ticks-per-quarter.
+	int32 CurrentTick = 0;        // Current tick count for current track where one quarter note is TicksPerQuarterNote.
+	int32 LastTick = 0;
+	uint8 PrevStatus = 0;         // Previous status byte (for running status).
+	int32 CurrentTrackIndex = -1; // What track are we currently reading?
+	FString CurrentTrackName;
+
+	// stuff for sorting Midi messages on the same tick
+	struct FRawMidiMsg
+	{
+		uint8 Status;
+		uint8 Data1;
+		uint8 Data2;
+	};
+	using FRawMidiList = TArray<FRawMidiMsg>;
+	FRawMidiList EventsOnSameTick;         // a list of midi msgs all at the same tick (to be sorted)
+	int32          MidiListTick = 0; // the tick of the messages in mMidiList
+	struct RawMidiLess
+	{
+		inline int32 MidiRank(uint8 Status) const
+		{
+			switch (Status & MidiConstants::kMessageTypeMask)
+			{
+			case MidiConstants::kNoteOff:  return 1;
+			case MidiConstants::kControl:  return 2;
+			case MidiConstants::kProgram:  return 3;
+			case MidiConstants::kChanPres: return 4;
+			case MidiConstants::kPitch:    return 5;
+			case MidiConstants::kPolyPres: return 6;
+			case MidiConstants::kNoteOn:   return 7;
+			default:             return 8;
+			}
+		}
+		bool operator()(const FStdMidiFileReader::FRawMidiMsg& lhs, const FStdMidiFileReader::FRawMidiMsg& rhs) const
+		{
+			return MidiRank(lhs.Status) < MidiRank(rhs.Status);
+		}
+	};
+
+	// we'll keep track of our own bar map so we can display
+	// more meaningful error/warning messages 
+	TSharedPtr<FBarMap> BarMap;
+
+	bool Failed = false;
+};
diff --git a/Engine/Plugins/Runtime/Harmonix/Source/HarmonixMidi/Public/HarmonixMidi/MidiReceiver.h b/Engine/Plugins/Runtime/Harmonix/Source/HarmonixMidi/Public/HarmonixMidi/MidiReceiver.h
new file mode 100644
--- /dev/null
+++ b/Engine/Plugins/Runtime/Harmonix/Source/HarmonixMidi/Public/HarmonixMidi/MidiReceiver.h
@@ -0,0 +1,154 @@
+// Copyright Epic Games, Inc. All Rights Reserved.
+#pragma once
+#include "Algo/ForEach.h"
+#include "Containers/Array.h"
+#include "Containers/UnrealString.h"
+
+/**
+ * MidiReceivers are used in conjunction with a IMidiReaders to read data from a
+ * MIDI file. As the IMidiReader processes the file, the appropriate member 
+ * functions on the receiver are called.
+ * 
+ * FSongMaps has a type specific IMidiReceiver that takes callbacks from the IMidiReader
+ * and constructs its internal maps. 
+ * 
+ * Games can implement their own midi receivers if they want to process midi messages 
+ * when midi files are read. In this way, the game doesn't need to worry about the 
+ * low level byte wrangling and can respond to callbacks that receive well formed 
+ * midi messages
+ */
+
+class IMidiReader;
+
+class IMidiReceiver
+{
+public:
+	virtual ~IMidiReceiver() {}
+
+	/** The reader is about to read a new MIDI file. Now is your chance to reset your state. */
+	virtual void Reset() = 0;
+
+	/** Called when the reader has determined how many tracks are in the midi data */
+	virtual void SetNumTracks(int32 Num) {}
+
+	/** called when all tracks have been read. */
+	virtual void Finalize(int32 InLastFileTick) {}
+
+	// For Midi type 0 files
+	virtual void AddDummyConductorTrack() {};
+
+	virtual void OnNewTrack(int32 NewTrackIndex) = 0;
+	virtual void OnEndOfTrack(int32 LastTick) = 0;
+	virtual void OnAllTracksRead() = 0;
+
+	/** For standard 1- or 2-byte MIDI message */
+	virtual void OnMidiMessage(int32 Tick, uint8 Status, uint8 Data1, uint8 Data2) = 0;
+
+	/**
+	 * Text, Copyright, TrackName, InstName, Lyric, Marker, CuePoint meta-event:
+	 * "type" is the type of meta-event (constants defined in MidiConstants.h) 
+	 */
+	virtual void OnText(int32 Tick, const FString& Str, uint8 Type) = 0;
+
+	/**
+	 * Tempo Change meta - event:
+	 * tempo is in microseconds per quarter-note
+	 */
+	virtual void OnTempo(int32 Tick, int32 Tempo) {}
+
+	/** Time Signature meta - event:
+	 *   time signature is numerator/denominator
+	 */
+	virtual void OnTimeSignature(int32 Tick, int32 Numerator, int32 Denominator, bool FailOnError = true) {}
+
+	/**
+	 * Called by the reader before it starts processing the midi data stream.
+	 * This saves off a back pointer that is very useful for things like writing 
+	 * good error messages when a receiver encounters unexpected data, or for 
+	 * telling the reader to skip the current track and move on.
+	 */
+	virtual void SetMidiReader(IMidiReader* InReader) { Reader = InReader; }
+
+	/** Tells the reader to Skip the rest of this track and go to the next one */
+	void SkipCurrentTrack();
+
+	const IMidiReader* GetReader() const { return Reader; }
+
+protected:
+	IMidiReader* Reader = nullptr;
+};
+
+/**
+ * An implementation of IMidiReceiver that allows one to parse a midi file
+ * and send the found data to many receivers at the same time.
+ * 
+ * Again, this can be very useful if a game has multiple systems that want to 
+ * build their own runtime data based on data found in a standard midi file.
+ *
+ * @see IMidiReceiver
+ */
+class FMidiReceiverList : public IMidiReceiver
+{
+#define _FOR_EACH_MIDI_RECEIVER(func_body) Algo::ForEach(ReceiverList, [=](IMidiReceiver* recvr)func_body);
+public:
+	virtual void Reset()
+	{
+		_FOR_EACH_MIDI_RECEIVER({ recvr->Reset(); });
+	}
+
+	virtual void Finalize(int32 LastFileTick)
+	{
+		_FOR_EACH_MIDI_RECEIVER({ recvr->Finalize(LastFileTick); });
+	}
+
+	virtual void OnNewTrack(int32 NewTrackIndex)
+	{
+		_FOR_EACH_MIDI_RECEIVER({ recvr->OnNewTrack(NewTrackIndex); });
+	}
+
+	virtual void OnEndOfTrack(int32 LastTick)
+	{
+		_FOR_EACH_MIDI_RECEIVER({ recvr->OnEndOfTrack(LastTick); });
+	}
+
+	virtual void OnAllTracksRead()
+	{
+		_FOR_EACH_MIDI_RECEIVER({ recvr->OnAllTracksRead(); });
+	}
+
+	virtual void OnMidiMessage(int32 Tick, uint8 Status, uint8 Data1, uint8 Data2)
+	{
+		_FOR_EACH_MIDI_RECEIVER({ recvr->OnMidiMessage(Tick, Status, Data1, Data2); });
+	}
+
+	virtual void OnText(int32 Tick, const FString& Str, uint8 Type)
+	{
+		_FOR_EACH_MIDI_RECEIVER({ recvr->OnText(Tick, Str, Type); });
+	}
+
+	virtual void OnTempo(int32 Tick, int32 Tempo)
+	{
+		_FOR_EACH_MIDI_RECEIVER({ recvr->OnTempo(Tick, Tempo); });
+	}
+
+	virtual void OnTimeSignature(int32 Tick, int32 Numerator, int32 Denominator, bool FailOnError = true)
+	{
+		_FOR_EACH_MIDI_RECEIVER({ recvr->OnTimeSignature(Tick, Numerator, Denominator, FailOnError); });
+	}
+
+	virtual void SetMidiReader(IMidiReader* InReader)
+	{
+		IMidiReceiver::SetMidiReader(InReader);
+		_FOR_EACH_MIDI_RECEIVER({ recvr->SetMidiReader(InReader); });
+	}
+
+	void Add(IMidiReceiver* recvr) { ReceiverList.Add(recvr); }
+
+	void Empty() { ReceiverList.Empty(); }
+
+	int32 Num() { return ReceiverList.Num(); }
+
+private:
+	TArray<IMidiReceiver*> ReceiverList;
+#undef _FOR_EACH_MIDI_RECEIVER
+};
diff --git a/Engine/Plugins/Runtime/Harmonix/Source/HarmonixMidi/Public/HarmonixMidi/MidiSongPos.h b/Engine/Plugins/Runtime/Harmonix/Source/HarmonixMidi/Public/HarmonixMidi/MidiSongPos.h
new file mode 100644
--- /dev/null
+++ b/Engine/Plugins/Runtime/Harmonix/Source/HarmonixMidi/Public/HarmonixMidi/MidiSongPos.h
@@ -0,0 +1,86 @@
+// Copyright Epic Games, Inc. All Rights Reserved.
+#pragma once
+#include "HarmonixMidi/BarMap.h"
+#include "HarmonixMidi/BeatMap.h"
+#include "HarmonixMidi/MidiConstants.h"
+#include "Math/UnrealMathUtility.h"
+#include "SectionMap.h"
+
+#include "MidiSongPos.generated.h"
+
+struct FSongMaps;
+
+/////////////////////////////////////////////////////////////////////////////
+// Position within a song (midi info)
+//
+USTRUCT(BlueprintType)
+struct HARMONIXMIDI_API FMidiSongPos
+{
+	GENERATED_BODY()
+
+public:
+	UPROPERTY(BlueprintReadOnly, Category = "MidiSongPosition")
+	float SecondsFromBarOne = 0.0f; // total seconds from bar 1 beat 1. (negative for count-in and/or pickup bars)
+	UPROPERTY(BlueprintReadOnly, Category = "MidiSongPosition")
+	float SecondsIncludingCountIn = 0.0f; // total seconds from the beginning of the musical content (ie. includes all count-in and pickup bars)
+	UPROPERTY(BlueprintReadOnly, Category = "MidiSongPosition")
+	int32 TimeSigNumerator = 4;
+	UPROPERTY(BlueprintReadOnly, Category = "MidiSongPosition")
+	int32 TimeSigDenominator = 4;
+	UPROPERTY(BlueprintReadOnly, Category = "MidiSongPosition")
+	float Tempo = 120.0f;
+	UPROPERTY(BlueprintReadOnly, Category = "MidiSongPosition")
+	float BarsIncludingCountIn = 0.0f; // total bars from the beginning of the song.
+	UPROPERTY(BlueprintReadOnly, Category = "MidiSongPosition")
+	float BeatsIncludingCountIn = 0.0f; // total beats from the beginning of the song.
+	UPROPERTY(BlueprintReadOnly, Category = "MidiSongPosition")
+	EMusicalBeatType BeatType = EMusicalBeatType::Normal;
+	UPROPERTY(BlueprintReadOnly, Category = "MidiSongPosition")
+	FMusicTimestamp Timestamp;
+
+	FSongSection CurrentSongSection;
+
+	// This version is for use when song maps are unavailable. StartBar can be used if you know the content
+	// has, or will have, etc. a pickup or count-in. So for example... 
+	// If tempo = 60 bpm, time signature = 4/4, and StartBar = -1 (two bars or count-in, bars -1 and 0), then...
+	//     InElapsedMs = 0000.0 --> Seconds = -8, ElapsedBars = 0.00, MusicPosition.Bar = -1, MusicPosition.Beat = 1
+	//     InElapsedMs = 1000.0 --> Seconds = -7, ElapsedBars = 0.25, MusicPosition.Bar = -1, MusicPosition.Beat = 2
+	//     InElapsedMs = 2000.0 --> Seconds = -6, ElapsedBars = 0.50, MusicPosition.Bar = -1, MusicPosition.Beat = 3
+	//     InElapsedMs = 3000.0 --> Seconds = -5, ElapsedBars = 0.75, MusicPosition.Bar = -1, MusicPosition.Beat = 4
+	//     InElapsedMs = 4000.0 --> Seconds = -4, ElapsedBars = 1.00, MusicPosition.Bar =  0, MusicPosition.Beat = 1
+	//     InElapsedMs = 5000.0 --> Seconds = -3, ElapsedBars = 1.25, MusicPosition.Bar =  0, MusicPosition.Beat = 2
+	//     InElapsedMs = 6000.0 --> Seconds = -2, ElapsedBars = 1.50, MusicPosition.Bar =  0, MusicPosition.Beat = 3
+	//     InElapsedMs = 7000.0 --> Seconds = -1, ElapsedBars = 1.75, MusicPosition.Bar =  0, MusicPosition.Beat = 4
+	//     InElapsedMs = 8000.0 --> Seconds =  0, ElapsedBars = 2.00, MusicPosition.Bar =  1, MusicPosition.Beat = 1
+	// etc.
+	void SetByTime(float InElapsedMs, float Bpm, int32 TimeSigNum = 4, int32 TimeSigDenom = 4, int32 StartBar = 1);
+
+	void SetByTime(float InElapsedMs, const FSongMaps& Maps);
+
+	// Low-level version for midi players / parsers that use the low-level
+	// "midi tick system" for advancing song position.
+	void SetByTick(float Tick, const FSongMaps& Maps);
+
+	void Reset()	
+	{
+		SecondsFromBarOne = 0.0f;
+		SecondsIncludingCountIn = 0.0f;
+		TimeSigNumerator = 4;
+		TimeSigDenominator = 4;
+		Tempo = 120.0f;
+		Timestamp.Reset();
+		BarsIncludingCountIn = 0.0f;
+		BeatsIncludingCountIn = 0.0f;
+		BeatType = EMusicalBeatType::Normal;
+		CurrentSongSection = FSongSection();
+	}
+
+	bool IsZero() const	{ return FMath::IsNearlyZero(BeatsIncludingCountIn, UE_KINDA_SMALL_NUMBER); }
+
+	// operators
+	bool operator<(const FMidiSongPos& rhs) const;
+	bool operator<=(const FMidiSongPos& rhs) const;
+	bool operator>(const FMidiSongPos& rhs) const;
+	bool operator>=(const FMidiSongPos& rhs) const;
+	bool operator==(const FMidiSongPos& rhs) const;
+};
diff --git a/Engine/Plugins/Runtime/Harmonix/Source/HarmonixMidi/Public/HarmonixMidi/MidiTrack.h b/Engine/Plugins/Runtime/Harmonix/Source/HarmonixMidi/Public/HarmonixMidi/MidiTrack.h
new file mode 100644
--- /dev/null
+++ b/Engine/Plugins/Runtime/Harmonix/Source/HarmonixMidi/Public/HarmonixMidi/MidiTrack.h
@@ -0,0 +1,82 @@
+// Copyright Epic Games, Inc. All Rights Reserved.
+#pragma once
+#include "HarmonixMidi/MidiEvent.h"
+#include "HarmonixMidi/MidiMsg.h"
+
+#include "MidiTrack.generated.h"
+
+using FMidiEventList = TArray<FMidiEvent>;
+class FMidiWriter;
+
+/**
+	* An FMidiTrack is a collection of FMidiEvents in chronological order.  
+	* 
+	* It can be created dynamically or be the end result of importing a 
+	* standard midi file.
+	*/
+USTRUCT(BlueprintType)
+struct HARMONIXMIDI_API FMidiTrack
+{
+	GENERATED_BODY()
+
+public:
+	FMidiTrack();
+	FMidiTrack(const FString& name);
+
+	const FMidiEventList& GetEvents() const;  		// this will ASSERT if events aren't sorted!
+	const FMidiEventList& GetUnsortedEvents() const; // this won't sort
+
+	/** Do not call unless you know what you're doing! */
+	FMidiEventList& GetRawEvents();
+
+	int32 GetNumEvents() const { return (int32)Events.Num(); }
+	const FMidiEvent* GetEvent(int32 index) const { return &Events[index]; }
+
+	void SetName(const FString& InName);
+	const FString* GetName() const;
+
+	void AddEvent(const FMidiEvent& Event);
+
+	/**
+		* This will not move the event's location in the list so that iterators into
+		* Events do note get screwed up, BUT it CAN result in the midi events not being 
+		* sorted. So if you call this function one or more times you should call Sort()
+		* after doing so.
+		* @see Sort
+		*/
+	void ChangeTick(FMidiEventList::TIterator Iterator, int32 NewTick);
+
+	void WriteStdMidi(FMidiWriter& writer);
+
+	void Sort();
+
+	void Empty() { Events.Empty(); }
+	void ClearEventsAfter(int32 Tick, bool IncludeTick);
+	void ClearEventsBefore(int32 Tick, bool IncludeTick);
+
+	int32 GetPrimaryMidiChannel() const { return PrimaryMidiChannel; }
+
+	int32 CopyEvents(FMidiTrack& SourceTrack, int32 FromTick, int32 ThruTick, int32 TickOffset = 0,
+		bool Clear = true, int32 MinNote = 0, int32 MaxNote = 127, int32 NoteTranspose = 0, bool FilterTrackName = true);
+
+	const FString& GetTextAtIndex(int32 Index) const
+	{
+		return Strings[Index];
+	}
+
+	uint16 AddText(const FString& Str) { return (uint16)Strings.AddUnique(Str); }
+
+	FMidiTextRepository* GetTextRepository() { return &Strings; }
+	const FMidiTextRepository* GetTextRepository() const { return &Strings; }
+	FString GetTextForMsg(const FMidiMsg& Message) const { check(Message.MsgType() == FMidiMsg::EType::Text); return GetTextAtIndex(Message.GetTextIndex()); }
+
+private:
+	UPROPERTY()
+	TArray<FMidiEvent> Events; // All the midi events
+	UPROPERTY()
+	bool Sorted;               // Are the events sorted.
+	UPROPERTY()
+	int32 PrimaryMidiChannel;  // The midi channel of the first event on the track!
+	UPROPERTY()
+	TArray<FString> Strings;
+};
diff --git a/Engine/Plugins/Runtime/Harmonix/Source/HarmonixMidi/Public/HarmonixMidi/MidiVoiceId.h b/Engine/Plugins/Runtime/Harmonix/Source/HarmonixMidi/Public/HarmonixMidi/MidiVoiceId.h
new file mode 100644
--- /dev/null
+++ b/Engine/Plugins/Runtime/Harmonix/Source/HarmonixMidi/Public/HarmonixMidi/MidiVoiceId.h
@@ -0,0 +1,118 @@
+// Copyright Epic Games, Inc. All Rights Reserved.
+#pragma once
+#include "HarmonixMidi/MidiConstants.h"
+#include "HarmonixMidi/MidiMsg.h"
+
+class HARMONIXMIDI_API FMidiVoiceGeneratorBase
+{
+public:
+	FMidiVoiceGeneratorBase();
+	uint32 GetIdBits() const { return IdBits; }
+
+	static const uint32 kIdWidth = 20;
+
+private:
+	uint32 IdBits;
+	static FCriticalSection GeneratorIdLock;
+	static uint32 NextGeneratorId;
+};
+
+class HARMONIXMIDI_API FMidiVoiceId
+{
+public:
+	FMidiVoiceId()
+		: Id(0)
+	{}
+
+	FMidiVoiceId(const FMidiVoiceGeneratorBase* Generator, const FMidiMsg& FromMessage)
+	{
+		static_assert(FMidiVoiceGeneratorBase::kIdWidth <= 20);
+		uint8 NotePart = (FromMessage.IsNoteOn() || FromMessage.IsNoteOff()) ? FromMessage.GetStdData1() : 0;
+		uint8 ChPart   = (FromMessage.IsStd()) ? FromMessage.GetStdChannel() : 0;
+		Id = (Generator ? Generator->GetIdBits() : 0) | ((uint32)(ChPart & 0xF) << 8) | (uint32)NotePart;
+	}
+
+	FMidiVoiceId(uint32 GeneratorId, const FMidiMsg& FromMessage)
+	{
+		static_assert(FMidiVoiceGeneratorBase::kIdWidth <= 20);
+		uint8 NotePart = (FromMessage.IsNoteOn() || FromMessage.IsNoteOff()) ? FromMessage.GetStdData1() : 0;
+		uint8 ChPart = (FromMessage.IsStd()) ? FromMessage.GetStdChannel() : 0;
+		Id = GeneratorId | ((uint32)(ChPart & 0xF) << 8) | (uint32)NotePart;
+	}
+
+	FMidiVoiceId(const FMidiVoiceId& Other)
+		: Id (Other.Id)
+	{}
+
+	FMidiVoiceId(FMidiVoiceId&& Other)
+		: Id(Other.Id)
+	{
+		Other.Id = 0;
+	}
+
+	void GetGeneratorChannelAndNote(uint32& GeneratorOut, uint8& ChOut, uint8& NoteOut) const
+	{
+		GeneratorOut = Id & 0xFFFFFC00;
+		GetChannelAndNote(ChOut, NoteOut);
+	}
+
+	void GetChannelAndNote(uint8& ChOut, uint8& NoteOut) const
+	{
+		ChOut = (Id >> 8) & 0xF;
+		NoteOut = Id & 0xFF;
+	}
+
+	FMidiVoiceId& operator=(const FMidiVoiceId& Other)
+	{
+		Id = Other.Id;
+		return *this;
+	}
+
+	FMidiVoiceId& operator=(uint32 RawId)
+	{
+		Id = RawId;
+		return *this;
+	}
+
+	operator bool() const
+	{
+		return Id != 0;
+	}
+
+	bool operator==(const FMidiVoiceId& Other) const
+	{
+		return Other.Id == Id;
+	}
+
+	bool operator==(uint32 RawId) const
+	{
+		return RawId == Id;
+	}
+
+	static FMidiVoiceId Any()
+	{
+		return FMidiVoiceId(std::numeric_limits<uint32>::max());
+	}
+
+	static FMidiVoiceId None()
+	{
+		return FMidiVoiceId();
+	}
+
+	friend uint32 GetTypeHash(const FMidiVoiceId& VoiceId);
+
+private:
+	FMidiVoiceId(uint32 RawId)
+		: Id(RawId)
+	{}
+
+	uint32 Id;
+};
+
+FORCEINLINE uint32 GetTypeHash(const FMidiVoiceId& VoiceId)
+{
+	return VoiceId.Id;
+}
+
+static_assert(sizeof(FMidiVoiceId) == 4);
+
diff --git a/Engine/Plugins/Runtime/Harmonix/Source/HarmonixMidi/Public/HarmonixMidi/SectionMap.h b/Engine/Plugins/Runtime/Harmonix/Source/HarmonixMidi/Public/HarmonixMidi/SectionMap.h
new file mode 100644
--- /dev/null
+++ b/Engine/Plugins/Runtime/Harmonix/Source/HarmonixMidi/Public/HarmonixMidi/SectionMap.h
@@ -0,0 +1,77 @@
+// Copyright Epic Games, Inc. All Rights Reserved.
+#pragma once
+#include "HarmonixMidi/MusicMapBase.h"
+#include "HarmonixMidi/MidiConstants.h"
+
+#include "SectionMap.generated.h"
+
+/**
+ * A section in a piece of music has a name, a starting point, and a length
+ */
+USTRUCT(BlueprintType)
+struct HARMONIXMIDI_API FSongSection : public FMusicMapTimespanBase
+{
+	GENERATED_BODY()
+
+public:
+	static constexpr bool DefinedAsRegions = true;
+
+	FSongSection() {}
+
+	FSongSection(const FString& InName, int32 InStartTick, int32 InLengthTicks = 1)
+		: FMusicMapTimespanBase(InStartTick, InLengthTicks)
+		, Name(InName)
+	{}
+
+	UPROPERTY(BlueprintReadOnly, Category = "SongSection")
+	FString Name;
+};
+
+/**
+ * A map of sections in a piece of music
+ */
+USTRUCT()
+struct HARMONIXMIDI_API FSectionMap
+{
+	GENERATED_BODY()
+
+public:
+	FSectionMap()
+		: TicksPerQuarterNote(MidiConstants::kTicksPerQuarterNoteInt)
+	{}
+
+	void Finalize(int32 LastTick);
+
+	void Empty();
+	void Copy(const FSectionMap& Other, int32 StartTick = 0, int32 EndTick = -1);
+	bool IsEmpty() const;
+
+	/** Called by the midi file importer before map points are added to this map */
+	void SetTicksPerQuarterNote(int32 InTicksPerQuarterNote)
+	{
+		TicksPerQuarterNote = InTicksPerQuarterNote;
+	}
+
+	bool AddSection(const FString& Name, int32 StartTick, int32 LengthTicks, bool SortNow = true);
+
+	int32 TickToSectionIndex(int32 Tick) const;
+	const FSongSection* TickToSection(int32 Tick) const;
+	int32 GetSectionStartTick(const FString& Name) const;
+	int32 GetSectionStartTick(int32 SectionIndex) const;
+
+	const TArray<FSongSection>& GetSections() const { return Points; }
+	int32 GetNumSections() const { return Points.Num(); }
+	const FSongSection* GetSection(int32 SectionIndex) const;
+	FString GetSectionName(int32 SectionIndex) const;
+	FString GetSectionNameAtTick(int32 Tick) const;
+	void  GetSectionNames(TArray<FString>& Names) const;
+	int32   FindSectionIndex(const FString& Name) const;
+	const FSongSection* FindSectionInfo(const FString& Name) const;
+
+protected:
+	UPROPERTY()
+	int32 TicksPerQuarterNote;
+	UPROPERTY()
+	TArray<FSongSection> Points;
+};
+
diff --git a/Engine/Plugins/Runtime/Harmonix/Source/HarmonixMidi/Public/HarmonixMidi/SmoothedMidiPlayCursor.h b/Engine/Plugins/Runtime/Harmonix/Source/HarmonixMidi/Public/HarmonixMidi/SmoothedMidiPlayCursor.h
new file mode 100644
--- /dev/null
+++ b/Engine/Plugins/Runtime/Harmonix/Source/HarmonixMidi/Public/HarmonixMidi/SmoothedMidiPlayCursor.h
@@ -0,0 +1,44 @@
+// Copyright Epic Games, Inc. All Rights Reserved.
+#pragma once
+
+#include "MidiPlayCursor.h"
+#include "Harmonix/LocalMinimumMagnitudeTracker.h"
+#include "Harmonix/VariableSpeedTimer.h"
+#include "MidiSongPos.h"
+
+class HARMONIXMIDI_API FSmoothedMidiPlayCursor : public FMidiPlayCursor
+{
+public:
+	FSmoothedMidiPlayCursor();
+
+	void  SetSmoothingLatencyMs(float ms) { SmoothingLatencyMs = ms; }
+	float GetSmoothedMs() const           { return SmoothedMs;       }
+	float GetSmoothedTick() const         { return SmoothedTick;     }
+	const FMidiSongPos& GetCurrentSongPos() const { return CurrentSongPos; }
+
+	//** BEGIN FMidiPlayCursor
+	virtual void Reset(bool ForceNoBroadcast = false) override;
+	virtual void AdvanceByTicks(bool processLoops = true, bool broadcast = true, bool isPreRoll = false) override;
+	virtual void AdvanceByMs(bool processLoops = true, bool broadcast = true, bool isPreRoll = false) override;
+	virtual void OnLoop(int loopStartTick, int loopEndTick) override;
+	//** END FMidiPlayCursor
+
+protected:
+
+private:
+	static const int kFramesOfErrorHistory = 10;
+	FLocalMinimumMagnitudeTracker<kFramesOfErrorHistory> ErrorTracker;
+	FVariableSpeedTimer SmoothingTimer;
+
+	float SmoothingLatencyMs = 30.0f;
+	float SmoothedMsDelta    = 0.0f;
+	float SmoothedMs         = 0.0f;
+	float SmoothedTick       = 0.0f;
+	bool  LoopedThisPass     = false;
+	float BaseSpeed          = 1.0f;
+	FMidiSongPos CurrentSongPos;
+
+	void  SetSpeed(float NewSpeed);
+	void  SyncSmoothingTimer();
+	void  UpdateSongPosition();
+};
diff --git a/Engine/Plugins/Runtime/Harmonix/Source/HarmonixMidi/Public/HarmonixMidi/SongMaps.h b/Engine/Plugins/Runtime/Harmonix/Source/HarmonixMidi/Public/HarmonixMidi/SongMaps.h
new file mode 100644
--- /dev/null
+++ b/Engine/Plugins/Runtime/Harmonix/Source/HarmonixMidi/Public/HarmonixMidi/SongMaps.h
@@ -0,0 +1,172 @@
+// Copyright Epic Games, Inc. All Rights Reserved.
+#pragma once
+#include "HarmonixMidi/TempoMap.h"
+#include "HarmonixMidi/BarMap.h"
+#include "HarmonixMidi/BeatMap.h"
+#include "HarmonixMidi/ChordMap.h"
+#include "HarmonixMidi/SectionMap.h"
+
+#include "SongMaps.generated.h"
+
+class IMidiReader;
+class FStdMidiFileReader;
+class FSongMapReceiver;
+class UMidiFile;
+
+USTRUCT()
+struct FSongLengthData
+{
+	GENERATED_BODY()
+
+public:
+	UPROPERTY()
+	int32 LengthTicks = 0;
+	UPROPERTY()
+	int32 LengthBars = 0;
+	UPROPERTY()
+	int32 LastTick = 0;
+};
+
+/**
+ * FSongMaps encapsulates a number of other musical/midi map types
+ * that are very useful for musical gameplay and interactivity. 
+ * 
+ * With this class and the current playback position of a piece of music you can
+ * do things like determine the current Bar | Beat | Tick, song section, tempo,
+ * chord, etc.
+ */
+USTRUCT(BlueprintType)
+struct HARMONIXMIDI_API FSongMaps
+{
+	GENERATED_BODY()
+
+public:
+	FSongMaps();
+
+	void Init(int32 InTicksPerQuarterNote);
+	void Copy(const FSongMaps& Other, int32 StartTick = 0, int32 EndTick = -1);
+
+	// For importing...
+	bool LoadFromStdMidiFile(const FString& FilePath);
+	bool LoadFromStdMidiFile(void* Buffer, int32 BufferSize, const FString& Filename);
+	bool LoadFromStdMidiFile(TSharedPtr<FArchive> Archive, const FString& Filename);
+
+	float TickToMs(int32 Tick) const;
+	int32 MsToTick(float Ms) const;
+	float GetCountInSeconds() const;
+
+	// tracks
+	TArray<FString>&       GetTrackNames() { return TrackNames; }
+	const TArray<FString>& GetTrackNames() const { return TrackNames; }
+	FString                GetTrackName(int32 Index) const;
+	bool                   TrackNamesIsEmpty() const { return TrackNames.IsEmpty(); }
+	void                   EmptyTrackNames() { TrackNames.Empty(); }
+
+	// tempo
+	const FTempoInfoPoint* GetTempoInfoForMs(float Ms) const;
+	const FTempoInfoPoint* GetTempoInfoForTick(int32 Tick) const;
+	float                  GetTempoAtMs(float Ms) const;
+	float                  GetTempoAtTick(int32 Tick) const;
+	FTempoMap&             GetTempoMap() { return TempoMap; }
+	const FTempoMap&       GetTempoMap() const { return TempoMap; }
+	bool                   TempoMapIsEmpty() { return TempoMap.GetNumTempoChangePoints() == 0; }
+	void                   EmptyTempoMap() { TempoMap.Empty(); }
+
+	// pulses
+	const FBeatMapPoint*  GetBeatAtMs(float Ms) const;
+	float                 GetMsAtBeat(float Beat) const;
+	const FBeatMapPoint*  GetBeatAtTick(int32 Tick) const;
+	float                 GetMsPerBeatAtMs(float Ms) const;
+	float                 GetMsPerBeatAtTick(int32 Tick) const;
+	float                 GetFractionalBeatAtMs(float Ms) const;
+	float                 GetFractionalBeatAtTick(int32 Tick) const;
+	int32                 GetBeatIndexAtMs(float Ms) const;
+	int32                 GetBeatIndexAtTick(int32 Tick) const;
+	EMusicalBeatType      GetBeatTypeAtMs(float Ms) const;
+	EMusicalBeatType      GetBeatTypeAtTick(int32 Tick) const;
+	float                 GetBeatInPulseBarAtMs(float Ms) const;
+	float                 GetBeatInPulseBarAtTick(int32 Tick) const;
+	float                 GetNumBeatsInPulseBarAtMs(float Ms) const ;
+	float                 GetNumBeatsInPulseBarAtTick(int32 Tick) const;
+	FBeatMap&             GetBeatMap() { return BeatMap; }
+	const FBeatMap&       GetBeatMap() const { return BeatMap; }
+	bool                  BeatMapIsEmpty() { return BeatMap.GetNumMapPoints() == 0; }
+	void                  EmptyBeatMap() { BeatMap.Empty(); }
+
+	// bars
+	const FTimeSignature* GetTimeSignatureAtMs(float Ms) const;
+	const FTimeSignature* GetTimeSignatureAtTick(int32 Tick) const;
+	const FTimeSignature* GetTimeSignatureAtBar(int32 Bar) const;
+	float                 GetBarIncludingCountInAtMs(float Ms) const;
+	float                 GetBarIncludingCountInAtTick(int32 Tick) const;
+	float                 GetMsPerBarAtMs(float Ms) const;
+	float                 GetMsPerBarAtTick(int32 Tick) const;
+	FBarMap&              GetBarMap() { return BarMap; }
+	const FBarMap&        GetBarMap() const { return BarMap; }
+	bool                  BarMapIsEmpty() { return BarMap.GetNumTimeSignaturePoints() == 0; }
+	void                  EmptyBarMap() { BarMap.Empty(); }
+	void                  SetLengthTotalBars(int32 Bars);
+
+	// sections
+	float               GetSectionStartMsAtMs(float Ms) const;
+	float               GetSectionEndMsAtMs(float Ms) const;
+	const FSongSection* GetSectionAtMs(float Ms) const;
+	const FSongSection* GetSectionAtTick(int32 Tick) const;
+	const FSongSection* GetSectionWithName(const FString& Name) const;
+	FString             GetSectionNameAtMs(float Ms) const;
+	FString             GetSectionNameAtTick(int32 Tick) const;
+	float               GetSectionLengthMsAtMs(float Ms) const;
+	float               GetSectionLengthMsAtTick(int32 Tick) const;
+	FSectionMap&        GetSectionMap() { return SectionMap; }
+	const FSectionMap&  GetSectionMap() const { return SectionMap; }
+	bool                SectionMapIsEmpty() { return SectionMap.GetNumSections() == 0; }
+	void                EmptySectionMap() { SectionMap.Empty(); }
+
+	// chords
+	const FChordMapPoint* GetChordAtMs(float Ms) const;
+	const FChordMapPoint* GetChordAtTick(float Tick) const;
+	FName                 GetChordNameAtMs(float Ms) const;
+	FName                 GetChordNameAtTick(int32 Tick) const;
+	float                 GetChordLengthMsAtMs(float Ms) const;
+	float                 GetChordLengthMsAtTick(int32 Tick) const;
+	FChordProgressionMap& GetChordMap() { return ChordMap; }
+	const FChordProgressionMap&  GetChordMap() const { return ChordMap; }
+	bool                  ChordMapIsEmpty() const { return ChordMap.GetNumChords() == 0; }
+	void                  EmptyChordMap() { ChordMap.Empty(); }
+
+	void EmptyAllMaps();
+
+	FSongLengthData& GetSongLengthData() { return LengthData; }
+	const FSongLengthData& GetSongLengthData() const { return LengthData; }
+	float GetSongLengthMs() const;
+	int32 GetSongLengthBeats() const;
+	int32 GetSongLengthBars() const;
+
+	int32 GetTicksPerQuarterNote() const { return TicksPerQuarterNote; }
+
+protected:
+	friend class FSongMapReceiver;
+	UPROPERTY()
+	int32 TicksPerQuarterNote = MidiConstants::kTicksPerQuarterNoteInt;
+	UPROPERTY()
+	FTempoMap TempoMap;
+	UPROPERTY()
+	FBarMap BarMap;
+	UPROPERTY()
+	FBeatMap BeatMap;
+	UPROPERTY()
+	FSectionMap SectionMap;
+	UPROPERTY()
+	FChordProgressionMap ChordMap;
+	UPROPERTY()
+	TArray<FString>	TrackNames;
+
+private:
+	UPROPERTY()
+	FSongLengthData LengthData;
+
+	void StringLengthToMT(const FString& LengthString, int32& OutBars, int32& OutTicks);
+	bool ReadWithReader(FStdMidiFileReader& Reader);
+	bool FinalizeRead(IMidiReader* Reader);
+};
+
diff --git a/Engine/Plugins/Runtime/Harmonix/Source/HarmonixMidi/Public/HarmonixMidi/TempoMap.h b/Engine/Plugins/Runtime/Harmonix/Source/HarmonixMidi/Public/HarmonixMidi/TempoMap.h
new file mode 100644
--- /dev/null
+++ b/Engine/Plugins/Runtime/Harmonix/Source/HarmonixMidi/Public/HarmonixMidi/TempoMap.h
@@ -0,0 +1,182 @@
+// Copyright Epic Games, Inc. All Rights Reserved.
+#pragma once
+#include "HarmonixMidi/MidiConstants.h"
+#include "HarmonixMidi/MusicMapBase.h"
+
+#include "TempoMap.generated.h"
+
+/**
+ * A position in a song where the tempo is specified (the rate at which the music plays)
+ */
+USTRUCT()
+struct HARMONIXMIDI_API FTempoInfoPoint : public FMusicMapTimespanBase
+{
+	GENERATED_BODY()
+
+public:
+	static constexpr bool DefinedAsRegions = false;
+
+	FTempoInfoPoint() 
+	{}
+
+	FTempoInfoPoint(float InMs, int32 InMidiTempo, int32 InStartTick, int32 InLengthTicks = 1)
+		: FMusicMapTimespanBase(InStartTick, InLengthTicks)
+		, Ms(InMs)
+		, MidiTempo(InMidiTempo)
+	{
+	}
+
+	float GetBPM() const { return MidiConstants::MidiTempoToBPM(MidiTempo); }
+
+	UPROPERTY()
+	float Ms = 0;        // The time at which the tempo is changing
+	UPROPERTY()
+	int32 MidiTempo = 0; // microseconds per beat at that point
+
+	struct TimeLessThan
+	{
+		bool operator()(float InMs, const FTempoInfoPoint& Point) const { return InMs < Point.Ms; }
+		bool operator()(const FTempoInfoPoint& Point, float InMs) const { return Point.Ms < InMs; }
+		bool operator()(const FTempoInfoPoint& PointA, const FTempoInfoPoint& PointB) const { return PointA.Ms < PointB.Ms; }
+	};
+};
+
+/**
+ * A TempoMap that can change over time; the graph of the correspondence
+ * between time and tick has multiple line segments.
+ *
+ * The tempo changes are specified by adding (tick, tempo) pairs, 
+ */
+
+ USTRUCT()
+struct HARMONIXMIDI_API FTempoMap 
+{
+	GENERATED_BODY()
+
+public:
+
+	FTempoMap() 
+		: TicksPerQuarterNote(MidiConstants::kTicksPerQuarterNoteInt)
+	{}
+	virtual ~FTempoMap() {}
+
+	void Empty();
+	void Copy(const FTempoMap& Other, int32 StartTick = 0, int32 EndTick = -1);
+	bool IsEmpty() const;
+
+	/** Called by the midi file importer before map points are added to this map */
+	void SetTicksPerQuarterNote(int32 InTicksPerQuarterNote)
+	{
+		TicksPerQuarterNote = InTicksPerQuarterNote;
+	}
+
+/*
+	template <typename T>
+	float TickToMs(T Tick) const;
+	template <typename T>
+	float TickToSeconds(T Tick) const;
+*/
+
+	/** Get the time (in milliseconds) at a given tick: */
+	float TickToMs(float Tick) const;
+	/** Get the time (in milliseconds) at a given tick: */
+	float TickToMs(int32 Tick) const;
+
+	float TickToSeconds(float Tick) const;
+	float TickToSeconds(int32 Tick) const;
+
+	/**
+	 * Get the time (in milliseconds) at the given tick.
+	 * 
+	 * This flavor of the function can speed up iteration through midi data
+	 * when you might be making many calls to TickToMs in a row with sequencial 
+	 * tick numbers.
+	 * 
+	 * @param Tick Midi tick position
+	 * @param Index The index of the tempo point to be checked first.
+	 * @param NextIndex Filles in with the tempo point most likely to contain the next tick you are likely to request.
+	 */
+	float TickToMs(float Tick, int32 Index, int32* NextIndex) const;
+
+	/** Get the tick at a given time (in milliseconds) */
+	float MsToTick(float TimeMs) const;
+
+	/**
+		* Get the tick at the given millisecond.
+		* 
+		* This flavor of the function can speed up iteration through midi data
+		* when you might be making many calls to MsToTick in a row with sequencial 
+		* times.
+		* 
+		* @param TimeMs Millisecond
+		* @param Index The index of the tempo point to be checked first.
+		* @param NextIndex Fills in with the tempo point most likely to contain the next tick you are likely to request.
+		*/
+	float MsToTick(float TimeMs, int32 Index, int32* NextIndex) const;
+
+	inline float SecondsToTick(const float Seconds) const { return MsToTick(Seconds * 1000.f); }
+
+	/** Get the tempo in ms/quarter-note */
+	float GetMsPerQuarterNoteAtTick(int32 Tick) const;
+
+	/** Get the tempo in us/quarter-note */
+	int32 GetMicrosecondsPerQuarterNoteAtTick(int32 Tick) const;
+
+	/** Get the tempo in beats (quarter-notes)/min */
+	float GetTempoAtTick(int32 Tick) const;
+
+	/** Get the tempo info point */
+	const FTempoInfoPoint* GetTempoPointAtTick(int32 Tick) const;
+
+	/** The number of 'tempo events' in the song */
+	int32 GetNumTempoChangePoints() const;
+
+	/** Return the tick of the nth tempo change point. */
+	int32 GetTempoChangePointTick(int32 Index) const;
+
+	/** Call when you're done changing it */
+	void Finalize(int32 LastTick);
+
+	bool AddTempoInfoPoint(int32 MicrosecondsPerQuarterNote, int32 Tick, bool SortNow = true);
+
+	/** Remove every point from the given tick onward */
+	void WipeTempoInfoPoints(int32 Tick);
+
+	/** if there are no points, supply a default tempo of 120. */
+	void SupplyDefault();
+
+	// Returns the tempo info points for inspection.
+	const TArray<FTempoInfoPoint>& GetTempoPoints() const { return Points; }
+
+protected:
+	UPROPERTY()
+	int32 TicksPerQuarterNote;
+	UPROPERTY()
+	TArray<FTempoInfoPoint> Points;
+
+private:
+
+	// Given a time, find the TempoInfoPoint at or before it.
+	int32 PointIndexForTime(float TimeMs) const;
+
+	// Assuming prevInfoPoint is the last TempoInfoPoint before tick,
+	// calculate tick as time in ms
+	float TickToMsInternal(float Tick, const FTempoInfoPoint& prevTempoInfoPoint) const;
+	float MsToTickInternal(float TimeMs, const FTempoInfoPoint& prevTempoInfoPoint) const;
+
+};
+
+inline float FTempoMap::TickToMs(int32 Tick) const
+{
+	return TickToMs(static_cast<float>(Tick));
+}
+
+inline float FTempoMap::TickToSeconds(int32 Tick) const
+{
+	return TickToMs(Tick) / 1000.0f;
+}
+
+inline float FTempoMap::TickToSeconds(float Tick) const
+{
+	return TickToMs(Tick) / 1000.0f;
+}
diff --git a/Engine/Plugins/Runtime/Harmonix/Source/HarmonixMidiEditor/Private/MidiFileFactory.cpp b/Engine/Plugins/Runtime/Harmonix/Source/HarmonixMidiEditor/Private/MidiFileFactory.cpp
new file mode 100644
--- /dev/null
+++ b/Engine/Plugins/Runtime/Harmonix/Source/HarmonixMidiEditor/Private/MidiFileFactory.cpp
@@ -0,0 +1,312 @@
+// Copyright Epic Games, Inc. All Rights Reserved.
+
+#include "MidiFileFactory.h"
+#include "HarmonixMidi/MidiConstants.h"
+#include "HarmonixMidiEditorModule.h"
+#include "EditorFramework/AssetImportData.h"
+#include "Misc/MessageDialog.h"
+#include "Misc/Paths.h"
+
+UMidiFileFactory::UMidiFileFactory(const FObjectInitializer& ObjectInitializer)
+	: Super(ObjectInitializer)
+{
+	bCreateNew = false;
+	SupportedClass = UMidiFile::StaticClass();
+
+	bEditorImport = true;
+	bText = false;
+
+	Formats.Add(TEXT("mid;Standard Midi File"));
+}
+
+FText UMidiFileFactory::GetToolTip() const
+{
+	return NSLOCTEXT("Midi", "MidiImporterFactoryDescription", "Standard Midi Files exported from Digital Audio Workstations");
+}
+
+bool UMidiFileFactory::FactoryCanImport(const FString& Filename)
+{
+	return true;
+}
+
+UObject* UMidiFileFactory::FactoryCreateFile(UClass* InClass, UObject* InParent, FName InName, EObjectFlags Flags, const FString& Filename, const TCHAR* Parms, FFeedbackContext* Warn, bool& bOutOperationCanceled)
+{
+	UMidiFile* MidiFileAsset = FindObject<UMidiFile>(InParent, *InName.ToString());
+	if (!MidiFileAsset)
+	{
+		MidiFileAsset = NewObject<UMidiFile>(InParent, InClass, InName, Flags);
+	}
+	MidiFileAsset->LoadStdMidiFile(Filename, MidiConstants::kTicksPerQuarterNoteInt);
+	
+	//add the current file to the file array (for handling multi-batch import)
+	ImportedFiles.Add(MidiFileAsset);
+	
+	return MidiFileAsset;
+}
+
+bool UMidiFileFactory::CanReimport(UObject* Obj, TArray<FString>& OutFilenames)
+{
+	if (UMidiFile* AsMidiFile = Cast<UMidiFile>(Obj))
+	{
+		FString FilePath = AsMidiFile->GetImportedSrcFilePath();
+		if (FilePath.IsEmpty())
+		{
+			UE_LOG(LogHarmonixMidiEditor, Warning, TEXT("Couldn't find source path for %s!"), *AsMidiFile->GetPathName());
+			return false;
+		}
+		OutFilenames.Push(FilePath);
+		return true;
+	}
+	return false;
+}
+
+void UMidiFileFactory::SetReimportPaths(UObject* Obj, const TArray<FString>& NewReimportPaths)
+{
+	if (UMidiFile* AsMidiFile = Cast<UMidiFile>(Obj))
+	{
+		AsMidiFile->AssetImportData->UpdateFilenameOnly(FPaths::ConvertRelativePathToFull(NewReimportPaths[0]));
+	}
+}
+
+void UMidiFileFactory::CleanUp()
+{ 
+	for (int FileIndex = 0; FileIndex < ImportedFiles.Num(); FileIndex++)
+	{
+		UMidiFile* MidiFileAsset = ImportedFiles[FileIndex];
+		//only prompt this dialog if imported file(s) can be conformed by at least one conform option (round up or round down)
+		if (MidiFileAsset->ShouldConformMidiFileLength(EMidiFileLengthConformOption::RoundDown) || MidiFileAsset->ShouldConformMidiFileLength(EMidiFileLengthConformOption::RoundUp))
+		{
+			if (bShouldApplyToAll && FileIndex > 0)
+			{
+				MidiFileAsset->ConformMidiFileLength(ApplyToAllOption);
+				continue;
+			}
+			ShowConformMidiFileLengthDialog(FileIndex);
+		}
+
+	}
+	ImportedFiles.Empty();
+	bShouldApplyToAll = true;
+	ApplyToAllOption = EMidiFileLengthConformOption::RoundUp;
+}
+
+void UMidiFileFactory::ShowConformMidiFileLengthDialog(int32 MidiFileAssetIndex)
+{
+	UMidiFile* MidiFileAsset = ImportedFiles[MidiFileAssetIndex];
+	
+	//show the dialog widgets on a window
+	TSharedPtr<SWindow> ParentWindow;
+	SAssignNew(ParentWindow, SWindow)
+		.Title(FText::FromString(TEXT("Conform Midi File Length")))
+		.ClientSize(FVector2D(490, 250));
+	
+	//the custom pop up dialog that asks for conforming midi file length
+	TSharedRef<SConformMidiFileLengthDialog> ConformMidiFileLengthDialog = SNew(SConformMidiFileLengthDialog);
+	ParentWindow->SetContent(ConformMidiFileLengthDialog);
+
+	// enable the Apply To All checkbox if importing more than 1 file 
+	(ImportedFiles.Num() > 1 && MidiFileAssetIndex != ImportedFiles.Num() - 1) ? ConformMidiFileLengthDialog->ApplyToAllCheckBox->SetVisibility(EVisibility::Visible) : ConformMidiFileLengthDialog->ApplyToAllCheckBox->SetVisibility(EVisibility::Collapsed);
+	
+	//if the file has length less than 1 bar (e.g. 0.25 bars),it cannot be rounded down to 0 bar, this is checked in UMidiFile::ShouldConformFileLength(EMidiFileLengthConformOption Option)
+	//disable the Round Down checkbox and the Round To Nearest checkbox, the only available option should be Round Up
+	if (!MidiFileAsset->ShouldConformMidiFileLength(EMidiFileLengthConformOption::RoundDown) && MidiFileAsset->ShouldConformMidiFileLength(EMidiFileLengthConformOption::RoundUp))
+	{
+		ConformMidiFileLengthDialog->RoundDownCheckBox->SetEnabled(false);
+		ConformMidiFileLengthDialog->RoundDownCheckBox->SetIsChecked(ECheckBoxState::Unchecked);
+
+		ConformMidiFileLengthDialog->RoundToNearestCheckBox->SetEnabled(false);
+		ConformMidiFileLengthDialog->RoundToNearestCheckBox->SetIsChecked(ECheckBoxState::Unchecked);
+
+		ConformMidiFileLengthDialog->RoundUpCheckBox->SetIsChecked(ECheckBoxState::Checked);
+		ConformMidiFileLengthDialog->ConformOption = EMidiFileLengthConformOption::RoundUp;
+	}
+
+	//set text for text block that asks user whether or not to conform midi file length
+	FString MidiFileName = MidiFileAsset->GetName();
+	float MidiFileFractionalLength = MidiFileAsset->GetSongMaps()->GetBarIncludingCountInAtTick(MidiFileAsset->GetLastEventTick());
+	ConformMidiFileLengthDialog->AskConformFileLengthText->SetText(FText::FromString(FString::Printf(
+		TEXT("Midi File Name: %s.mid, Length: %.3f bars"),
+		*MidiFileName,
+		MidiFileFractionalLength)));
+
+	//Set callback function for Ok button
+	ConformMidiFileLengthDialog->OkButton->SetOnClicked(FOnClicked::CreateLambda([this, ParentWindow, ConformMidiFileLengthDialog, &MidiFileAsset]()-> FReply {
+		MidiFileAsset->ConformMidiFileLength(ConformMidiFileLengthDialog->ConformOption);
+		ParentWindow->RequestDestroyWindow();
+		return FReply::Handled();
+		}));
+
+	//show the dialog
+	FSlateApplication::Get().AddModalWindow(ParentWindow.ToSharedRef(), nullptr);
+		
+	//remember the current conform option and mark that this option should be apply to all other files imported
+	if (ConformMidiFileLengthDialog->ApplyToAll)
+	{
+		ApplyToAllOption = ConformMidiFileLengthDialog->ConformOption;
+		bShouldApplyToAll = true;
+	}
+	else {
+		bShouldApplyToAll = false;
+	}
+}
+
+EReimportResult::Type UMidiFileFactory::Reimport(UObject* Obj)
+{
+	UMidiFile* AsMidiFile = Cast<UMidiFile>(Obj);
+	if (!AsMidiFile)
+		return EReimportResult::Failed;
+
+	FString ReimportPath = AsMidiFile->GetImportedSrcFilePath();
+	if (ReimportPath.IsEmpty())
+	{
+		UE_LOG(LogHarmonixMidiEditor, Warning, TEXT("Failed to reimport midi file: %s"), *AsMidiFile->GetFullName());
+		return EReimportResult::Failed;
+	}
+	
+	AsMidiFile->LoadStdMidiFile(ReimportPath, MidiConstants::kTicksPerQuarterNoteInt);
+
+	//Allow file length to be reconformed after reimport 
+	AsMidiFile->bLengthRoundedDown = false;
+	AsMidiFile->bLengthRoundedUp = false;
+	AsMidiFile->bLengthRoundedToNearest = false;
+
+	return EReimportResult::Succeeded;
+	
+}
+
+void SConformMidiFileLengthDialog::Construct(const FArguments& InArgs)
+{
+	ChildSlot
+	[
+		SNew(SBorder)
+		.HAlign(HAlign_Center)
+		.VAlign(VAlign_Center)
+		.BorderImage(FCoreStyle::Get().GetBrush("ToolPanel.GroupBorder"))
+		[
+			SNew(SVerticalBox)
+			// Text for the conform option
+			+ SVerticalBox::Slot()
+			.AutoHeight()
+			.HAlign(HAlign_Center)
+			.VAlign(VAlign_Center)
+			.Padding(10)
+			[
+				SAssignNew(AskConformFileLengthText,STextBlock)
+				.TextStyle(&MidiFileInformationStyle)
+			]
+			+ SVerticalBox::Slot()
+			.AutoHeight()
+			.HAlign(HAlign_Center)
+			.VAlign(VAlign_Center)
+			.Padding(5)
+			[
+				SNew(STextBlock)
+				.Text(FText::FromString(TEXT(
+					"Currently, all midi asset lengths must be conformed to some full number of bars."
+					"\n\n      Some or all of the midi files you are importing have fractional bar counts."
+					"\n\n                                        How would you like them conformed?\n")))
+			]
+			+ SVerticalBox::Slot()
+			.AutoHeight()
+			.Padding(5)
+			.VAlign(VAlign_Center)
+			.HAlign(HAlign_Center)
+			[
+				// Round Up checkbox
+				SNew(SHorizontalBox)
+				+ SHorizontalBox::Slot()
+				.Padding(5, 5)
+				.HAlign(HAlign_Center)
+				.VAlign(VAlign_Center)
+				.AutoWidth()
+				[
+					SAssignNew(RoundUpCheckBox, SCheckBox)
+					.ToolTipText(FText::FromString(TEXT("this does NOT change/add events in file because file length is automatically rounded up to the next integer bar after import")))
+					.OnCheckStateChanged(this, &SConformMidiFileLengthDialog::HandleConformOptionCheckboxChanged, EMidiFileLengthConformOption::RoundUp)
+					.IsChecked_Lambda([this]() { return ConformOption == EMidiFileLengthConformOption::RoundUp ? ECheckBoxState::Checked : ECheckBoxState::Unchecked; })
+					.Content()
+					[
+						SNew(STextBlock)
+						.Text(FText::FromString(TEXT("Round Up")))
+					]
+				]
+				// Round Down checkbox
+				+ SHorizontalBox::Slot()
+				.HAlign(HAlign_Center)
+				.VAlign(VAlign_Center)
+				.AutoWidth()
+				.Padding(5, 5)
+				[
+					SAssignNew(RoundDownCheckBox, SCheckBox)
+					.OnCheckStateChanged(this, &SConformMidiFileLengthDialog::HandleConformOptionCheckboxChanged, EMidiFileLengthConformOption::RoundDown)
+					.IsChecked_Lambda([this]() { return ConformOption == EMidiFileLengthConformOption::RoundDown ? ECheckBoxState::Checked : ECheckBoxState::Unchecked; })
+					.Content()
+					[
+						SNew(STextBlock)
+						.ToolTipText(FText::FromString(TEXT("move Midi events in file that exceed the last integer bar to the last tick of that bar and remove excessive events on that tick")))
+						.Text(FText::FromString(TEXT("Round Down")))
+					]
+				]
+				//Round To Nearest checkbox
+				+ SHorizontalBox::Slot()
+				.HAlign(HAlign_Center)
+				.VAlign(VAlign_Center)
+				.AutoWidth()
+				.Padding(5, 5)
+				[
+					SAssignNew(RoundToNearestCheckBox, SCheckBox)
+					.OnCheckStateChanged(this, &SConformMidiFileLengthDialog::HandleConformOptionCheckboxChanged, EMidiFileLengthConformOption::Nearest)
+					.IsChecked_Lambda([this]() { return ConformOption == EMidiFileLengthConformOption::Nearest ? ECheckBoxState::Checked : ECheckBoxState::Unchecked; })
+					.Content()
+					[
+						SNew(STextBlock)
+						.ToolTipText(FText::FromString(TEXT("round Midi File length to the nearest integer bar (either rounding up or rounding down)")))
+						.Text(FText::FromString(TEXT("Round to Nearest")))
+					]
+				]
+			]
+			// "Apply to All" checkbox (visible only if ApplyToAll is true)
+			+ SVerticalBox::Slot()
+			.AutoHeight()
+			.VAlign(VAlign_Center)
+			.HAlign(HAlign_Center)
+			.Padding(5)
+			[
+				SNew(SBox)
+				[
+					SAssignNew(ApplyToAllCheckBox, SCheckBox)
+					.OnCheckStateChanged(this, &SConformMidiFileLengthDialog::HandleApplyToAllCheckboxChanged)
+					.IsChecked_Lambda([this]() { return ApplyToAll ? ECheckBoxState::Checked : ECheckBoxState::Unchecked; })
+					.Content()
+					[
+						SNew(STextBlock)
+						.Text(FText::FromString(TEXT("Apply to All")))
+					]
+				]
+			]
+			// "Ok" button
+			+ SVerticalBox::Slot()
+			.AutoHeight()
+			.Padding(15)
+			.VAlign(VAlign_Bottom)
+			.HAlign(HAlign_Center)
+			[
+				SAssignNew(OkButton, SButton)
+				.Text(FText::FromString(TEXT("Ok")))
+			]
+		]	
+	];
+}
+
+void SConformMidiFileLengthDialog::HandleConformOptionCheckboxChanged(ECheckBoxState NewState, EMidiFileLengthConformOption Option)
+{
+	if (NewState == ECheckBoxState::Checked)
+	{
+		ConformOption = Option;
+	}
+}
+
+void SConformMidiFileLengthDialog::HandleApplyToAllCheckboxChanged(ECheckBoxState NewState)
+{
+	ApplyToAll = (NewState == ECheckBoxState::Checked);
+}
diff --git a/Engine/Plugins/Runtime/Harmonix/Source/HarmonixMidiTests/Private/MidiTestUtility.cpp b/Engine/Plugins/Runtime/Harmonix/Source/HarmonixMidiTests/Private/MidiTestUtility.cpp
new file mode 100644
--- /dev/null
+++ b/Engine/Plugins/Runtime/Harmonix/Source/HarmonixMidiTests/Private/MidiTestUtility.cpp
@@ -0,0 +1,123 @@
+// Copyright Epic Games, Inc. All Rights Reserved.
+#include "MidiTestUtility.h"
+
+namespace Harmonix::Testing::Utility::MidiTestUtility
+{
+	UMidiFile* BuildMidiFile(float InFileLengthBars, int32 InNumChannels, int32 InNumTracks, int32 InTimeSigNum, int32 InTimeSigDenom, int32 InTempo)
+	{
+		UMidiFile* TheMidiFile = NewObject<UMidiFile>();
+
+		//some default values for creating midi events
+		constexpr int32 DefaultNoteNumber = 60;//C4
+		constexpr int32 DefaultNoteVelocity = 90;
+		constexpr int32 DefaultTicksPerQuarter = 960;
+		constexpr uint8 DefaultControllerID = 69; //Hold Pedal
+		constexpr uint8 DefaultControlValue = 127;
+		constexpr uint8 DefaultPitchBendValueLSB = 64;
+		constexpr uint8 DefaultPitchBendValueMSB = 127;
+		constexpr uint8 DefaultPolyPresValue = 127;
+
+		// Set initial tempo and time signature
+		TheMidiFile->GetSongMaps()->GetTempoMap().AddTempoInfoPoint(MidiConstants::BPMToMidiTempo(InTempo), 0);
+		TheMidiFile->GetSongMaps()->GetBarMap().AddTimeSignatureAtBarIncludingCountIn(0, InTimeSigNum, InTimeSigDenom);
+		TheMidiFile->GetSongMaps()->GetBarMap().SetTicksPerQuarterNote(DefaultTicksPerQuarter);
+
+		//Add tracks to Midi file according to input argument
+		for (int32 Track = 1; Track < InNumTracks; ++Track)
+		{
+			//Add midi tracks 1 - InNumTracks to file, Track 0 is the conductor track
+			FMidiTrack* CurrentTrack = TheMidiFile->AddTrack(FString::Printf(TEXT("TestMidiTrack%i"), Track));
+
+			//Add a text message at the end of track 1, marking the file's last event tick
+			if (Track == 1)
+			{
+				uint16 TextIndex = CurrentTrack->AddText("TheEndOfTheFile");
+				int32 LastEventTickInFile = TheMidiFile->GetSongMaps()->GetBarMap().FractionalBarIncludingCountInToTick(InFileLengthBars);
+				if (InFileLengthBars == (int32)InFileLengthBars && InFileLengthBars != 0.f)
+				{
+					//if the input bar length is an integer number, the last event tick is 1 tick before the next bar
+					//(InFileLengthBars will end up being the very beginning of the next bar)
+					LastEventTickInFile -= 1;
+				}
+				FMidiEvent TextEventAtTheEnd(LastEventTickInFile, FMidiMsg::CreateText(TextIndex, MidiConstants::kMeta_Text));
+				CurrentTrack->AddEvent(TextEventAtTheEnd);
+			}
+
+		}
+		//update song length information
+		TheMidiFile->TracksChanged();
+
+		//To keep up with the results of imported Midi files, we ALWAYS round the file's length up to an
+		//interger number of bars
+		TheMidiFile->ConformMidiFileLength(EMidiFileLengthConformOption::RoundUp);
+		//if the fractional part of the file's length is greater than 0.5,
+		//rounding up and rounding to the nearest will get the same results
+		if(FGenericPlatformMath::Fractional(InFileLengthBars) >= 0.5f) TheMidiFile->bLengthRoundedToNearest = true;
+
+		return TheMidiFile;
+	}
+
+	void AddNoteOnNoteOffPairsToFile(UMidiFile* InFile, int32 InNoteNumber, int32 InNoteVelocity, int32 InTrackIndex, int32 InChannel, float InBarIndex)
+	{
+		FMidiTrack* CurrentTrack = InFile->GetTrack(InTrackIndex);
+		//tick number at the start of every bar for adding Note On/Note Off midi events
+		int32 CurrentBarTick = InFile->GetSongMaps()->GetBarMap().BarIncludingCountInToTick(InBarIndex);
+		int32 DefaultTicksPerQuarter = 960;
+		//Note On / Note Off pair
+		FMidiEvent CurrentNoteOnEvent(CurrentBarTick,
+			FMidiMsg::CreateNoteOn(InChannel, InNoteNumber, InNoteVelocity));
+		//Note Offs are 1 beat away from Note Ons (ticks per quarter note)
+		FMidiEvent CurrentNoteOffEvent(CurrentBarTick + DefaultTicksPerQuarter,
+			FMidiMsg::CreateNoteOff(InChannel, InNoteNumber));
+		CurrentTrack->AddEvent(CurrentNoteOnEvent);
+		CurrentTrack->AddEvent(CurrentNoteOffEvent);
+	}
+
+	void AddCCEventsToFile(UMidiFile* InFile, uint8 InControllerID, uint8 InControlValue, int32 InTrackIndex, int32 InChannel, float InBarIndex)
+	{
+		FMidiTrack* CurrentTrack = InFile->GetTrack(InTrackIndex);
+		//tick number at the start of every bar for adding CC events
+		int32 CurrentBarTick = InFile->GetSongMaps()->GetBarMap().BarIncludingCountInToTick(InBarIndex);
+		//Add Control Events to each channels/tracks
+		FMidiEvent CurrentCCEvent(CurrentBarTick,
+			FMidiMsg(MidiConstants::kControl + InChannel, InControllerID, InControlValue)
+		);
+		CurrentTrack->AddEvent(CurrentCCEvent);
+	}
+
+	void AddTextEventsToFile(UMidiFile* InFile, FString InText, int32 InTrackIndex, float InBarIndex)
+	{
+		FMidiTrack* CurrentTrack = InFile->GetTrack(InTrackIndex);
+		//tick number at the start of every bar for adding text events
+		int32 CurrentBarTick = InFile->GetSongMaps()->GetBarMap().BarIncludingCountInToTick(InBarIndex);
+		//text events
+		uint16 TextIndex = CurrentTrack->AddText(InText);
+		FMidiEvent CurrentTextEvent(CurrentBarTick, FMidiMsg::CreateText(TextIndex, MidiConstants::kMeta_Text));
+		CurrentTrack->AddEvent(CurrentTextEvent);
+	}
+
+	void AddPitchEventsToFile(UMidiFile* InFile, uint8 InPitchValueLSB, uint8 InPitchValueMSB, int32 InTrackIndex, int32 InChannel, float InBarIndex)
+	{
+		FMidiTrack* CurrentTrack = InFile->GetTrack(InTrackIndex);
+		//tick number at the start of every bar for adding text events
+		int32 CurrentBarTick = InFile->GetSongMaps()->GetBarMap().BarIncludingCountInToTick(InBarIndex);
+		//Add pitch bend events to channels/tracks
+		FMidiEvent CurrentPitchEvent(CurrentBarTick,
+			FMidiMsg(MidiConstants::kPitch + InChannel, InPitchValueLSB, InPitchValueMSB)
+		);
+		CurrentTrack->AddEvent(CurrentPitchEvent);
+	}
+
+	void AddPolyPresEventsToFile(UMidiFile* InFile, uint8 InNoteNumber, uint8 InPolyPresValue, int32 InTrackIndex, int32 InChannel, float InBarIndex)
+	{
+		FMidiTrack* CurrentTrack = InFile->GetTrack(InTrackIndex);
+		//tick number at the start of every bar for adding text events
+		//TODO: BarIndex is a whole number but 
+		int32 CurrentBarTick = InFile->GetSongMaps()->GetBarMap().BarIncludingCountInToTick(InBarIndex);
+		//Add Poly Pres events to channels/tracks
+		FMidiEvent CurrentPolyPresEvent(CurrentBarTick,
+			FMidiMsg(MidiConstants::kPolyPres + InChannel, InNoteNumber, InPolyPresValue)
+		);
+		CurrentTrack->AddEvent(CurrentPolyPresEvent);
+	}
+}
\ No newline at end of file
diff --git a/Engine/Plugins/Runtime/Harmonix/Source/HarmonixMidiTests/Private/MidiTestUtility.h b/Engine/Plugins/Runtime/Harmonix/Source/HarmonixMidiTests/Private/MidiTestUtility.h
new file mode 100644
--- /dev/null
+++ b/Engine/Plugins/Runtime/Harmonix/Source/HarmonixMidiTests/Private/MidiTestUtility.h
@@ -0,0 +1,48 @@
+// Copyright Epic Games, Inc. All Rights Reserved.
+#pragma once
+
+#include "HarmonixMidi/MidiMsg.h"
+#include "HarmonixMidi/MidiFile.h"
+#include "HarmonixMidi/MidiEvent.h"
+
+namespace Harmonix::Testing::Utility::MidiTestUtility
+{
+	/**
+	 * Construct an empty UMidiFile according to:
+	 * input file length in bars (fractional), input number of midi channels, input number of midi tracks, input time signature, and input tempo (bpm)
+	 * 1 text event is added at the end of the file according to input file length to mark the last event tick
+	 */
+	UMidiFile* BuildMidiFile(float FileLengthBars, int32 NumChannels, int32 NumTracks,int32 InTimeSigNum, int32 InTimeSigDenom, int32 InTempo);
+	
+	/**
+	 * Add 1 Note On/Note Off pair to the input midi file given the input Midi note number, note velocity,
+	 * track index, channel, and fractional bar position
+	 * 
+	 */
+	void AddNoteOnNoteOffPairsToFile(UMidiFile* InFile, int32 InNoteNumber, int32 InNoteVelocity,int32 InTrackIndex, int32 InChannel,float InBarIndex);
+
+	/**
+	* Add 1 CC Event to the input midi file given the input Midi Controller ID, control value,
+	* track index, channel, and fractional bar position
+	*/
+	void AddCCEventsToFile(UMidiFile* InFile, uint8 InControllerID, uint8 InControlValue,int32 InTrackIndex, int32 InChannel, float InBarIndex);
+
+	/**
+	* Add 1 CC Event to the input midi file given the input text,
+	* track index, and fractional bar position
+	*/
+	void AddTextEventsToFile(UMidiFile* InFile, FString InText, int32 InTrackIndex, float InBarIndex);
+
+	/**
+	* Add 1 Pitch Bend Event to the input midi file given the input Pitch Bend value (LSB), Pitch Bend value (MSB),
+	* track index, channel, and fractional bar position
+	*/
+	void AddPitchEventsToFile(UMidiFile* InFile, uint8 InPitchValueLSB, uint8 InPitchValueMSB, int32 InTrackIndex, int32 InChannel,float InBarIndex);
+
+	/**
+	* Add 1 Pitch Bend Event to the input midi file given the input Midi note number, polyphonic pressure value,
+	* track index, channel, and fractional bar position
+	*/
+	void AddPolyPresEventsToFile(UMidiFile* InFile, uint8 InNoteNumber, uint8 InPolyPresValue, int32 InTrackIndex, int32 InChannel, float InBarIndex);
+
+}
\ No newline at end of file
