[Harmonix] Fix input binding bugs, which were caught by the automated MetaSounds node tests

#tests automated tests

diff --git a/Engine/Plugins/Runtime/Harmonix/Source/HarmonixMetasound/Private/HarmonixMetasound/Nodes/MultiThreadedFusionSynchronizerNode.cpp b/Engine/Plugins/Runtime/Harmonix/Source/HarmonixMetasound/Private/HarmonixMetasound/Nodes/MultiThreadedFusionSynchronizerNode.cpp
--- a/Engine/Plugins/Runtime/Harmonix/Source/HarmonixMetasound/Private/HarmonixMetasound/Nodes/MultiThreadedFusionSynchronizerNode.cpp
+++ b/Engine/Plugins/Runtime/Harmonix/Source/HarmonixMetasound/Private/HarmonixMetasound/Nodes/MultiThreadedFusionSynchronizerNode.cpp
@@ -239,12 +239,12 @@ namespace HarmonixMetasound
 				InVertexData.BindReadVertex(FusionSyncNode::Names[InstanceIndex].SyncIn.Name, IO[InstanceIndex].Sync);
 				if constexpr (NUM_CHANNELS == 1)
 				{
-					InVertexData.BindReadVertex(FusionSyncNode::Names[InstanceIndex].AudioInLeft.Name, IO[InstanceIndex].AudioInput[0]);
+					InVertexData.BindReadVertex(FusionSyncNode::Names[InstanceIndex].AudioInMono.Name, IO[InstanceIndex].AudioInput[0]);
 				}
 				else
 				{
 					InVertexData.BindReadVertex(FusionSyncNode::Names[InstanceIndex].AudioInLeft.Name, IO[InstanceIndex].AudioInput[0]);
-					InVertexData.BindReadVertex(FusionSyncNode::Names[InstanceIndex].AudioInLeft.Name, IO[InstanceIndex].AudioInput[1]);
+					InVertexData.BindReadVertex(FusionSyncNode::Names[InstanceIndex].AudioInRight.Name, IO[InstanceIndex].AudioInput[1]);
 				}
 			}
 
