[Harmonix] Apply peak on bind, remove const from reset method to avoid confusing operator caching

diff --git a/Engine/Plugins/Runtime/Harmonix/Source/HarmonixMetasound/Private/HarmonixMetasound/Nodes/PeakNode.cpp b/Engine/Plugins/Runtime/Harmonix/Source/HarmonixMetasound/Private/HarmonixMetasound/Nodes/PeakNode.cpp
--- a/Engine/Plugins/Runtime/Harmonix/Source/HarmonixMetasound/Private/HarmonixMetasound/Nodes/PeakNode.cpp
+++ b/Engine/Plugins/Runtime/Harmonix/Source/HarmonixMetasound/Private/HarmonixMetasound/Nodes/PeakNode.cpp
@@ -128,6 +128,8 @@ namespace HarmonixMetasound::Nodes::Peak
 		{
 			InVertexData.BindReadVertex(Inputs::AudioMonoName, Inputs.Audio);
 			InVertexData.BindReadVertex(Inputs::ParamEnableName, Inputs.Enable);
+
+			UpdatePeak();
 		}
 
 		virtual void BindOutputs(Metasound::FOutputVertexInterfaceData& InVertexData) override
@@ -135,20 +137,29 @@ namespace HarmonixMetasound::Nodes::Peak
 			InVertexData.BindReadVertex(Outputs::PeakName, Outputs.Peak);
 		}
 
-		void Reset(const FResetParams& ResetParams) const
+		void Reset(const FResetParams&)
 		{
 			*Outputs.Peak = 0.0f;
 		}
 
-		void Execute() const
+		void Execute()
+		{
+			UpdatePeak();
+		}
+
+	private:
+		void UpdatePeak()
 		{
 			if (*Inputs.Enable)
 			{
 				*Outputs.Peak = Audio::ArrayMaxAbsValue(*Inputs.Audio);
 			}
+			else
+			{
+				*Outputs.Peak = 0.0f;
+			}
 		}
-
-	private:
+		
 		FInputs Inputs;
 		FOutputs Outputs;
 	};
