[Harmonix] Clean FFT analyzer settings before applying to avoid crashes

#tests automated tests

diff --git a/Engine/Plugins/Runtime/Harmonix/Source/HarmonixDsp/Private/HarmonixDsp/AudioAnalysis/FFTAnalyzer.cpp b/Engine/Plugins/Runtime/Harmonix/Source/HarmonixDsp/Private/HarmonixDsp/AudioAnalysis/FFTAnalyzer.cpp
--- a/Engine/Plugins/Runtime/Harmonix/Source/HarmonixDsp/Private/HarmonixDsp/AudioAnalysis/FFTAnalyzer.cpp
+++ b/Engine/Plugins/Runtime/Harmonix/Source/HarmonixDsp/Private/HarmonixDsp/AudioAnalysis/FFTAnalyzer.cpp
@@ -23,11 +23,27 @@ namespace Harmonix::Dsp::AudioAnalysis
 		|| Old.NumResultBins != New.NumResultBins;
 	}
 
+	// Clean the settings before applying so we get reasonable values and don't crash
+	void CleanSettings(FHarmonixFFTAnalyzerSettings& InOutSettings)
+	{
+		InOutSettings.FFTSize = FMath::Clamp(InOutSettings.FFTSize, FFFTAnalyzer::MinFFTSize, FFFTAnalyzer::MaxFFTSize);
+		InOutSettings.MinFrequencyHz = FMath::Clamp(InOutSettings.MinFrequencyHz, FFFTAnalyzer::MinFrequency, FFFTAnalyzer::MaxFrequency);
+		InOutSettings.MaxFrequencyHz = FMath::Clamp(InOutSettings.MaxFrequencyHz, FFFTAnalyzer::MinFrequency, FFFTAnalyzer::MaxFrequency);
+		InOutSettings.NumResultBins = InOutSettings.NumResultBins > 0
+		? FMath::Min(InOutSettings.NumResultBins, InOutSettings.FFTSize / 2)
+		: InOutSettings.FFTSize / 2;
+		InOutSettings.OutputSettings.RiseMs = FMath::Clamp(InOutSettings.OutputSettings.RiseMs, FFFTAnalyzer::MinSmoothingTime, FFFTAnalyzer::MaxSmoothingTime);
+		InOutSettings.OutputSettings.FallMs = FMath::Clamp(InOutSettings.OutputSettings.FallMs, FFFTAnalyzer::MinSmoothingTime, FFFTAnalyzer::MaxSmoothingTime);
+	}
+
 	void FFFTAnalyzer::SetSettings(const FHarmonixFFTAnalyzerSettings& InSettings)
 	{
+		FHarmonixFFTAnalyzerSettings SettingsCopy = InSettings;
+		CleanSettings(SettingsCopy);
+		
 		FScopeLock Lock{ &SettingsGuard };
-		NeedsRecalculate = NeedsRecalculate || SettingsRequireRecalculate(Settings, InSettings);
-		Settings = InSettings;
+		NeedsRecalculate = NeedsRecalculate || SettingsRequireRecalculate(Settings, SettingsCopy);
+		Settings = SettingsCopy;
 	}
 
 	void FFFTAnalyzer::Reset()
diff --git a/Engine/Plugins/Runtime/Harmonix/Source/HarmonixDsp/Public/HarmonixDsp/AudioAnalysis/FFTAnalyzer.h b/Engine/Plugins/Runtime/Harmonix/Source/HarmonixDsp/Public/HarmonixDsp/AudioAnalysis/FFTAnalyzer.h
--- a/Engine/Plugins/Runtime/Harmonix/Source/HarmonixDsp/Public/HarmonixDsp/AudioAnalysis/FFTAnalyzer.h
+++ b/Engine/Plugins/Runtime/Harmonix/Source/HarmonixDsp/Public/HarmonixDsp/AudioAnalysis/FFTAnalyzer.h
@@ -58,6 +58,12 @@ namespace Harmonix::Dsp::AudioAnalysis
 		static_assert(FMath::IsPowerOfTwo(MaxFFTSize), "Max FFT size should be a power of two.");
 		inline static constexpr int32 MaxBins = MaxFFTSize / 2;
 
+		inline static constexpr float MinFrequency = 20.0f;
+		inline static constexpr float MaxFrequency = 20000.0f;
+
+		inline static constexpr float MinSmoothingTime = 0.0f;
+		inline static constexpr float MaxSmoothingTime = 10000.0f;
+
 		explicit FFFTAnalyzer(float InSampleRate);
 
 		/**
