[Harmonix] MIDI Filter By [REDACTED_PERSON]
- New MidiStreamTrackFilterNode that allows midi track filtering on single, dashed or wildcard tracks.

diff --git a/Engine/Plugins/Runtime/Harmonix/Source/HarmonixMetasound/Private/HarmonixMetasound/Nodes/MidiStreamTrackFilterNode.cpp b/Engine/Plugins/Runtime/Harmonix/Source/HarmonixMetasound/Private/HarmonixMetasound/Nodes/MidiStreamTrackFilterNode.cpp
new file mode 100644
--- /dev/null
+++ b/Engine/Plugins/Runtime/Harmonix/Source/HarmonixMetasound/Private/HarmonixMetasound/Nodes/MidiStreamTrackFilterNode.cpp
@@ -0,0 +1,217 @@
+// Copyright Epic Games, Inc. All Rights Reserved.
+
+#include "MetasoundExecutableOperator.h"
+#include "MetasoundFacade.h"
+#include "MetasoundNodeInterface.h"
+#include "MetasoundParamHelper.h"
+#include "MetasoundSampleCounter.h"
+#include "MetasoundStandardNodesCategories.h"
+#include "MetasoundVertex.h"
+
+#include "HarmonixMetasound/Common.h"
+#include "HarmonixMetasound/DataTypes/MidiStream.h"
+#include "HarmonixMetasound/DataTypes/MusicTransport.h"
+#include "HarmonixMetasound/Nodes/MidiStreamTrackFilterNode.h"
+
+DEFINE_LOG_CATEGORY_STATIC(LogMidiStreamTrackFilter, Log, All);
+
+#define LOCTEXT_NAMESPACE "HarmonixMetaSound"
+
+namespace HarmonixMetasound::Nodes::MidiStreamTrackFilter
+{
+	const Metasound::FNodeClassName& GetClassName()
+	{
+		static Metasound::FNodeClassName ClassName
+		{
+			HarmonixNodeNamespace,
+			"MidiStreamTrackFilter",
+			""
+		};
+		return ClassName;
+	}
+
+	class FOp : public TExecutableOperator<FOp>
+	{
+	public:
+		static const FNodeClassMetadata& GetNodeInfo();
+		static const FVertexInterface& GetVertexInterface();
+		static TUniquePtr<IOperator> CreateOperator(const FBuildOperatorParams& InParams, FBuildResults& OutResults);
+
+		FOp(const FBuildOperatorParams& InParams,
+								  const FBoolReadRef&       InEnabled,
+								  const FStringReadRef&     InTrackSelect,
+								  const FMidiStreamReadRef& InMidiStream);
+
+		virtual void BindInputs(FInputVertexInterfaceData& InVertexData) override;
+		virtual void BindOutputs(FOutputVertexInterfaceData& InVertexData) override;
+		virtual FDataReferenceCollection GetInputs() const override;
+		virtual FDataReferenceCollection GetOutputs() const override;
+
+		void Reset(const FResetParams& ResetParams);
+
+		void Execute();
+
+	private:
+		//** INPUTS
+		FBoolReadRef       EnableInPin;
+		FStringReadRef     TrackSelectInPin;
+		FMidiStreamReadRef MidiStreamInPin;
+
+		//** OUTPUTS
+		FMidiStreamWriteRef MidiStreamOutPin;
+
+		//** DATA
+		FString CurrentTrackIndexFilter;
+		FMidiStreamEventTrackChannelFilter Filter;
+	};
+
+	class FMidiStreamTrackFilterNode : public FNodeFacade
+	{
+	public:
+		FMidiStreamTrackFilterNode(const FNodeInitData& InInitData)
+			: FNodeFacade(InInitData.InstanceName, InInitData.InstanceID, TFacadeOperatorClass<FOp>())
+		{}
+		virtual ~FMidiStreamTrackFilterNode() = default;
+	};
+
+	METASOUND_REGISTER_NODE(FMidiStreamTrackFilterNode)
+		
+	const FNodeClassMetadata& FOp::GetNodeInfo()
+	{
+		auto InitNodeInfo = []() -> FNodeClassMetadata
+		{
+			FNodeClassMetadata Info;
+			Info.ClassName        = GetClassName();
+			Info.MajorVersion     = 0;
+			Info.MinorVersion     = 1;
+			Info.DisplayName      = METASOUND_LOCTEXT("MidiStreamTrackFilterNode_DisplayName", "Midi Stream Track Filter");
+			Info.Description      = METASOUND_LOCTEXT("MidiStreamTrackFilterNode_Description", "Duplicates the incoming midi stream to its output after filtering tracks from the input Midi Stream");
+			Info.Author           = PluginAuthor;
+			Info.PromptIfMissing  = PluginNodeMissingPrompt;
+			Info.DefaultInterface = GetVertexInterface();
+			Info.CategoryHierarchy.Emplace(NodeCategories::Music);
+			return Info;
+		};
+
+		static const FNodeClassMetadata Info = InitNodeInfo();
+
+		return Info;
+	}
+
+	const FVertexInterface& FOp::GetVertexInterface()
+	{
+		using namespace CommonPinNames;
+
+		static const FVertexInterface Interface(
+			FInputVertexInterface(
+				TInputDataVertex<bool>(METASOUND_GET_PARAM_NAME_AND_METADATA(Inputs::Enable), true),
+				TInputDataVertex<FString>(METASOUND_GET_PARAM_NAME_AND_METADATA(Inputs::MidiTrackIndexFilterSpecifier), FString(TEXT("*"))),
+				TInputDataVertex<FMidiStream>(METASOUND_GET_PARAM_NAME_AND_METADATA(Inputs::MidiStream))
+			),
+			FOutputVertexInterface(
+				TOutputDataVertex<FMidiStream>(METASOUND_GET_PARAM_NAME_AND_METADATA(Outputs::MidiStream))
+			)
+		);
+
+		return Interface;
+	}
+
+	TUniquePtr<IOperator> FOp::CreateOperator(const FBuildOperatorParams& InParams, FBuildResults& OutResults)
+	{
+		using namespace CommonPinNames;
+
+		const FMidiStreamTrackFilterNode& LoggerNode = static_cast<const FMidiStreamTrackFilterNode&>(InParams.Node);
+
+		const FInputVertexInterfaceData& InputData = InParams.InputData;
+		FBoolReadRef InEnabled = InputData.GetOrCreateDefaultDataReadReference<bool>(METASOUND_GET_PARAM_NAME(Inputs::Enable), InParams.OperatorSettings);
+		FStringReadRef InTrackIndexFilter = InputData.GetOrCreateDefaultDataReadReference<FString>(METASOUND_GET_PARAM_NAME(Inputs::MidiTrackIndexFilterSpecifier), InParams.OperatorSettings);
+		FMidiStreamReadRef InMidiStream = InputData.GetOrConstructDataReadReference<FMidiStream>(METASOUND_GET_PARAM_NAME(Inputs::MidiStream), InParams.OperatorSettings);
+
+		return MakeUnique<FOp>(InParams, 
+			InEnabled,
+			InTrackIndexFilter,
+			InMidiStream);
+	}
+
+	FOp::FOp(const FBuildOperatorParams& InParams,
+											   const FBoolReadRef& InEnabled,
+											   const FStringReadRef& InTrackSelect,
+											   const FMidiStreamReadRef& InMidiStream)
+		: EnableInPin(InEnabled)
+		, TrackSelectInPin(InTrackSelect)
+		, MidiStreamInPin(InMidiStream)
+		, MidiStreamOutPin(FMidiStreamWriteRef::CreateNew(InParams.OperatorSettings))
+	{
+		Reset(InParams);
+	}
+
+	void FOp::BindInputs(FInputVertexInterfaceData& InVertexData)
+	{
+		using namespace CommonPinNames;
+
+		InVertexData.BindReadVertex(METASOUND_GET_PARAM_NAME(Inputs::Enable), EnableInPin);
+		InVertexData.BindReadVertex(METASOUND_GET_PARAM_NAME(Inputs::MidiTrackIndexFilterSpecifier), TrackSelectInPin);
+		InVertexData.BindReadVertex(METASOUND_GET_PARAM_NAME(Inputs::MidiStream), MidiStreamInPin);
+	}
+
+	void FOp::BindOutputs(FOutputVertexInterfaceData& InVertexData)
+	{
+		using namespace CommonPinNames;
+
+		InVertexData.BindReadVertex(METASOUND_GET_PARAM_NAME(Outputs::MidiStream), MidiStreamOutPin);
+	}
+
+	FDataReferenceCollection FOp::GetInputs() const
+	{
+		checkNoEntry();
+		return {};
+	}
+
+	FDataReferenceCollection FOp::GetOutputs() const
+	{
+		checkNoEntry();
+		return {};
+	}
+
+	void FOp::Reset(const FResetParams& ResetParams)
+	{
+		MidiStreamOutPin->PrepareBlock();
+		
+		if (const FMidiClockReadRef* Clock = MidiStreamInPin->GetMidiClockSource())
+		{
+			MidiStreamOutPin->SetClockSource(*Clock);
+		}
+
+		CurrentTrackIndexFilter = *TrackSelectInPin;
+		
+		FString ParseTrackIndexFilterErrorMessage;
+		if (!Filter.SetTrackFilterFromString(CurrentTrackIndexFilter, ParseTrackIndexFilterErrorMessage))
+		{
+			UE_LOG(LogMidiStreamTrackFilter, Warning, TEXT("Error parsing Track Filter string: %s"), *ParseTrackIndexFilterErrorMessage);
+		}
+
+		FString ParseMidiChannelFilterErrorMessage;
+		if (!Filter.SetMidiChannelFilterFromString(TEXT("*"), ParseMidiChannelFilterErrorMessage))
+		{
+			UE_LOG(LogMidiStreamTrackFilter, Warning, TEXT("Error parsing Midi Channel Filter string: %s"), *ParseMidiChannelFilterErrorMessage);
+		}
+	}
+
+	void FOp::Execute()
+	{
+		if (CurrentTrackIndexFilter != *TrackSelectInPin)
+		{
+			CurrentTrackIndexFilter = *TrackSelectInPin;
+
+			FString ParseTrackIndexFilterErrorMessage;
+			if (!Filter.SetTrackFilterFromString(CurrentTrackIndexFilter, ParseTrackIndexFilterErrorMessage))
+			{
+				UE_LOG(LogMidiStreamTrackFilter, Warning, TEXT("Error parsing Track Filter string: %s"), *ParseTrackIndexFilterErrorMessage);
+			}
+		}
+
+		MidiStreamOutPin->Copy(MidiStreamInPin, [&](const FMidiStreamEvent& InEvent) {return Filter(InEvent); });
+	}
+}
+
+#undef LOCTEXT_NAMESPACE // "HarmonixMetaSound"
diff --git a/Engine/Plugins/Runtime/Harmonix/Source/HarmonixMetasound/Public/HarmonixMetasound/DataTypes/MidiStream.h b/Engine/Plugins/Runtime/Harmonix/Source/HarmonixMetasound/Public/HarmonixMetasound/DataTypes/MidiStream.h
--- a/Engine/Plugins/Runtime/Harmonix/Source/HarmonixMetasound/Public/HarmonixMetasound/DataTypes/MidiStream.h
+++ b/Engine/Plugins/Runtime/Harmonix/Source/HarmonixMetasound/Public/HarmonixMetasound/DataTypes/MidiStream.h
@@ -64,7 +64,7 @@ namespace HarmonixMetasound
 
 		/**
 		* Parses the FilterString to a Bitfield and assigns it to the TrackFilter
-		* if it fails, will assign MidiChannelFilter to 0
+		* if it fails, will assign TrackFilter to 0
 		*/
 		bool SetTrackFilterFromString(const FString& FilterString, FString& OutErrorMessage);
 
diff --git a/Engine/Plugins/Runtime/Harmonix/Source/HarmonixMetasound/Public/HarmonixMetasound/Nodes/MidiStreamTrackFilterNode.h b/Engine/Plugins/Runtime/Harmonix/Source/HarmonixMetasound/Public/HarmonixMetasound/Nodes/MidiStreamTrackFilterNode.h
new file mode 100644
--- /dev/null
+++ b/Engine/Plugins/Runtime/Harmonix/Source/HarmonixMetasound/Public/HarmonixMetasound/Nodes/MidiStreamTrackFilterNode.h
@@ -0,0 +1,12 @@
+// Copyright Epic Games, Inc. All Rights Reserved.
+
+#pragma once
+
+#include "MetasoundNodeInterface.h"
+
+#include "HarmonixMetasound/Common.h"
+
+namespace HarmonixMetasound::Nodes::MidiStreamTrackFilter
+{
+	const HARMONIXMETASOUND_API Metasound::FNodeClassName& GetClassName();
+}
diff --git a/Engine/Plugins/Runtime/Harmonix/Source/HarmonixMetasoundTests/Private/Nodes/MidiStreamTrackFilterNode.cpp b/Engine/Plugins/Runtime/Harmonix/Source/HarmonixMetasoundTests/Private/Nodes/MidiStreamTrackFilterNode.cpp
new file mode 100644
--- /dev/null
+++ b/Engine/Plugins/Runtime/Harmonix/Source/HarmonixMetasoundTests/Private/Nodes/MidiStreamTrackFilterNode.cpp
@@ -0,0 +1,134 @@
+// Copyright Epic Games, Inc. All Rights Reserved.
+
+#include "NodeTestGraphBuilder.h"
+#include "HarmonixDsp/AudioBuffer.h"
+#include "HarmonixMetasound/Common.h"
+#include "HarmonixMetasound/DataTypes/MidiAsset.h"
+#include "HarmonixMetasound/DataTypes/MidiStream.h"
+#include "HarmonixMetasound/Nodes/MidiStreamTrackFilterNode.h"
+#include "HarmonixMidi/MidiMsg.h"
+#include "Misc/AutomationTest.h"
+
+#if WITH_DEV_AUTOMATION_TESTS
+
+namespace HarmonixMetasoundTests::MidiStreamTrackFilterNode
+{
+	using GraphBuilder = Metasound::Test::FNodeTestGraphBuilder;
+	using namespace Metasound;
+	using namespace Metasound::Frontend;
+	using namespace HarmonixMetasound;
+
+	UMidiFile* BuildMidiFile()
+	{
+		// Make a midi file...
+		UMidiFile* TheMidi = NewObject<UMidiFile>();
+		// Set the initial tempo...
+		TheMidi->GetSongMaps()->GetTempoMap().AddTempoInfoPoint(MidiConstants::BPMToMidiTempo(120), 0);
+		// Make two tracks to test filtering
+		FMidiTrack* FirstTrack = TheMidi->AddTrack("FirstTrack");
+		FMidiTrack* SecondTrack = TheMidi->AddTrack("SecondTrack");
+		// Add midi events to the two tracks
+		FMidiMsg FirstMessage = FMidiMsg::CreateNoteOn(1, 1, 1);		
+		FMidiMsg SecondMessage = FMidiMsg::CreateNoteOn(1, 2, 2);
+		FMidiEvent FirstEvent(0, FirstMessage);
+		FMidiEvent SecondEvent(0, SecondMessage);		
+		FirstTrack->AddEvent(FirstEvent);
+		SecondTrack->AddEvent(SecondEvent);
+		// Tell the midi file its tracks have been changed so it can recalculate song length data...
+		TheMidi->TracksChanged();
+		// Now round appropriately...
+		TheMidi->ConformMidiFileLength(EMidiFileLengthConformOption::RoundUp);
+		return TheMidi;
+	}
+
+	IMPLEMENT_SIMPLE_AUTOMATION_TEST(
+		FMidiStreamTrackIndexFunctionalityTest,
+		"Harmonix.Metasound.Nodes.MidiStreamTrackFilterNode.TestFunctionality",
+		EAutomationTestFlags::EditorContext | EAutomationTestFlags::EngineFilter)
+		bool FMidiStreamTrackIndexFunctionalityTest::RunTest(const FString&)
+	{
+		//************************************************************************************************************************
+		// Build the graph.
+		//************************************************************************************************************************
+		GraphBuilder Builder;
+		Builder.AddOutput("AudioOut", GetMetasoundDataTypeName<FAudioBuffer>());
+
+		// transport...
+		Frontend::FNodeHandle TransportNode = Builder.AddNode({ HarmonixMetasound::HarmonixNodeNamespace, "TriggerToTransport", "" }, 0);
+		UTEST_TRUE("Transport Node Created", TransportNode.Get().IsValid());
+		// hoist up play trigger input...
+		Builder.AddAndConnectDataReferenceInput(TransportNode, CommonPinNames::Inputs::TransportPlayName, GetMetasoundDataTypeName<FTrigger>());
+
+		// midi player...
+		Frontend::FNodeHandle MidiPlayerNode = Builder.AddNode({ HarmonixMetasound::HarmonixNodeNamespace, "MidiPlayer", "" }, 0);
+		UTEST_TRUE("MidiPlayer Node Created", MidiPlayerNode.Get().IsValid());
+		// hoist up midi asset input...
+		Builder.AddAndConnectDataReferenceInput(MidiPlayerNode, CommonPinNames::Inputs::MidiFileAssetName, GetMetasoundDataTypeName<FMidiAsset>());
+
+		// wire transport to midi player...
+		bool ConnectionSuccess = Builder.ConnectNodes(TransportNode, CommonPinNames::Outputs::TransportName, MidiPlayerNode, CommonPinNames::Inputs::TransportName);
+		UTEST_TRUE("Transport Out connected to MidiPlayer node", ConnectionSuccess);
+
+		// Midi stream track index filter...
+		FNodeHandle MidiStreamTrackFilterNode = Builder.AddNode(HarmonixMetasound::Nodes::MidiStreamTrackFilter::GetClassName(), 0);
+		UTEST_TRUE("MidiStreamTrackFilter Node Created", MidiStreamTrackFilterNode.Get().IsValid());
+		// hoist up filter input...
+		Builder.AddAndConnectDataReferenceInput(MidiStreamTrackFilterNode, CommonPinNames::Inputs::MidiTrackIndexFilterSpecifierName, GetMetasoundDataTypeName<FString>());
+		// hoist up output trigger...
+		Builder.AddAndConnectDataReferenceOutput(MidiStreamTrackFilterNode, CommonPinNames::Outputs::MidiStreamName, GetMetasoundDataTypeName<FMidiStream>(), CommonPinNames::Outputs::MidiStreamName);
+
+		// wire midi player to midi stream track index filter...
+		ConnectionSuccess = Builder.ConnectNodes(MidiPlayerNode, CommonPinNames::Outputs::MidiStreamName, MidiStreamTrackFilterNode, CommonPinNames::Inputs::MidiStreamName);
+		UTEST_TRUE("MidiPlayer Out connected to MidiStreamTrackFilter node", ConnectionSuccess);
+
+		//************************************************************************************************************************
+		// Make the generator.
+		//************************************************************************************************************************
+		constexpr int32 NumSamplesPerBlock = 480; // 10ms per block
+		const TUniquePtr<FMetasoundGenerator> Generator = Builder.BuildGenerator(48000, NumSamplesPerBlock);
+		UTEST_TRUE("Graph successfully built", Generator.IsValid());
+
+		//************************************************************************************************************************
+		// Get reference to outputs we want to check.
+		//************************************************************************************************************************
+		TOptional<FMidiStreamReadRef> OutputMidiStream = Generator->GetOutputReadReference<FMidiStream>(CommonPinNames::Outputs::MidiStreamName);
+		UTEST_TRUE("Output exists", OutputMidiStream.IsSet());		
+
+		//************************************************************************************************************************
+		// Make initial inputs.
+		//************************************************************************************************************************
+		// Midi file to play...
+		UMidiFile* MidiFile = BuildMidiFile();
+		TSharedPtr<Audio::IProxyData> MidiFileProxy = MidiFile->CreateProxyData({ "MidiStreamTrackFilterTest" });
+		Generator->SetInputValue<FMidiAsset>(CommonPinNames::Inputs::MidiFileAssetName, MidiFileProxy);
+
+		// Single Track Index to Filter
+		Generator->SetInputValue<FString>(CommonPinNames::Inputs::MidiTrackIndexFilterSpecifierName, "2");
+
+		// Trigger transport...
+		Generator->ApplyToInputValue<FTrigger>(CommonPinNames::Inputs::TransportPlayName, [](FTrigger& TriggerIn) { TriggerIn.TriggerFrame(0); });
+
+		//************************************************************************************************************************
+		// Generate 10ms...
+		//************************************************************************************************************************
+		TAudioBuffer<float> Buffer{ Generator->GetNumChannels(), NumSamplesPerBlock, EAudioBufferCleanupMode::Delete };
+		Generator->OnGenerateAudio(Buffer.GetRawChannelData(0), Buffer.GetNumTotalValidSamples());
+
+		// There are 2 note-on messages in the midi stream, in two different tracks.
+		// Ensure that the MidiStreamTrackFilter node correctly filtered out one track.
+		int32 MidiMessageCount = 0;
+		for (const FMidiStreamEvent& MidiStreamEvent : (*OutputMidiStream)->GetEventsInBlock())
+		{
+			if (MidiStreamEvent.MidiMessage.IsNoteOn())
+			{
+				++MidiMessageCount;
+			}
+		}
+
+		UTEST_EQUAL("Number of midi events matches filtered", MidiMessageCount, 1);
+
+		return true;
+	}
+}
+
+#endif
