[Gentech] Fusion Keyzone editor customization

- a custom SMinMaxSlider class derived from SSlider for selecting the min to max value range for note number and note velocity, with 2 handle thumbs controlling a lower value and a higher value (2 thumbs cannot overlap, meaning lower value can never be greater than higher value)
- a drop down menu displaying the keyzones name, only selected keyzone's properties are displayed below (similar to the fusion VST)

#tests PIE

diff --git a/Engine/Plugins/Runtime/Harmonix/Source/HarmonixDspEditor/HarmonixDspEditor.Build.cs b/Engine/Plugins/Runtime/Harmonix/Source/HarmonixDspEditor/HarmonixDspEditor.Build.cs
--- a/Engine/Plugins/Runtime/Harmonix/Source/HarmonixDspEditor/HarmonixDspEditor.Build.cs
+++ b/Engine/Plugins/Runtime/Harmonix/Source/HarmonixDspEditor/HarmonixDspEditor.Build.cs
@@ -32,6 +32,7 @@ public class HarmonixDspEditor : ModuleRules
 				"Json",
 				"AssetDefinition",
 				"PropertyEditor",
+				"DetailCustomizations",
 				"Slate",
 				"SlateCore",
 				"Settings",
diff --git a/Engine/Plugins/Runtime/Harmonix/Source/HarmonixDspEditor/Private/Customization/FusionPatchDetailCustomization.cpp b/Engine/Plugins/Runtime/Harmonix/Source/HarmonixDspEditor/Private/Customization/FusionPatchDetailCustomization.cpp
new file mode 100644
--- /dev/null
+++ b/Engine/Plugins/Runtime/Harmonix/Source/HarmonixDspEditor/Private/Customization/FusionPatchDetailCustomization.cpp
@@ -0,0 +1,267 @@
+// Copyright Epic Games, Inc. All Rights Reserved.
+#include "FusionPatchDetailCustomization.h"
+
+#include "DetailLayoutBuilder.h"
+#include "DetailWidgetRow.h"
+#include "PropertyCustomizationHelpers.h"
+#include "HarmonixDsp/FusionSampler/Settings/FusionPatchSettings.h"
+#include "HarmonixDsp/FusionSampler/Settings/KeyzoneSettings.h"
+#include "HarmonixDsp/FusionSampler/FusionPatch.h"
+
+#include "Widgets/Input/SComboBox.h"
+#include "Widgets/Text/STextBlock.h"
+
+#include "Widgets/SMinMaxSlider.h"
+
+void FFusionPatchDetailCustomization::CustomizeDetails(IDetailLayoutBuilder& DetailLayout)
+{
+	IDetailCategoryBuilder& FusionPatchDataCategory = DetailLayout.EditCategory(TEXT("Fusion Patch Data"));
+	//IDetailCategoryBuilder& FusionPatchDataCategory = DetailLayout.EditCategory(TEXT("Fusion Patch Data"));
+	IDetailCategoryBuilder& FilePathCategory = DetailLayout.EditCategory(TEXT("File Path"));
+
+	//get a handle of Fusion Patch Data
+	TSharedPtr<IPropertyHandle> FusionPatchDataHandle = DetailLayout.GetProperty("FusionPatchData");
+	check(FusionPatchDataHandle);
+	TSharedPtr<IPropertyHandle> PresetsHandle = FusionPatchDataHandle->GetChildHandle("Presets");
+	check(PresetsHandle);
+	TSharedPtr<IPropertyHandle> CurrentPresetIndexHandle = FusionPatchDataHandle->GetChildHandle("CurrentPresetIndex");
+	check(CurrentPresetIndexHandle);
+
+	//get a handle of the keyzones as array
+	TSharedPtr<IPropertyHandleArray> KeyzonesHandle = FusionPatchDataHandle->GetChildHandle("Keyzones")->AsArray();
+	check(KeyzonesHandle.IsValid());
+
+	//get the number of keyzones in the current patch
+	uint32 NumKeyzones;
+	KeyzonesHandle->GetNumElements(NumKeyzones);
+
+	//add all keyzones' name (sample path string) to an array for displaying in the dropdown menu
+	AddKeyzonesNameToMenuArray(KeyzonesHandle, static_cast<int32>(NumKeyzones));
+	if (!KeyzonesNameMenu.IsEmpty()) CurrentKeyzoneName = KeyzonesNameMenu[0];
+
+	//get the current fusion patch being edited
+	TArray<TWeakObjectPtr<UObject>> Objects;
+	DetailLayout.GetObjectsBeingCustomized(Objects);
+	if (Objects.Num() != 1)
+	{
+		return;
+	}
+	TWeakObjectPtr<UObject> FusionPatch = Objects.Last();
+	TWeakObjectPtr<UFusionPatch> FusionPatchBeingEdited = Cast<UFusionPatch>(FusionPatch);
+	
+	//add the dropdown menu displaying the sample path/name of keyzones to the name content of this row
+	FusionPatchDataCategory.AddCustomRow(FText::FromString("Fusion Patch"))
+		.NameContent()
+		[
+			SNew(STextBlock)
+			.Text(FText::FromString("Selected Keyzone"))
+			.Font(IDetailLayoutBuilder::GetDetailFontBold())
+		]
+		.ValueContent()
+		.MinDesiredWidth(350.f)
+		[
+			//create the dropdown combo box displaying all the keyzones name/sample path for selection
+			SNew(SComboBox<TSharedPtr<FString>>)
+				.OptionsSource(&KeyzonesNameMenu)
+				.OnSelectionChanged_Lambda([this](TSharedPtr<FString> NewKeyzoneName, ESelectInfo::Type)
+				{
+					CurrentKeyzoneName = NewKeyzoneName;
+					CurrentKeyzoneIndex = KeyzonesNameMenu.Find(NewKeyzoneName);
+				})
+				.OnGenerateWidget_Lambda([](TSharedPtr<FString> InKeyzoneName) -> TSharedRef<SWidget>
+				{
+					return SNew(STextBlock).Text(FText::FromString(*InKeyzoneName));
+				})
+				.InitiallySelectedItem(CurrentKeyzoneName)
+				.Content()
+				[
+					SNew(STextBlock).Text_Lambda([this]() -> FText
+					{
+						if (CurrentKeyzoneName.IsValid())
+						{
+							return FText::FromString(*CurrentKeyzoneName);
+						}
+
+						return FText::FromString(TEXT("InvalidComboEntryText"));
+					})
+				]
+		];
+
+	DrawSelectedKeyzoneProperties(FusionPatchDataCategory, KeyzonesHandle, NumKeyzones);
+
+	//add the presets properties since they're not customized 
+	FusionPatchDataCategory.AddProperty(PresetsHandle);
+	FusionPatchDataCategory.AddProperty(CurrentPresetIndexHandle);
+
+	//hide the non-customized versions of these properties
+	DetailLayout.HideProperty(FusionPatchDataHandle);
+}
+
+
+
+TSharedRef<SWidget> FFusionPatchDetailCustomization::CreateMinMaxSliderWidget(TSharedPtr<IPropertyHandle> MinValuePropertyHandle, TSharedPtr<IPropertyHandle> MaxValuePropertyHandle)
+{
+	int8 CurrentMinValue = 0.0f;
+	MinValuePropertyHandle->GetValue(CurrentMinValue);
+	const int32 ClampMinValue = MinValuePropertyHandle->GetIntMetaData("ClampMin");
+	int8 CurrentMaxValue = 0.0f;
+	MaxValuePropertyHandle->GetValue(CurrentMaxValue);
+	const int32 ClampMaxValue = MinValuePropertyHandle->GetIntMetaData("ClampMax");
+	TSharedPtr<SMinMaxSlider> CustomRangeSlider;
+	TSharedRef<SWidget> OutWidget = SNew(SHorizontalBox)
+		+SHorizontalBox::Slot()
+		.HAlign(HAlign_Left)
+		.Padding(0, 5, 0, 5)
+		.AutoWidth()
+		[
+			MinValuePropertyHandle->CreatePropertyValueWidget()
+		]
+		+SHorizontalBox::Slot()
+		.HAlign(HAlign_Fill)
+		.VAlign(VAlign_Center)
+		.FillWidth(1)
+		[
+			SAssignNew(CustomRangeSlider, SMinMaxSlider)
+			.MinValue(ClampMinValue)
+			.MaxValue(ClampMaxValue)
+			.LowerHandleValue(static_cast<float>(CurrentMinValue))
+			.UpperHandleValue(static_cast<float>(CurrentMaxValue))
+			.OnLowerHandleValueChanged_Lambda([this, MinValuePropertyHandle, MaxValuePropertyHandle](float NewValue)
+			{
+				int8 MaxValue;
+				MaxValuePropertyHandle->GetValue(MaxValue);
+				int8 NewMinValue = FMath::RoundToInt(NewValue);
+				MinValuePropertyHandle->SetValue(NewMinValue <= MaxValue ? NewMinValue : MaxValue);
+			})
+			.OnUpperHandleValueChanged_Lambda([this, MinValuePropertyHandle, MaxValuePropertyHandle](float NewValue)
+			{
+				int8 MinValue;
+				MinValuePropertyHandle->GetValue(MinValue);
+				int8 NewMaxValue = FMath::RoundToInt(NewValue);
+				MaxValuePropertyHandle->SetValue(NewMaxValue >= MinValue ? NewMaxValue : MinValue);
+			})
+		]
+		+SHorizontalBox::Slot()
+		.HAlign(HAlign_Left)
+		.VAlign(VAlign_Center)
+		.AutoWidth()
+		.Padding(0, 5, 0, 5)
+		[
+			MaxValuePropertyHandle->CreatePropertyValueWidget()
+		];
+
+	//if the min or max note numeric boxes are edited, modify the slider values accordingly
+	MinValuePropertyHandle->SetOnPropertyValueChanged(FSimpleDelegate::CreateLambda([this, CustomRangeSlider, MinValuePropertyHandle, MaxValuePropertyHandle]()
+	{
+		int8 MinValue;
+		MinValuePropertyHandle->GetValue(MinValue);
+		int8 MaxValue;
+		MaxValuePropertyHandle->GetValue(MaxValue);
+		CustomRangeSlider->SetLowerValue(static_cast<float>(MinValue <= MaxValue ? MinValue : MaxValue));
+		CustomRangeSlider->SetUpperValue(static_cast<float>(MaxValue >= MinValue ? MaxValue : MinValue));
+	}));
+
+	MaxValuePropertyHandle->SetOnPropertyValueChanged(FSimpleDelegate::CreateLambda([this, CustomRangeSlider, MinValuePropertyHandle, MaxValuePropertyHandle]()
+	{
+		int8 MinValue;
+		MinValuePropertyHandle->GetValue(MinValue);
+		int8 MaxValue;
+		MaxValuePropertyHandle->GetValue(MaxValue);
+		CustomRangeSlider->SetUpperValue(static_cast<float>(MaxValue >= MinValue ? MaxValue : MinValue));
+		CustomRangeSlider->SetLowerValue(static_cast<float>(MinValue <= MaxValue ? MinValue : MaxValue));
+	}));
+	
+	return OutWidget;
+}
+
+FDetailWidgetRow& FFusionPatchDetailCustomization::AddCustomMinMaxSliderRow(IDetailCategoryBuilder& FusionPatchDataCategory, const FText& DisplayName, TSharedPtr<IPropertyHandle> MinPropertyHandle, TSharedPtr<IPropertyHandle> MaxPropertyHandle)
+{
+	return FusionPatchDataCategory.AddCustomRow(FText::FromString("Fusion Patch"))
+		.NameContent()
+		[
+			SNew(STextBlock)
+			.Text(DisplayName)
+			.Font(IDetailLayoutBuilder::GetDetailFont())
+		]
+		.ValueContent()
+		.MinDesiredWidth(350.f)
+		[
+			CreateMinMaxSliderWidget(MinPropertyHandle, MaxPropertyHandle)
+		];
+}
+
+void FFusionPatchDetailCustomization::AddKeyzonesNameToMenuArray(TSharedPtr<IPropertyHandleArray> KeyzonesHandle, int32 NumKeyzones)
+{
+	for (int32 KeyzoneIndex = 0; KeyzoneIndex < static_cast<int32>(NumKeyzones); ++KeyzoneIndex)
+	{
+		TSharedPtr<IPropertyHandle> SamplePathHandle = KeyzonesHandle->GetElement(KeyzoneIndex)->GetChildHandle(GET_MEMBER_NAME_CHECKED(FKeyzoneSettings, SamplePath));
+		FString CurrentSamplePath;
+		if (SamplePathHandle->GetValue(CurrentSamplePath) == FPropertyAccess::Success)
+		{
+			FString Path;
+			FString Filename;
+			FString Ext;
+			FPaths::Split(CurrentSamplePath, Path, Filename, Ext);
+			if (Filename.IsEmpty())
+			{
+				KeyzonesNameMenu.Add(MakeShareable(new FString(CurrentSamplePath)));
+			}
+			else
+			{
+				KeyzonesNameMenu.Add(MakeShareable(new FString(Filename)));
+			}
+		}
+	}
+}
+
+void FFusionPatchDetailCustomization::DrawSelectedKeyzoneProperties(IDetailCategoryBuilder& FusionPatchDataCategory, TSharedPtr<IPropertyHandleArray> KeyzonesHandle, int32 NumKeyzones)
+{
+	for (int32 KeyzoneIndex = 0; KeyzoneIndex < static_cast<int32>(NumKeyzones); ++KeyzoneIndex)
+	{
+		TFunction<EVisibility()> CustomVisibility = [this, KeyzoneIndex]() -> EVisibility
+		{
+			if (KeyzoneIndex != CurrentKeyzoneIndex)
+			{
+				return EVisibility::Collapsed;
+			}
+			else
+			{
+				return EVisibility::Visible;
+			}
+		};
+		
+		TSharedRef<IPropertyHandle> KeyzoneHandle = KeyzonesHandle->GetElement(KeyzoneIndex);
+		uint32 NumChildren = 0;
+		KeyzoneHandle->GetNumChildren(NumChildren);
+		for (uint32 ChildIdx = 0; ChildIdx < NumChildren; ++ChildIdx)
+		{
+			TSharedPtr<IPropertyHandle> PropertyHandle = KeyzoneHandle->GetChildHandle(ChildIdx);
+			check(PropertyHandle);
+
+			if (PropertyHandle->GetProperty()->GetFName() == GET_MEMBER_NAME_CHECKED(FKeyzoneSettings, MinNote))
+			{
+				TSharedPtr<IPropertyHandle> MaxPropertyHandle = KeyzoneHandle->GetChildHandle(ChildIdx + 1);
+				check(MaxPropertyHandle);
+				FDetailWidgetRow& MinMaxRow = AddCustomMinMaxSliderRow(FusionPatchDataCategory, NSLOCTEXT("FusionPatch_Details", "NoteRange", "Note Range (Min/Max)"), PropertyHandle, MaxPropertyHandle);
+				MinMaxRow.Visibility(TAttribute<EVisibility>::Create(TAttribute<EVisibility>::FGetter::CreateLambda(CustomVisibility)));
+				// increment ChildIdx since we customized two properties
+				++ChildIdx;
+			}
+			else if (PropertyHandle->GetProperty()->GetFName() == GET_MEMBER_NAME_CHECKED(FKeyzoneSettings, MinVelocity))
+			{
+				TSharedPtr<IPropertyHandle> MaxPropertyHandle = KeyzoneHandle->GetChildHandle(ChildIdx + 1);
+				check(MaxPropertyHandle);
+				FDetailWidgetRow& MinMaxRow = AddCustomMinMaxSliderRow(FusionPatchDataCategory, NSLOCTEXT("FusionPatch_Details", "VelocityRange", "Velocity Range (Min/Max)"), PropertyHandle, MaxPropertyHandle);
+				MinMaxRow.Visibility(TAttribute<EVisibility>::Create(TAttribute<EVisibility>::FGetter::CreateLambda(CustomVisibility)));
+				// increment ChildIdx since we customized two properties
+				++ChildIdx;
+			}
+			else // default property
+			{
+				IDetailPropertyRow& PropertyRow = FusionPatchDataCategory.AddProperty(PropertyHandle);
+				PropertyRow.Visibility(TAttribute<EVisibility>::Create(TAttribute<EVisibility>::FGetter::CreateLambda(CustomVisibility)));
+			}
+
+		}
+	}
+}
\ No newline at end of file
diff --git a/Engine/Plugins/Runtime/Harmonix/Source/HarmonixDspEditor/Private/Customization/FusionPatchDetailCustomization.h b/Engine/Plugins/Runtime/Harmonix/Source/HarmonixDspEditor/Private/Customization/FusionPatchDetailCustomization.h
new file mode 100644
--- /dev/null
+++ b/Engine/Plugins/Runtime/Harmonix/Source/HarmonixDspEditor/Private/Customization/FusionPatchDetailCustomization.h
@@ -0,0 +1,39 @@
+// Copyright Epic Games, Inc. All Rights Reserved.
+
+#pragma once
+
+#include "CoreMinimal.h"
+#include "IDetailCustomization.h"
+#include "DetailCategoryBuilder.h"
+#include "UObject/NameTypes.h"
+#include "Templates/SharedPointer.h"
+#include "Layout/Visibility.h"
+
+class SWidget;
+
+//Fusion Patch detail customizatiaon class
+class FFusionPatchDetailCustomization : public IDetailCustomization
+{
+public:
+	static TSharedRef<IDetailCustomization> MakeInstance()
+	{
+		return MakeShareable(new FFusionPatchDetailCustomization);
+	}
+
+	virtual void CustomizeDetails(IDetailLayoutBuilder& DetailLayout) override;
+private:
+	int32 CurrentKeyzoneIndex = 0;
+	
+	void AddKeyzonesNameToMenuArray(TSharedPtr<IPropertyHandleArray> KeyzonesHandle, int32 NumKeyzones);
+
+	//helper functions to paint customized properties
+	FDetailWidgetRow& AddCustomMinMaxSliderRow(IDetailCategoryBuilder& FusionPatchDataCategory, const FText& DisplayName, TSharedPtr<IPropertyHandle> MinPropertyHandle, TSharedPtr<IPropertyHandle> MaxPropertyHandle);
+	void DrawSelectedKeyzoneProperties(IDetailCategoryBuilder& FusionPatchDataCategory, TSharedPtr<IPropertyHandleArray> KeyzonesHandle, int32 NumKeyzones);
+
+	TSharedPtr<FString> CurrentKeyzoneName;
+	TArray<TSharedPtr<FString>> KeyzonesNameMenu;
+
+	TSharedRef<SWidget> CreateMinMaxSliderWidget(TSharedPtr<IPropertyHandle> MinValuePropertyHandle, TSharedPtr<IPropertyHandle> MaxValuePropertyHandle);
+};
+
+
diff --git a/Engine/Plugins/Runtime/Harmonix/Source/HarmonixDspEditor/Private/HarmonixDspEditorModule.cpp b/Engine/Plugins/Runtime/Harmonix/Source/HarmonixDspEditor/Private/HarmonixDspEditorModule.cpp
--- a/Engine/Plugins/Runtime/Harmonix/Source/HarmonixDspEditor/Private/HarmonixDspEditorModule.cpp
+++ b/Engine/Plugins/Runtime/Harmonix/Source/HarmonixDspEditor/Private/HarmonixDspEditorModule.cpp
@@ -12,6 +12,8 @@
 #include "Customization/PitchShifterNameCustomization.h"
 #include "Customization/PitchShifterConfigCustomization.h"
 #include "Customization/PannerDetailsCustomization.h"
+#include "Customization/FusionPatchDetailCustomization.h"
+#include "HarmonixDsp/FusionSampler/FusionPatch.h"
 
 #define LOCTEXT_NAMESPACE "HarmonixDspEditor"
 
@@ -25,13 +27,14 @@ void FHarmonixDspEditorModule::StartupModule()
 	PropertyEditorModule.RegisterCustomPropertyTypeLayout("PitchShifterName", FOnGetPropertyTypeCustomizationInstance::CreateStatic(&FPitchShifterNameCustomization::MakeInstance));
 	PropertyEditorModule.RegisterCustomPropertyTypeLayout("PannerDetails", FOnGetPropertyTypeCustomizationInstance::CreateStatic(&FPannerDetailsCustomization::MakeInstance));
 	PropertyEditorModule.RegisterCustomClassLayout("StretcherAndPitchShifterFactoryConfig", FOnGetDetailCustomizationInstance::CreateStatic(&FPitchShifterConfigCustomization::MakeInstance));
-	
+	PropertyEditorModule.RegisterCustomClassLayout(UFusionPatch::StaticClass()->GetFName(), FOnGetDetailCustomizationInstance::CreateStatic(&FFusionPatchDetailCustomization::MakeInstance));
 }
 
 void FHarmonixDspEditorModule::ShutdownModule()
 {
 
 	FPropertyEditorModule& PropertyEditorModule = FModuleManager::LoadModuleChecked<FPropertyEditorModule>("PropertyEditor");
+	PropertyEditorModule.UnregisterCustomClassLayout(UFusionPatch::StaticClass()->GetFName());
 	PropertyEditorModule.UnregisterCustomClassLayout("StretcherAndPitchShifterFactoryConfig");
 	PropertyEditorModule.UnregisterCustomPropertyTypeLayout("PannerDetails");
 	PropertyEditorModule.UnregisterCustomPropertyTypeLayout("PitchShifterName");
diff --git a/Engine/Plugins/Runtime/Harmonix/Source/HarmonixDspEditor/Private/Widgets/SMinMaxSlider.cpp b/Engine/Plugins/Runtime/Harmonix/Source/HarmonixDspEditor/Private/Widgets/SMinMaxSlider.cpp
new file mode 100644
--- /dev/null
+++ b/Engine/Plugins/Runtime/Harmonix/Source/HarmonixDspEditor/Private/Widgets/SMinMaxSlider.cpp
@@ -0,0 +1,300 @@
+// Copyright Epic Games, Inc. All Rights Reserved.
+
+#include "SMinMaxSlider.h"
+
+void SMinMaxSlider::Construct(const FArguments& InDeclaration)
+{
+	SSlider::Construct(SSlider::FArguments()
+		.MinValue(InDeclaration._MinValue)
+		.MaxValue(InDeclaration._MaxValue)
+	);
+
+	LowerHandleValue = InDeclaration._LowerHandleValue;
+	UpperHandleValue = InDeclaration._UpperHandleValue;
+	OnLowerHandleValueChanged = InDeclaration._OnLowerHandleValueChanged;
+	OnUpperHandleValueChanged = InDeclaration._OnUpperHandleValueChanged;
+	IndentHandleAttribute = InDeclaration._IndentHandle;
+	SliderBarColorAttribute = InDeclaration._SliderBarColor;
+	SliderLowerHandleColorAttribute = InDeclaration._SliderLowerHandleColor;
+	SliderUpperHandleColorAttribute = InDeclaration._SliderUpperHandleColor;
+}
+
+float SMinMaxSlider::GetLowerNormalizedValue() const
+{
+	if (MaxValue == MinValue)
+	{
+		return 1.0f;
+	}
+	else
+	{
+		return (LowerHandleValue.Get() - MinValue) / (MaxValue - MinValue);
+	}
+}
+
+
+float SMinMaxSlider::GetUpperNormalizedValue() const
+{
+	if (MaxValue == MinValue)
+	{
+		return 1.0f;
+	}
+	else
+	{
+		return (UpperHandleValue.Get() - MinValue) / (MaxValue - MinValue);
+	}
+}
+
+void SMinMaxSlider::SetLowerValue(TAttribute<float> InValueAttribute)
+{
+	float NewLowerValue = InValueAttribute.Get() <= UpperHandleValue.Get() ? InValueAttribute.Get() : UpperHandleValue.Get();
+	LowerHandleValue.Set(NewLowerValue);
+}
+
+void SMinMaxSlider::SetUpperValue(TAttribute<float> InValueAttribute)
+{
+	float NewUpperValue = InValueAttribute.Get() >= LowerHandleValue.Get() ? InValueAttribute.Get() : LowerHandleValue.Get();
+	UpperHandleValue.Set(NewUpperValue);
+}
+
+void SMinMaxSlider::OnSliderValueChanged(float NewValue)
+{
+	if (CurrentHandle == ESliderHandle::LowerHandle)
+	{
+		OnLowerHandleValueChanged.ExecuteIfBound(NewValue);
+	}
+	if (CurrentHandle == ESliderHandle::UpperHandle)
+	{
+		OnUpperHandleValueChanged.ExecuteIfBound(NewValue);
+	}
+}
+
+FVector2D SMinMaxSlider::CalculateLowerHandlePosition(const FGeometry& HandleGeometry, float Value) const
+{
+	const float ClampedValue = FMath::Clamp(GetLowerNormalizedValue(), 0.0f, 1.0f);
+
+	float XPosition = ClampedValue * HandleGeometry.GetLocalSize().X;
+
+	float YPosition = HandleGeometry.GetLocalSize().Y * 0.5f;
+
+	return FVector2D(XPosition, YPosition);
+}
+
+FVector2D SMinMaxSlider::CalculateUpperHandlePosition(const FGeometry& HandleGeometry, float Value) const
+{
+	const float ClampedValue = FMath::Clamp(GetUpperNormalizedValue(), 0.0f, 1.0f);
+
+	float XPosition = ClampedValue * HandleGeometry.GetLocalSize().X;
+
+	float YPosition = HandleGeometry.GetLocalSize().Y * 0.5f;
+
+	return FVector2D(XPosition, YPosition);
+}
+
+SMinMaxSlider::ESliderHandle SMinMaxSlider::DetermineClickedHandle(const FGeometry& HandleGeometry, const FVector2D& LocalMousePosition) const
+{
+	FVector2D LowerHandlePos = CalculateLowerHandlePosition(HandleGeometry, LowerHandleValue.Get());
+	FVector2D UpperHandlePos = CalculateUpperHandlePosition(HandleGeometry, UpperHandleValue.Get());
+
+	float LowerHandleDistance = FVector2D::DistSquared(LowerHandlePos, LocalMousePosition);
+	float UpperHandleDistance = FVector2D::DistSquared(UpperHandlePos, LocalMousePosition);
+
+	if (LowerHandleDistance < UpperHandleDistance)
+	{
+		return ESliderHandle::LowerHandle;
+	}
+	else {
+		return ESliderHandle::UpperHandle;
+	}
+}
+
+int32 SMinMaxSlider::OnPaint(const FPaintArgs& Args, const FGeometry& AllottedGeometry, const FSlateRect& MyCullingRect, FSlateWindowElementList& OutDrawElements, int32 LayerId, const FWidgetStyle& InWidgetStyle, bool bParentEnabled) const
+{
+	const float AllottedWidth = Orientation == Orient_Horizontal ? AllottedGeometry.GetLocalSize().X : AllottedGeometry.GetLocalSize().Y;
+	const float AllottedHeight = Orientation == Orient_Horizontal ? AllottedGeometry.GetLocalSize().Y : AllottedGeometry.GetLocalSize().X;
+
+	float HandleRotation;
+	FVector2f HandleTopRightPoint;
+	FVector2f HandleTopLeftPoint;
+	FVector2f SliderStartPoint;
+	FVector2f SliderEndPoint;
+
+	// calculate slider geometry as if it's a horizontal slider (we'll rotate it later if it's vertical)
+	const FVector2f HandleSize = GetThumbImage()->ImageSize;
+	const FVector2f HalfHandleSize = 0.5f * HandleSize;
+	const float Indentation = IndentHandleAttribute.Get() ? HandleSize.X : 0.0f;
+
+	// We clamp to make sure that the slider cannot go out of the slider Length.
+	const float SliderPercent = FMath::Clamp(GetLowerNormalizedValue(), 0.0f, 1.0f);
+	const float UpperSliderPercent = FMath::Clamp(GetUpperNormalizedValue(), 0.0f, 1.0f);
+
+	const float SliderLength = AllottedWidth - (Indentation + HandleSize.X);
+	const float SliderHandleOffset = SliderPercent * SliderLength;
+	const float SliderUpperHandleOffset = UpperSliderPercent * SliderLength;
+	const float SliderY = 0.5f * AllottedHeight;
+
+	HandleRotation = 0.0f;
+	HandleTopRightPoint = FVector2f(SliderUpperHandleOffset + (0.5f * Indentation), SliderY - HalfHandleSize.Y);
+	HandleTopLeftPoint = FVector2f(SliderHandleOffset + (0.5f * Indentation), SliderY - HalfHandleSize.Y);
+
+	SliderStartPoint = FVector2f(HalfHandleSize.X, SliderY);
+	SliderEndPoint = FVector2f(AllottedWidth - HalfHandleSize.X, SliderY);
+
+	FGeometry SliderGeometry = AllottedGeometry;
+
+	// rotate the slider 90deg if it's vertical. The 0 side goes on the bottom, the 1 side on the top.
+	if (Orientation == Orient_Vertical)
+	{
+		FSlateRenderTransform SlateRenderTransform = TransformCast<FSlateRenderTransform>(Concatenate(Inverse(FVector2f(AllottedWidth, 0)), FQuat2D(FMath::DegreesToRadians(-90.0f))));
+		SliderGeometry = AllottedGeometry.MakeChild(
+			FVector2f(AllottedWidth, AllottedHeight),
+			FSlateLayoutTransform(),
+			SlateRenderTransform, FVector2f::ZeroVector);
+	}
+
+	const bool bEnabled = ShouldBeEnabled(bParentEnabled);
+	const ESlateDrawEffect DrawEffects = bEnabled ? ESlateDrawEffect::None : ESlateDrawEffect::DisabledEffect;
+
+	// draw slider bar
+	auto BarTopLeft = FVector2f(SliderStartPoint.X, SliderStartPoint.Y - Style->BarThickness * 0.5f);
+	auto BarSize = FVector2f(SliderEndPoint.X - SliderStartPoint.X, Style->BarThickness);
+	auto BarImage = GetBarImage();
+	auto ThumbImage = GetThumbImage();
+
+	FSlateDrawElement::MakeBox(
+		OutDrawElements,
+		LayerId,
+		SliderGeometry.ToPaintGeometry(BarSize, FSlateLayoutTransform(BarTopLeft)),
+		BarImage,
+		DrawEffects,
+		BarImage->GetTint(InWidgetStyle) * SliderBarColorAttribute.Get().GetColor(InWidgetStyle) * InWidgetStyle.GetColorAndOpacityTint()
+	);
+
+	++LayerId;
+
+	// draw slider thumbs
+	FSlateDrawElement::MakeBox(
+		OutDrawElements,
+		LayerId,
+		SliderGeometry.ToPaintGeometry(GetThumbImage()->ImageSize, FSlateLayoutTransform(HandleTopLeftPoint)),
+		ThumbImage,
+		DrawEffects,
+		ThumbImage->GetTint(InWidgetStyle) * SliderLowerHandleColorAttribute.Get().GetColor(InWidgetStyle) * InWidgetStyle.GetColorAndOpacityTint()
+	);
+	FSlateDrawElement::MakeBox(
+		OutDrawElements,
+		LayerId,
+		SliderGeometry.ToPaintGeometry(GetThumbImage()->ImageSize, FSlateLayoutTransform(HandleTopRightPoint)),
+		ThumbImage,
+		DrawEffects,
+		ThumbImage->GetTint(InWidgetStyle) * SliderUpperHandleColorAttribute.Get().GetColor(InWidgetStyle) * InWidgetStyle.GetColorAndOpacityTint()
+	);
+
+	return LayerId;
+}
+
+float SMinMaxSlider::PositionToValue(const FGeometry& MyGeometry, const UE::Slate::FDeprecateVector2DParameter& AbsolutePosition)
+{
+	float HandleValue = LowerHandleValue.Get();
+	if (CurrentHandle == ESliderHandle::UpperHandle)
+	{
+		HandleValue = UpperHandleValue.Get();
+	}
+
+	const FVector2f LocalPosition = MyGeometry.AbsoluteToLocal(AbsolutePosition);
+
+	float RelativeValue;
+	float Denominator;
+	// Only need X as we rotate the thumb image when rendering vertically
+	const float Indentation = SSlider::GetThumbImage()->ImageSize.X * (IndentHandleAttribute.Get() ? 2.f : 1.f);
+	const float HalfIndentation = 0.5f * Indentation;
+
+	if (Orientation == Orient_Horizontal)
+	{
+		Denominator = MyGeometry.Size.X - Indentation;
+		RelativeValue = (Denominator != 0.f) ? (LocalPosition.X - HalfIndentation) / Denominator : 0.f;
+	}
+	else
+	{
+		Denominator = MyGeometry.Size.Y - Indentation;
+		// Inverse the calculation as top is 0 and bottom is 1
+		RelativeValue = (Denominator != 0.f) ? ((MyGeometry.Size.Y - LocalPosition.Y) - HalfIndentation) / Denominator : 0.f;
+	}
+
+	RelativeValue = FMath::Clamp(RelativeValue, 0.0f, 1.0f) * (MaxValue - MinValue) + MinValue;
+	if (bMouseUsesStep)
+	{
+		float direction = HandleValue - RelativeValue;
+		float CurrentStepSize = StepSize.Get();
+		if (direction > CurrentStepSize / 2.0f)
+		{
+			return FMath::Clamp(HandleValue - CurrentStepSize, MinValue, MaxValue);
+		}
+		else if (direction < CurrentStepSize / -2.0f)
+		{
+			return FMath::Clamp(HandleValue + CurrentStepSize, MinValue, MaxValue);
+		}
+		else
+		{
+			return HandleValue;
+		}
+	}
+	return RelativeValue;
+}
+
+void SMinMaxSlider::CommitValue(float NewValue)
+{
+	const float OldValue = GetValue();
+
+	if (NewValue != OldValue)
+	{
+		if (CurrentHandle == ESliderHandle::LowerHandle)
+		{
+			SetLowerValue(NewValue);
+
+			Invalidate(EInvalidateWidgetReason::Paint);
+
+			OnLowerHandleValueChanged.ExecuteIfBound(LowerHandleValue.Get());
+		}
+
+		if (CurrentHandle == ESliderHandle::UpperHandle)
+		{
+			SetUpperValue(NewValue);
+
+			Invalidate(EInvalidateWidgetReason::Paint);
+
+			OnUpperHandleValueChanged.ExecuteIfBound(UpperHandleValue.Get());
+		}
+
+	}
+}
+
+FReply SMinMaxSlider::OnMouseButtonDown(const FGeometry& MyGeometry, const FPointerEvent& MouseEvent)
+{
+	if ((MouseEvent.GetEffectingButton() == EKeys::LeftMouseButton) && !IsLocked())
+	{
+		FVector2D LocalMousePosition = MyGeometry.AbsoluteToLocal(MouseEvent.GetScreenSpacePosition());
+		CurrentHandle = DetermineClickedHandle(MyGeometry, LocalMousePosition);
+		CachedCursor = GetCursor().Get(EMouseCursor::Default);
+
+		CommitValue(PositionToValue(MyGeometry, MouseEvent.GetScreenSpacePosition()));
+
+		return FReply::Handled().CaptureMouse(SharedThis(this));
+	}
+
+	return FReply::Unhandled();
+}
+
+FReply SMinMaxSlider::OnMouseButtonUp(const FGeometry& MyGeometry, const FPointerEvent& MouseEvent)
+{
+
+	if ((MouseEvent.GetEffectingButton() == EKeys::LeftMouseButton) && HasMouseCaptureByUser(MouseEvent.GetUserIndex(), MouseEvent.GetPointerIndex()))
+	{
+		CurrentHandle = ESliderHandle::None;
+		SetCursor(CachedCursor);
+
+		return FReply::Handled().ReleaseMouseCapture();
+	}
+
+	return FReply::Unhandled();
+
+}
\ No newline at end of file
diff --git a/Engine/Plugins/Runtime/Harmonix/Source/HarmonixDspEditor/Private/Widgets/SMinMaxSlider.h b/Engine/Plugins/Runtime/Harmonix/Source/HarmonixDspEditor/Private/Widgets/SMinMaxSlider.h
new file mode 100644
--- /dev/null
+++ b/Engine/Plugins/Runtime/Harmonix/Source/HarmonixDspEditor/Private/Widgets/SMinMaxSlider.h
@@ -0,0 +1,76 @@
+// Copyright Epic Games, Inc. All Rights Reserved.
+#pragma once
+
+#include "Widgets/Input/SSlider.h"
+
+DECLARE_DELEGATE_OneParam(FOnLowerValueChanged, float)
+DECLARE_DELEGATE_OneParam(FOnUpperValueChanged, float)
+
+/**
+ * A custom range slider class inherited from SSlider for customizing FusionPatch's min and max note/velocity properties
+ * 2 non-overlapped thumb handles controlling a lower value and a higher value
+ */
+class SMinMaxSlider : public SSlider
+{
+public:
+	SLATE_BEGIN_ARGS(SMinMaxSlider)
+		:_LowerHandleValue(0.0f)
+		, _UpperHandleValue(1.0f)
+		, _IndentHandle(true)
+		, _SliderBarColor(FLinearColor::White)
+		, _SliderLowerHandleColor(FLinearColor::White)
+		, _SliderUpperHandleColor(FLinearColor::White)
+		, _MinValue(0.0f)
+		, _MaxValue(1.0f)
+	{}
+	    SLATE_ATTRIBUTE(float, LowerHandleValue)
+		SLATE_ATTRIBUTE(float, UpperHandleValue)
+		SLATE_ATTRIBUTE(bool, IndentHandle)
+		SLATE_ATTRIBUTE(FSlateColor, SliderBarColor)
+		SLATE_ATTRIBUTE(FSlateColor, SliderLowerHandleColor)
+		SLATE_ATTRIBUTE(FSlateColor, SliderUpperHandleColor)
+		SLATE_EVENT(FOnLowerValueChanged, OnLowerHandleValueChanged)
+		SLATE_EVENT(FOnUpperValueChanged, OnUpperHandleValueChanged)
+		SLATE_ARGUMENT(float, MinValue)
+		SLATE_ARGUMENT(float, MaxValue)
+
+    SLATE_END_ARGS()
+
+	void Construct(const FArguments& InDeclaration);
+	
+    float GetLowerNormalizedValue() const;
+	float GetUpperNormalizedValue() const;
+	void SetLowerValue(TAttribute<float> InValueAttribute);
+	void SetUpperValue(TAttribute<float> InValueAttribute);
+	void OnSliderValueChanged(float NewValue);
+
+protected:
+	enum class ESliderHandle
+	{
+		None,
+		LowerHandle,
+		UpperHandle
+	};
+
+	FVector2D CalculateLowerHandlePosition(const FGeometry& HandleGeometry, float Value) const;
+	FVector2D CalculateUpperHandlePosition(const FGeometry& HandleGeometry, float Value) const;
+	ESliderHandle DetermineClickedHandle(const FGeometry& HandleGeometry, const FVector2D& LocalMousePosition) const;
+	float PositionToValue(const FGeometry& MyGeometry, const UE::Slate::FDeprecateVector2DParameter& AbsolutePosition);
+	
+	virtual void CommitValue(float NewValue) override;
+	virtual int32 OnPaint(const FPaintArgs& Args, const FGeometry& AllottedGeometry, const FSlateRect& MyCullingRect, FSlateWindowElementList& OutDrawElements, int32 LayerId, const FWidgetStyle& InWidgetStyle, bool bParentEnabled) const override;
+	virtual FReply OnMouseButtonDown(const FGeometry& MyGeometry, const FPointerEvent& MouseEvent) override;
+	virtual FReply OnMouseButtonUp(const FGeometry& MyGeometry, const FPointerEvent& MouseEvent) override;
+
+private:
+	TAttribute<float> LowerHandleValue;
+	TAttribute<float> UpperHandleValue;
+	FOnLowerValueChanged OnLowerHandleValueChanged;
+	FOnUpperValueChanged OnUpperHandleValueChanged;
+	TAttribute<bool> IndentHandleAttribute;
+	TAttribute<FSlateColor> SliderBarColorAttribute;
+	TAttribute<FSlateColor> SliderLowerHandleColorAttribute;
+	TAttribute<FSlateColor> SliderUpperHandleColorAttribute;
+	ESliderHandle CurrentHandle = ESliderHandle::None;
+
+};
\ No newline at end of file
