[Harmonix] Fix duplicate param localization names

diff --git a/Engine/Plugins/Runtime/Harmonix/Source/HarmonixDspEditor/Private/FusionPatchAssetFactory.cpp b/Engine/Plugins/Runtime/Harmonix/Source/HarmonixDspEditor/Private/FusionPatchAssetFactory.cpp
--- a/Engine/Plugins/Runtime/Harmonix/Source/HarmonixDspEditor/Private/FusionPatchAssetFactory.cpp
+++ b/Engine/Plugins/Runtime/Harmonix/Source/HarmonixDspEditor/Private/FusionPatchAssetFactory.cpp
@@ -173,11 +173,11 @@ void UFusionPatchAssetFactory::UpdateFusionPatchImportNotificationItem(TSharedPt
 	if (bImportSuccessful)
 	{
 		InItem->SetCompletionState(SNotificationItem::ECompletionState::CS_Success);
-		InItem->SetText(FText::Format(NSLOCTEXT("FusionPatchImporter", "FusionPatchImportProgressNotification", "Successfully imported Fusion Patch asset: {0}"), FText::FromString(InName.ToString())));
+		InItem->SetText(FText::Format(NSLOCTEXT("FusionPatchImporter", "FusionPatchImportProgressNotification_Success", "Successfully imported Fusion Patch asset: {0}"), FText::FromString(InName.ToString())));
 	}
 	else {
 		InItem->SetCompletionState(SNotificationItem::ECompletionState::CS_Fail);
-		InItem->SetText(FText::Format(NSLOCTEXT("FusionPatchImporter", "FusionPatchImportProgressNotification", "Failed to import Fusion Patch: {0}"), FText::FromString(InName.ToString())));
+		InItem->SetText(FText::Format(NSLOCTEXT("FusionPatchImporter", "FusionPatchImportProgressNotification_Failure", "Failed to import Fusion Patch: {0}"), FText::FromString(InName.ToString())));
 	}
 	InItem->SetExpireDuration(0.2f);
 	InItem->ExpireAndFadeout();
@@ -206,7 +206,7 @@ UObject* UFusionPatchAssetFactory::FactoryCreateText(UClass* InClass, UObject* I
 		const FString SourcePath = FPaths::GetPath(SourceFile);
 		
 		//create a notification that displays the import progress at the lower right corner
-		FNotificationInfo ImportNotificationInfo(NSLOCTEXT("FusionPatchImporter","FusionPatchImportProgressNotification","Importing Fusion Asset(s)..."));
+		FNotificationInfo ImportNotificationInfo(NSLOCTEXT("FusionPatchImporter","FusionPatchImportProgressNotification_InProgress","Importing Fusion Asset(s)..."));
 		ImportNotificationInfo.bFireAndForget = false;
 		TSharedPtr<SNotificationItem> ImportNotificationItem;
 		ImportNotificationItem = FSlateNotificationManager::Get().AddNotification(ImportNotificationInfo);
diff --git a/Engine/Plugins/Runtime/Harmonix/Source/HarmonixMetasound/Private/HarmonixMetasound/Nodes/DistortionNode.cpp b/Engine/Plugins/Runtime/Harmonix/Source/HarmonixMetasound/Private/HarmonixMetasound/Nodes/DistortionNode.cpp
--- a/Engine/Plugins/Runtime/Harmonix/Source/HarmonixMetasound/Private/HarmonixMetasound/Nodes/DistortionNode.cpp
+++ b/Engine/Plugins/Runtime/Harmonix/Source/HarmonixMetasound/Private/HarmonixMetasound/Nodes/DistortionNode.cpp
@@ -22,7 +22,7 @@
 
 #include "HAL/Platform.h"
 
-#define LOCTEXT_NAMESPACE "HarmonixMetaSound"
+#define LOCTEXT_NAMESPACE "HarmonixMetaSound_DistortionNode"
 
 namespace HarmonixMetasound::Effects::Distortion
 {
@@ -87,7 +87,7 @@ namespace HarmonixMetasound
 	{
 		namespace Inputs
 		{
-			METASOUND_PARAM(Audio, "In", "Audio Input");
+			METASOUND_PARAM(AudioIn, "In", "Audio Input");
 			METASOUND_PARAM(InputGain, "Input Gain", "Gain applied to the input signal");
 			METASOUND_PARAM(OutputGain, "Output Gain", "Gain applied to the output signal");
 			METASOUND_PARAM(WetGain, "Wet Gain", "The wet gain of the Distortion");
@@ -107,7 +107,7 @@ namespace HarmonixMetasound
 
 		namespace Outputs
 		{
-			METASOUND_PARAM(Audio, "Out", "Audio Output");
+			METASOUND_PARAM(AudioOut, "Out", "Audio Output");
 		}
 
 	}
@@ -196,7 +196,7 @@ namespace HarmonixMetasound
 
 				FInputVertexInterface InputInterface;
 
-				InputInterface.Add(TInputDataVertex<Metasound::FAudioBuffer>(METASOUND_GET_PARAM_NAME_AND_METADATA(Inputs::Audio)));
+				InputInterface.Add(TInputDataVertex<Metasound::FAudioBuffer>(METASOUND_GET_PARAM_NAME_AND_METADATA(Inputs::AudioIn)));
 				InputInterface.Add(TInputDataVertex<float>(METASOUND_GET_PARAM_NAME_AND_METADATA(Inputs::InputGain), HarmonixDsp::DBToLinear(DefaultSettings.InputGainDb)));
 				InputInterface.Add(TInputDataVertex<float>(METASOUND_GET_PARAM_NAME_AND_METADATA(Inputs::OutputGain), HarmonixDsp::DBToLinear(DefaultSettings.OutputGainDb)));
 				InputInterface.Add(TInputDataVertex<float>(METASOUND_GET_PARAM_NAME_AND_METADATA(Inputs::DryGain), DefaultSettings.DryGain));
@@ -216,7 +216,7 @@ namespace HarmonixMetasound
 				}
 
 				FOutputVertexInterface OutputInterface;
-				OutputInterface.Add(TOutputDataVertex<Metasound::FAudioBuffer>(METASOUND_GET_PARAM_NAME_AND_METADATA(Outputs::Audio)));
+				OutputInterface.Add(TOutputDataVertex<Metasound::FAudioBuffer>(METASOUND_GET_PARAM_NAME_AND_METADATA(Outputs::AudioOut)));
 
 				return { InputInterface, OutputInterface };
 			};
@@ -254,7 +254,7 @@ namespace HarmonixMetasound
 			const FInputVertexInterfaceData& InputData = InParams.InputData;
 
 			FAudioBufferReadRef InAudio = InputData.GetOrConstructDataReadReference<Metasound::FAudioBuffer>(
-				METASOUND_GET_PARAM_NAME(Inputs::Audio), InParams.OperatorSettings);
+				METASOUND_GET_PARAM_NAME(Inputs::AudioIn), InParams.OperatorSettings);
 			FFloatReadRef InInputGain = InputData.GetOrCreateDefaultDataReadReference<float>(METASOUND_GET_PARAM_NAME(Inputs::InputGain), InParams.OperatorSettings);
 			FFloatReadRef InOutputGain = InputData.GetOrCreateDefaultDataReadReference<float>(METASOUND_GET_PARAM_NAME(Inputs::OutputGain), InParams.OperatorSettings);
 			FFloatReadRef InWetGain = InputData.GetOrCreateDefaultDataReadReference<float>(METASOUND_GET_PARAM_NAME(Inputs::WetGain), InParams.OperatorSettings);
@@ -300,7 +300,7 @@ namespace HarmonixMetasound
 		virtual void BindInputs(FInputVertexInterfaceData& InVertexData) override
 		{
 			using namespace Distortion::PinNames;
-			InVertexData.BindReadVertex(METASOUND_GET_PARAM_NAME(Inputs::Audio), AudioIn);
+			InVertexData.BindReadVertex(METASOUND_GET_PARAM_NAME(Inputs::AudioIn), AudioIn);
 			InVertexData.BindReadVertex(METASOUND_GET_PARAM_NAME(Inputs::InputGain), InputGain);
 			InVertexData.BindReadVertex(METASOUND_GET_PARAM_NAME(Inputs::OutputGain), OutputGain);
 			InVertexData.BindReadVertex(METASOUND_GET_PARAM_NAME(Inputs::DryGain), DryGain);
@@ -323,7 +323,7 @@ namespace HarmonixMetasound
 		virtual void BindOutputs(FOutputVertexInterfaceData& InVertexData) override
 		{
 			using namespace Effects::Distortion::PinNames;
-			InVertexData.BindReadVertex(METASOUND_GET_PARAM_NAME(Outputs::Audio), AudioOut);
+			InVertexData.BindReadVertex(METASOUND_GET_PARAM_NAME(Outputs::AudioOut), AudioOut);
 		}
 
 		virtual FDataReferenceCollection GetInputs() const override
diff --git a/Engine/Plugins/Runtime/Harmonix/Source/HarmonixMetasound/Private/HarmonixMetasound/Nodes/MidiClockSubdivisionTriggerNode.cpp b/Engine/Plugins/Runtime/Harmonix/Source/HarmonixMetasound/Private/HarmonixMetasound/Nodes/MidiClockSubdivisionTriggerNode.cpp
--- a/Engine/Plugins/Runtime/Harmonix/Source/HarmonixMetasound/Private/HarmonixMetasound/Nodes/MidiClockSubdivisionTriggerNode.cpp
+++ b/Engine/Plugins/Runtime/Harmonix/Source/HarmonixMetasound/Private/HarmonixMetasound/Nodes/MidiClockSubdivisionTriggerNode.cpp
@@ -20,7 +20,7 @@
 
 DEFINE_LOG_CATEGORY_STATIC(LogMidiClockSubdivisionTrigger, Log, All);
 
-#define LOCTEXT_NAMESPACE "HarmonixMetaSound"
+#define LOCTEXT_NAMESPACE "HarmonixMetaSound_MidiClockSubdivisionTriggerNode"
 
 namespace HarmonixMetasound
 {
diff --git a/Engine/Plugins/Runtime/Harmonix/Source/HarmonixMetasound/Private/HarmonixMetasound/Nodes/MidiClockTimestampTriggerNode.cpp b/Engine/Plugins/Runtime/Harmonix/Source/HarmonixMetasound/Private/HarmonixMetasound/Nodes/MidiClockTimestampTriggerNode.cpp
--- a/Engine/Plugins/Runtime/Harmonix/Source/HarmonixMetasound/Private/HarmonixMetasound/Nodes/MidiClockTimestampTriggerNode.cpp
+++ b/Engine/Plugins/Runtime/Harmonix/Source/HarmonixMetasound/Private/HarmonixMetasound/Nodes/MidiClockTimestampTriggerNode.cpp
@@ -20,7 +20,7 @@
 
 DEFINE_LOG_CATEGORY_STATIC(LogMidiClockTimestampTrigger, Log, All);
 
-#define LOCTEXT_NAMESPACE "HarmonixMetaSound"
+#define LOCTEXT_NAMESPACE "HarmonixMetaSound_MidiClockTimestampTriggerNode"
 
 namespace HarmonixMetasound
 {
diff --git a/Engine/Plugins/Runtime/Harmonix/Source/HarmonixMetasound/Private/HarmonixMetasound/Nodes/MidiTextTriggerNode.cpp b/Engine/Plugins/Runtime/Harmonix/Source/HarmonixMetasound/Private/HarmonixMetasound/Nodes/MidiTextTriggerNode.cpp
--- a/Engine/Plugins/Runtime/Harmonix/Source/HarmonixMetasound/Private/HarmonixMetasound/Nodes/MidiTextTriggerNode.cpp
+++ b/Engine/Plugins/Runtime/Harmonix/Source/HarmonixMetasound/Private/HarmonixMetasound/Nodes/MidiTextTriggerNode.cpp
@@ -15,7 +15,7 @@
 
 DEFINE_LOG_CATEGORY_STATIC(LogMidiTextTrigger, Log, All);
 
-#define LOCTEXT_NAMESPACE "HarmonixMetaSound"
+#define LOCTEXT_NAMESPACE "HarmonixMetaSound_MidiTextTriggerNode"
 
 namespace HarmonixMetasound
 {
