MidiPlayCursor Smoothing Extrapolation

FMidiPlayCursor::AdvanceLowResCursors advances cursors even when the tracked time hasn't changed, in case a cursor is evaluating an approximation of the current time (i.e., smoothing).  This allows trackers to update at buffer rates, regardless of the video framerate, while smoothing cursors ensure time doesn't stop for graphical effects.

Cursors should only correct deviation from the expected time when the tracked time changes, which FSmoothedMidiPlayCursor now does.

Use GetCurrentLowResMs on the GameThread instead of GetCurrentHiResMs.

#tests Played Main Stage on Android and other platforms

diff --git a/Engine/Plugins/Runtime/Harmonix/Source/HarmonixMetasound/Private/HarmonixMetasound/Components/MetasoundMusicClockDriver.cpp b/Engine/Plugins/Runtime/Harmonix/Source/HarmonixMetasound/Private/HarmonixMetasound/Components/MetasoundMusicClockDriver.cpp
--- a/Engine/Plugins/Runtime/Harmonix/Source/HarmonixMetasound/Private/HarmonixMetasound/Components/MetasoundMusicClockDriver.cpp
+++ b/Engine/Plugins/Runtime/Harmonix/Source/HarmonixMetasound/Private/HarmonixMetasound/Components/MetasoundMusicClockDriver.cpp
@@ -13,18 +13,18 @@ bool FMetasoundMusicClockDriver::CalculateSongPosWithOffset(float MsOffset, ECal
 	// if we have an owner, ask them directly
 	if (const FMidiPlayCursorMgr* Owner = Cursor.GetOwner())
 	{
-		float RawMs = Owner->GetCurrentHiResMs();
+		float RawMs = Owner->GetCurrentLowResMs();
 		switch (Timebase)
 		{
 		case ECalibratedMusicTimebase::AudioRenderTime:
-			OutResult = Owner->CalculateSongPosWithOffsetMs((Clock->CurrentSmoothedAudioRenderSongPos.SecondsIncludingCountIn * 1000.0f) + MsOffset - RawMs);
+			OutResult = Owner->CalculateLowResSongPosWithOffsetMs((Clock->CurrentSmoothedAudioRenderSongPos.SecondsIncludingCountIn * 1000.0f) + MsOffset - RawMs);
 			break;
 		case ECalibratedMusicTimebase::ExperiencedTime:
-			OutResult = Owner->CalculateSongPosWithOffsetMs((Clock->CurrentPlayerExperiencedSongPos.SecondsIncludingCountIn * 1000.0f) + MsOffset - RawMs);
+			OutResult = Owner->CalculateLowResSongPosWithOffsetMs((Clock->CurrentPlayerExperiencedSongPos.SecondsIncludingCountIn * 1000.0f) + MsOffset - RawMs);
 			break;
 		case ECalibratedMusicTimebase::VideoRenderTime:
 		default:
-			OutResult = Owner->CalculateSongPosWithOffsetMs((Clock->CurrentVideoRenderSongPos.SecondsIncludingCountIn * 1000.0f) + MsOffset - RawMs);
+			OutResult = Owner->CalculateLowResSongPosWithOffsetMs((Clock->CurrentVideoRenderSongPos.SecondsIncludingCountIn * 1000.0f) + MsOffset - RawMs);
 			break;
 		}
 		return true;
@@ -239,10 +239,10 @@ void FMetasoundMusicClockDriver::RefreshCurrentSongPosFromCursor()
 	if (const FMidiPlayCursorMgr* Owner = Cursor.GetOwner())
 	{
 		Clock->CurrentClockAdvanceRate = Owner->GetCurrentAdvanceRate();
-		float RawMs = Owner->GetCurrentHiResMs();
-		Clock->CurrentPlayerExperiencedSongPos = Owner->CalculateSongPosWithOffsetMs(Clock->CurrentSmoothedAudioRenderSongPos.SecondsIncludingCountIn * 1000.0f - FHarmonixModule::GetMeasuredUserExperienceAndReactionToAudioRenderOffsetMs() - RawMs);
-		Clock->CurrentVideoRenderSongPos = Owner->CalculateSongPosWithOffsetMs(Clock->CurrentSmoothedAudioRenderSongPos.SecondsIncludingCountIn * 1000.0f - FHarmonixModule::GetMeasuredVideoToAudioRenderOffsetMs() - RawMs);
-		Clock->RawUnsmoothedAudioRenderPos = Owner->CalculateSongPosWithOffsetMs(0.0f);
+		float RawMs = Owner->GetCurrentLowResMs();
+		Clock->CurrentPlayerExperiencedSongPos = Owner->CalculateLowResSongPosWithOffsetMs(Clock->CurrentSmoothedAudioRenderSongPos.SecondsIncludingCountIn * 1000.0f - FHarmonixModule::GetMeasuredUserExperienceAndReactionToAudioRenderOffsetMs() - RawMs);
+		Clock->CurrentVideoRenderSongPos = Owner->CalculateLowResSongPosWithOffsetMs(Clock->CurrentSmoothedAudioRenderSongPos.SecondsIncludingCountIn * 1000.0f - FHarmonixModule::GetMeasuredVideoToAudioRenderOffsetMs() - RawMs);
+		Clock->RawUnsmoothedAudioRenderPos = Owner->CalculateLowResSongPosWithOffsetMs(0.0f);
 	}
 	else
 	{
diff --git a/Engine/Plugins/Runtime/Harmonix/Source/HarmonixMidi/Private/HarmonixMidi/MidiPlayCursorMgr.cpp b/Engine/Plugins/Runtime/Harmonix/Source/HarmonixMidi/Private/HarmonixMidi/MidiPlayCursorMgr.cpp
--- a/Engine/Plugins/Runtime/Harmonix/Source/HarmonixMidi/Private/HarmonixMidi/MidiPlayCursorMgr.cpp
+++ b/Engine/Plugins/Runtime/Harmonix/Source/HarmonixMidi/Private/HarmonixMidi/MidiPlayCursorMgr.cpp
@@ -85,10 +85,10 @@ void FMidiPlayCursorMgr::DetachFromMidiResource()
 	MidiDataChangeComplete(EMidiChangePositionCorrectMode::MaintainTick);
 }
 
-FMidiSongPos FMidiPlayCursorMgr::CalculateSongPosWithOffsetMs(float DeltaMs) const
+FMidiSongPos FMidiPlayCursorMgr::CalculateSongPosWithOffsetMs(float DeltaMs, bool IsLowRes) const
 {
 	FMidiSongPos OutSongPos;
-	float Ms = GetCurrentHiResMs() + DeltaMs;
+	float Ms = (IsLowRes ? GetCurrentLowResMs() : GetCurrentHiResMs()) + DeltaMs;
 	if (DoesLoop())
 	{
 		float LoopLengthMs = LoopEndMs - LoopStartMs;
@@ -103,7 +103,7 @@ FMidiSongPos FMidiPlayCursorMgr::CalculateSongPosWithOffsetMs(float DeltaMs) con
 	// currently only use the time authority for the tempo
 	if (TSharedPtr<FMidiPlayCursorMgr> TimeAuthorityPtr = TimeAuthority.Pin())
 	{
-		FMidiSongPos AuthoritySongPos = TimeAuthorityPtr->CalculateSongPosWithOffsetMs(DeltaMs);
+		FMidiSongPos AuthoritySongPos = TimeAuthorityPtr->CalculateSongPosWithOffsetMs(DeltaMs, IsLowRes);
 
 		OutSongPos.Tempo = AuthoritySongPos.Tempo;
 	}
@@ -504,7 +504,7 @@ void FMidiPlayCursorMgr::SeekTo(int32 Tick, int32 PreRollBars, bool IsRenderThre
 
 	// The hi-res cursors can just slam to the new position.
 	// If this is resultOfDirectMappedLoop, the hi-res cursors
-	// were already advanced to the end of the loop by the parent.
+	// were already advanced to the end of the loop by the master.
 	HiResTracker.Reset(Tick, NewPosMs, PreRollStartTick, PreRollStartMs, false);
 	if (!IsRenderThread)
 	{
@@ -745,12 +745,33 @@ bool FMidiPlayCursorMgr::AdvanceLowResCursors()
 		MsDiff = HiResCurrentMs - LowResTracker.CurrentMs;
 		if (MsDiff <= 0.0f)
 		{
+			UpdateLowResCursors(LowResTracker);
 			return false;
 		}
 	}
 	return AdvanceTrackerByDeltaMs(MsDiff, LowResTracker, true, true);
 }
 
+void FMidiPlayCursorMgr::UpdateLowResCursors(FMidiPlayCursorTracker& Tracker)
+{
+	TraversingLowResCursors = true;
+	for (auto it = Tracker.Cursors.begin(); it != Tracker.Cursors.end();)
+	{
+		if (!it.GetNode()->UpdateWithTrackerUnchanged())
+		{
+			auto DeadIt = it;
+			++it;
+			Tracker.Cursors.Remove(DeadIt.GetNode());
+			DeadIt.GetNode()->SetOwner(nullptr, nullptr);
+		}
+		else
+		{
+			++it;
+		}
+	}
+	TraversingLowResCursors = false;
+}
+
 bool FMidiPlayCursorMgr::AdvanceTrackerByDeltaMs(float Ms, FMidiPlayCursorTracker& Tracker, bool IsLowRes, bool Broadcast)
 {
 	Tracker.ElapsedMs += Ms;
diff --git a/Engine/Plugins/Runtime/Harmonix/Source/HarmonixMidi/Private/HarmonixMidi/SmoothedMidiPlayCursor.cpp b/Engine/Plugins/Runtime/Harmonix/Source/HarmonixMidi/Private/HarmonixMidi/SmoothedMidiPlayCursor.cpp
--- a/Engine/Plugins/Runtime/Harmonix/Source/HarmonixMidi/Private/HarmonixMidi/SmoothedMidiPlayCursor.cpp
+++ b/Engine/Plugins/Runtime/Harmonix/Source/HarmonixMidi/Private/HarmonixMidi/SmoothedMidiPlayCursor.cpp
@@ -45,17 +45,29 @@ void FSmoothedMidiPlayCursor::Reset(bool ForceNoBroadcast)
 	CurrentSongPos.Reset();
 }
 
+bool FSmoothedMidiPlayCursor::UpdateWithTrackerUnchanged()
+{
+	if (!FMidiPlayCursor::UpdateWithTrackerUnchanged())
+	{
+		return false;
+	}
+
+	SyncSmoothingTimer(false);
+	UpdateSongPosition();
+	return true;
+}
+
 void FSmoothedMidiPlayCursor::AdvanceByTicks(bool processLoops /*= true*/, bool broadcast /*= true*/, bool isPreRoll /*= false*/)
 {
 	FMidiPlayCursor::AdvanceByTicks(processLoops, broadcast, isPreRoll);
-	SyncSmoothingTimer();
+	SyncSmoothingTimer(true);
 	UpdateSongPosition();
 }
 
 void FSmoothedMidiPlayCursor::AdvanceByMs(bool processLoops /*= true*/, bool broadcast /*= true*/, bool isPreRoll /*= false*/)
 {
 	FMidiPlayCursor::AdvanceByMs(processLoops, broadcast, isPreRoll);
-	SyncSmoothingTimer();
+	SyncSmoothingTimer(true);
 	UpdateSongPosition();
 }
 
@@ -64,7 +76,7 @@ void FSmoothedMidiPlayCursor::OnLoop(int loopStartTick, int loopEndTick)
 	LoopedThisPass = true;
 }
 
-void FSmoothedMidiPlayCursor::SyncSmoothingTimer()
+void FSmoothedMidiPlayCursor::SyncSmoothingTimer(bool bEnableErrorCorrection)
 {
 	if (!Owner)
 	{
@@ -80,48 +92,52 @@ void FSmoothedMidiPlayCursor::SyncSmoothingTimer()
 
 	float RawMs = CurrentMs;
 	float SmoothMs = SmoothingTimer.Ms();
-	float Error = (SmoothMs - RawMs) * (1.0f / (float)SmoothingTimer.GetSpeed());
-	ErrorTracker.Push(Error);
-	float MinRecentError = ErrorTracker.Min();
-	float AbsError = FMath::Abs(MinRecentError);
 
-	// if we have a massive jump just slam the thing to the current "raw" time and be done with it...
-	if (FMath::Abs(Error) > SmoothedPlayCursor::kMassiveJump)
+	if (bEnableErrorCorrection)
 	{
-		UE_LOG(LogSmoothedMidiPlayCursor, Verbose, TEXT("Smoothing: Massive Error = %fms. Slamming."), Error);
-		SmoothingTimer.SetSpeed(BaseSpeed);
-		SmoothingTimer.Reset(RawMs);
-		SmoothedMs = RawMs;
-		SmoothedMsDelta = 0.0f;
-		SmoothedTick = Owner->GetTempoMap().MsToTick(SmoothedMs);
-		ErrorTracker.Reset();
-		return;
-	}
+		float Error = (SmoothMs - RawMs) * (1.0f / (float)SmoothingTimer.GetSpeed());
+		ErrorTracker.Push(Error);
+		float MinRecentError = ErrorTracker.Min();
+		float AbsError = FMath::Abs(MinRecentError);
 
-	if (AbsError < SmoothedPlayCursor::kMinorErrorThreshold)
-	{
-		if (SmoothingTimer.GetSpeed() != BaseSpeed)
+		// if we have a massive jump just slam the thing to the current "raw" time and be done with it...
+		if (FMath::Abs(Error) > SmoothedPlayCursor::kMassiveJump)
 		{
-			UE_LOG(LogSmoothedMidiPlayCursor, Verbose, TEXT("Smoothing: Error = %fms. Running normal speed."), MinRecentError);
+			UE_LOG(LogSmoothedMidiPlayCursor, Verbose, TEXT("Smoothing: Massive Error = %fms. Slamming."), Error);
 			SmoothingTimer.SetSpeed(BaseSpeed);
+			SmoothingTimer.Reset(RawMs);
+			SmoothedMs = RawMs;
+			SmoothedMsDelta = 0.0f;
+			SmoothedTick = Owner->GetTempoMap().MsToTick(SmoothedMs);
+			ErrorTracker.Reset();
+			return;
 		}
-	}
-	else if (AbsError < SmoothedPlayCursor::kMajorErrorThreshold)
-	{
-		float NewSpeedFactor = BaseSpeed + (MinRecentError < 0.0f ? SmoothedPlayCursor::kMinorCorrectionFactor : -SmoothedPlayCursor::kMinorCorrectionFactor);
-		if (SmoothingTimer.GetSpeed() != NewSpeedFactor)
+
+		if (AbsError < SmoothedPlayCursor::kMinorErrorThreshold)
 		{
-			UE_LOG(LogSmoothedMidiPlayCursor, Verbose, TEXT("Smoothing: Small error = %f. Adjusting speedfactor to %f."), MinRecentError, NewSpeedFactor);
-			SmoothingTimer.SetSpeed(NewSpeedFactor);
+			if (SmoothingTimer.GetSpeed() != BaseSpeed)
+			{
+				UE_LOG(LogSmoothedMidiPlayCursor, Verbose, TEXT("Smoothing: Error = %fms. Running normal speed."), MinRecentError);
+				SmoothingTimer.SetSpeed(BaseSpeed);
+			}
 		}
-	}
-	else
-	{
-		float NewSpeedFactor = BaseSpeed + (MinRecentError < 0.0f ? SmoothedPlayCursor::kMajorCorrectionFactor : -SmoothedPlayCursor::kMajorCorrectionFactor);
-		if (SmoothingTimer.GetSpeed() != NewSpeedFactor)
+		else if (AbsError < SmoothedPlayCursor::kMajorErrorThreshold)
+		{
+			float NewSpeedFactor = BaseSpeed + (MinRecentError < 0.0f ? SmoothedPlayCursor::kMinorCorrectionFactor : -SmoothedPlayCursor::kMinorCorrectionFactor);
+			if (SmoothingTimer.GetSpeed() != NewSpeedFactor)
+			{
+				UE_LOG(LogSmoothedMidiPlayCursor, Verbose, TEXT("Smoothing: Small error = %f. Adjusting speedfactor to %f."), MinRecentError, NewSpeedFactor);
+				SmoothingTimer.SetSpeed(NewSpeedFactor);
+			}
+		}
+		else
 		{
-			UE_LOG(LogSmoothedMidiPlayCursor, Verbose, TEXT("Smoothing: Significant error = %f. Adjusting speedfactor to %f."), MinRecentError, NewSpeedFactor);
-			SmoothingTimer.SetSpeed(NewSpeedFactor);
+			float NewSpeedFactor = BaseSpeed + (MinRecentError < 0.0f ? SmoothedPlayCursor::kMajorCorrectionFactor : -SmoothedPlayCursor::kMajorCorrectionFactor);
+			if (SmoothingTimer.GetSpeed() != NewSpeedFactor)
+			{
+				UE_LOG(LogSmoothedMidiPlayCursor, Verbose, TEXT("Smoothing: Significant error = %f. Adjusting speedfactor to %f."), MinRecentError, NewSpeedFactor);
+				SmoothingTimer.SetSpeed(NewSpeedFactor);
+			}
 		}
 	}
 
@@ -162,5 +178,5 @@ void FSmoothedMidiPlayCursor::UpdateSongPosition()
 		return;
 	}
 
-	CurrentSongPos = Owner->CalculateSongPosWithOffsetMs(SmoothedMsDelta);
+	CurrentSongPos = Owner->CalculateLowResSongPosWithOffsetMs(SmoothedMsDelta);
 }
diff --git a/Engine/Plugins/Runtime/Harmonix/Source/HarmonixMidi/Public/HarmonixMidi/MidiPlayCursor.h b/Engine/Plugins/Runtime/Harmonix/Source/HarmonixMidi/Public/HarmonixMidi/MidiPlayCursor.h
--- a/Engine/Plugins/Runtime/Harmonix/Source/HarmonixMidi/Public/HarmonixMidi/MidiPlayCursor.h
+++ b/Engine/Plugins/Runtime/Harmonix/Source/HarmonixMidi/Public/HarmonixMidi/MidiPlayCursor.h
@@ -23,6 +23,8 @@ public:
 	virtual void Reset(bool ForceNoBroadcast = false);
 	virtual bool IsDone() const;
 
+	virtual bool UpdateWithTrackerUnchanged() { return !UnregisterASAP; }
+
 	bool Advance(bool IsLowRes = true);
 	bool AdvanceAsPreRoll();
 
diff --git a/Engine/Plugins/Runtime/Harmonix/Source/HarmonixMidi/Public/HarmonixMidi/MidiPlayCursorMgr.h b/Engine/Plugins/Runtime/Harmonix/Source/HarmonixMidi/Public/HarmonixMidi/MidiPlayCursorMgr.h
--- a/Engine/Plugins/Runtime/Harmonix/Source/HarmonixMidi/Public/HarmonixMidi/MidiPlayCursorMgr.h
+++ b/Engine/Plugins/Runtime/Harmonix/Source/HarmonixMidi/Public/HarmonixMidi/MidiPlayCursorMgr.h
@@ -22,7 +22,17 @@ public:
 	void AttachToMidiResource(TSharedPtr<FMidiFileData> MidiDataProxy, bool ResetTrackersToStart = true, int32 PreRollBars = 0);
 	void DetachFromMidiResource();
 
-	FMidiSongPos CalculateSongPosWithOffsetMs(float Ms) const;
+	FMidiSongPos CalculateLowResSongPosWithOffsetMs(float DeltaMs) const
+	{
+		return CalculateSongPosWithOffsetMs(DeltaMs, true);
+	}
+
+	FMidiSongPos CalculateHiResSongPosWithOffsetMs(float DeltaMs) const
+	{
+		return CalculateSongPosWithOffsetMs(DeltaMs, false);
+	}
+
+	FMidiSongPos CalculateSongPosWithOffsetMs(float Ms, bool IsLowRes) const;
 
 	//////////////////////////////////////////////////////////////////////////
 	// Play cursor management
@@ -73,6 +83,7 @@ public:
 
 	// returns whether the cursors looped
 	bool AdvanceLowResCursors();
+	void UpdateLowResCursors(FMidiPlayCursorTracker& Tracker);
 
 	void ResetTrackers(); //reset to the beginning of the song
 
diff --git a/Engine/Plugins/Runtime/Harmonix/Source/HarmonixMidi/Public/HarmonixMidi/SmoothedMidiPlayCursor.h b/Engine/Plugins/Runtime/Harmonix/Source/HarmonixMidi/Public/HarmonixMidi/SmoothedMidiPlayCursor.h
--- a/Engine/Plugins/Runtime/Harmonix/Source/HarmonixMidi/Public/HarmonixMidi/SmoothedMidiPlayCursor.h
+++ b/Engine/Plugins/Runtime/Harmonix/Source/HarmonixMidi/Public/HarmonixMidi/SmoothedMidiPlayCursor.h
@@ -18,6 +18,7 @@ public:
 
 	//** BEGIN FMidiPlayCursor
 	virtual void Reset(bool ForceNoBroadcast = false) override;
+	virtual bool UpdateWithTrackerUnchanged() override;
 	virtual void AdvanceByTicks(bool processLoops = true, bool broadcast = true, bool isPreRoll = false) override;
 	virtual void AdvanceByMs(bool processLoops = true, bool broadcast = true, bool isPreRoll = false) override;
 	virtual void OnLoop(int loopStartTick, int loopEndTick) override;
@@ -39,6 +40,6 @@ private:
 	FMidiSongPos CurrentSongPos;
 
 	void  SetSpeed(float NewSpeed);
-	void  SyncSmoothingTimer();
+	void  SyncSmoothingTimer(bool bEnableErrorCorrection);
 	void  UpdateSongPosition();
 };
