Deprecate "Presets" from FusionPatches. Only have one "Settings".

Serialize the single "Settings" with the first element from the "Presets", which should be the only one and the default one.

#tests PIE

diff --git a/Engine/Plugins/Runtime/Harmonix/Source/HarmonixDsp/Private/HarmonixDsp/FusionSampler/FusionPatch.cpp b/Engine/Plugins/Runtime/Harmonix/Source/HarmonixDsp/Private/HarmonixDsp/FusionSampler/FusionPatch.cpp
--- a/Engine/Plugins/Runtime/Harmonix/Source/HarmonixDsp/Private/HarmonixDsp/FusionSampler/FusionPatch.cpp
+++ b/Engine/Plugins/Runtime/Harmonix/Source/HarmonixDsp/Private/HarmonixDsp/FusionSampler/FusionPatch.cpp
@@ -31,35 +31,13 @@ TSharedPtr<Audio::IProxyData> UFusionPatch::CreateProxyData(const Audio::FProxyD
 
 UFusionPatch::UFusionPatch() : Super()
 {
-	FusionPatchData.CurrentPresetIndex = 0;
-	FFusionPatchSettings& DefaultSettings = FusionPatchData.Presets.Emplace_GetRef();
+	FFusionPatchSettings& DefaultSettings = FusionPatchData.Settings;
 	FAdsrSettings& VolumeAdsr = DefaultSettings.Adsrs.Volume;
 	VolumeAdsr.Target = EAdsrTarget::Volume;
 	VolumeAdsr.Depth = 1.0f;
 	VolumeAdsr.IsEnabled = true;
 }
 
-void UFusionPatch::SetCurrentPresetIndex(int32 Index)
-{
-	FusionPatchData.SetCurrentPresetIndex(Index);
-
-	UpdateRenderableForNonTrivialChange();
-}
-
-void UFusionPatch::AddPreset(FFusionPatchSettings& Preset)
-{
-	FusionPatchData.AddPreset(Preset);
-
-	UpdateRenderableForNonTrivialChange();
-}
-
-void UFusionPatch::UpdatePreset(int32 Index, FFusionPatchSettings& Preset)
-{
-	FusionPatchData.UpdatePreset(Index, Preset);
-
-	UpdateRenderableForNonTrivialChange();
-}
-
 void UFusionPatch::UpdateRenderableForNonTrivialChange()
 {
 	if (!RenderableFusionPatchData)
@@ -92,7 +70,7 @@ void UFusionPatch::PostEditChangeChainProperty(FPropertyChangedChainEvent& Prope
 		return;
 	}
 
-	for (FKeyzoneSettings& Keyzone : FusionPatchData.GetKeyzones())
+	for (FKeyzoneSettings& Keyzone : FusionPatchData.Keyzones)
 	{
 		int8 MinNote = Keyzone.MinNote;
 		int8 MaxNote = Keyzone.MaxNote;
@@ -191,6 +169,14 @@ void UFusionPatch::Serialize(FArchive& Ar)
 		UE_LOG(LogFusionPatch, Warning, TEXT("Fusion patch was loaded with an outdated version. Asset requires reimport: %s"), *GetPathName());
 	}
 
+	if (Ar.IsLoading() && Version < FFusionPatchCustomVersion::DeprecatedPresets)
+	{
+		if (FusionPatchData.Presets_DEPRECATED.IsValidIndex(0))
+		{
+			FusionPatchData.Settings = FusionPatchData.Presets_DEPRECATED[0];
+		}
+	}
+
 	const UStretcherAndPitchShifterFactoryConfig* FactoryConfig = GetDefault<UStretcherAndPitchShifterFactoryConfig>();
 	for (FKeyzoneSettings& Keyzone : FusionPatchData.Keyzones)
 	{
@@ -206,33 +192,6 @@ void UFusionPatch::PostLoad()
 	Super::PostLoad();
 }
 
-TArray<FKeyzoneSettings>& FFusionPatchData::GetKeyzones()
-{
-	return Keyzones;
-}
-
-const TArray<FKeyzoneSettings>& FFusionPatchData::GetKeyzones() const
-{
-	return Keyzones;
-}
-
-int32 FFusionPatchData::GetCurrentPresetIndex() const
-{
-	return CurrentPresetIndex;
-}
-
-void FFusionPatchData::SetCurrentPresetIndex(int32 Index)
-{
-	ensureAlwaysMsgf(Presets.Num() > 0, TEXT("Fusion Patch Data presets have not been initialized"));
-	ensureAlwaysMsgf(Presets.IsValidIndex(Index), TEXT("assigning Fusion Patch Data to an invalid preset"));
-	CurrentPresetIndex = FMath::Clamp(Index, 0, Presets.Num());
-}
-
-int32 FFusionPatchData::GetNumPresets() const
-{
-	return Presets.Num();
-}
-
 void FFusionPatchData::InitProxyData(const Audio::FProxyDataInitParams& InitParams)
 {
 	for (auto& Keyzone : Keyzones)
@@ -241,36 +200,6 @@ void FFusionPatchData::InitProxyData(const Audio::FProxyDataInitParams& InitPara
 	}
 }
 
-void FFusionPatchData::AddPreset(const FFusionPatchSettings& Preset)
-{
-	// There shouldn't be that many presets, so this should be a pretty fast search
-	int32 Index = FindPresetIndex(Preset.Name);
-
-	// replace the preset of the same name with the new preset
-	if (Index != INDEX_NONE)
-	{
-		Presets[Index] = Preset;
-	}
-	else
-	{
-		Presets.Add(Preset);
-	}
-}
-
-int32 FFusionPatchData::FindPresetIndex(FName PresetName) const
-{
-	return Presets.IndexOfByPredicate([&PresetName](const FFusionPatchSettings& Elem)
-		{
-			return Elem.Name == PresetName;
-		});
-}
-
-void FFusionPatchData::UpdatePreset(int32 Index, const FFusionPatchSettings& Preset)
-{
-	check(Presets.IsValidIndex(Index));
-	Presets[Index] = Preset;
-}
-
 void FFusionPatchData::DisconnectSampler(const FFusionSampler* Sampler)
 {
 	for (FKeyzoneSettings& Keyzone : Keyzones)
@@ -280,38 +209,4 @@ void FFusionPatchData::DisconnectSampler(const FFusionSampler* Sampler)
 			Keyzone.SingletonFusionVoicePool->SamplerDisconnecting(Sampler);
 		}
 	}
-}
-
-bool FFusionPatchData::HasPreset(FName PresetName) const
-{
-	return Presets.ContainsByPredicate([&PresetName](const FFusionPatchSettings& Elem)
-		{
-			return Elem.Name == PresetName;
-		});
-}
-
-const FFusionPatchSettings& FFusionPatchData::GetPreset(int32 Index) const
-{
-	check(Presets.IsValidIndex(Index));
-	return Presets[Index];
-}
-
-FFusionPatchSettings& FFusionPatchData::GetPreset(int32 Index)
-{
-	check(Presets.IsValidIndex(Index));
-	return Presets[Index];
-}
-
-const FFusionPatchSettings& FFusionPatchData::GetCurrentPreset() const
-{
-	check(Presets.Num() > 0);
-	check(Presets.IsValidIndex(CurrentPresetIndex));
-	return Presets[GetCurrentPresetIndex()];
-}
-
-FFusionPatchSettings& FFusionPatchData::GetCurrentPreset()
-{
-	check(Presets.Num() > 0);
-	check(Presets.IsValidIndex(CurrentPresetIndex));
-	return Presets[CurrentPresetIndex];
 }
\ No newline at end of file
diff --git a/Engine/Plugins/Runtime/Harmonix/Source/HarmonixDsp/Private/HarmonixDsp/FusionSampler/FusionPatchCustomVersion.h b/Engine/Plugins/Runtime/Harmonix/Source/HarmonixDsp/Private/HarmonixDsp/FusionSampler/FusionPatchCustomVersion.h
--- a/Engine/Plugins/Runtime/Harmonix/Source/HarmonixDsp/Private/HarmonixDsp/FusionSampler/FusionPatchCustomVersion.h
+++ b/Engine/Plugins/Runtime/Harmonix/Source/HarmonixDsp/Private/HarmonixDsp/FusionSampler/FusionPatchCustomVersion.h
@@ -15,6 +15,8 @@ struct HARMONIXDSP_API FFusionPatchCustomVersion
 
 		PanImportingFromDTAFixed = 3,
 
+		DeprecatedPresets = 4,
+
 		// -----<new versions can be added above this line>-------------------------------------------------
 		VersionPlusOne,
 		LatestVersion = VersionPlusOne - 1
diff --git a/Engine/Plugins/Runtime/Harmonix/Source/HarmonixDsp/Private/HarmonixDsp/FusionSampler/FusionSampler.cpp b/Engine/Plugins/Runtime/Harmonix/Source/HarmonixDsp/Private/HarmonixDsp/FusionSampler/FusionSampler.cpp
--- a/Engine/Plugins/Runtime/Harmonix/Source/HarmonixDsp/Private/HarmonixDsp/FusionSampler/FusionSampler.cpp
+++ b/Engine/Plugins/Runtime/Harmonix/Source/HarmonixDsp/Private/HarmonixDsp/FusionSampler/FusionSampler.cpp
@@ -55,7 +55,6 @@ FFusionSampler::FFusionSampler()
 	, CurrentTempoBPM(120.0f)
 	, RawPitchMultiplier(1.0f)
 	, Transposition(0)
-	, PresetIdx(0)
 	, KeyzoneSelectMode(EKeyzoneSelectMode::Layers)
 {
 	// Can't use 'alignas' because we may have been allocated in an array of FusionSamplers.
@@ -687,38 +686,33 @@ void FFusionSampler::SetPatch(FFusionPatchData* PatchData)
 
 	if (FusionPatchData)
 	{
-		PresetIdx = FusionPatchData->GetCurrentPresetIndex();
-		ApplyPatchPresetSettings();
+		ApplyPatchSettings();
 	}
 }
 
-void FFusionSampler::ApplyPatchPresetSettings()
+void FFusionSampler::ApplyPatchSettings()
 {
 	check(FusionPatchData);
-
-	check(FusionPatchData->GetNumPresets() > 0);
-	PresetIdx = FMath::Clamp<size_t>(PresetIdx, 0, FusionPatchData->GetNumPresets() - 1);
-
-	const FFusionPatchSettings& Preset = FusionPatchData->GetPreset(PresetIdx);
-	SetTrimVolume(Preset.VolumeDb);
-	SetMinPitchBendCents(Preset.DownPitchBendCents);
-	SetMaxPitchBendCents(Preset.UpPitchBendCents);
-	SetMaxNumVoices(Preset.MaxVoices);
-	SetStartPointMs(Preset.StartPointOffsetMs);
-	SetFineTuneCents(Preset.FineTuneCents);
-	SetPan(Preset.PannerDetails);
-	KeyzoneSelectMode = Preset.KeyzoneSelectMode;
+	const FFusionPatchSettings& Settings = FusionPatchData->GetSettings();
+	SetTrimVolume(Settings.VolumeDb);
+	SetMinPitchBendCents(Settings.DownPitchBendCents);
+	SetMaxPitchBendCents(Settings.UpPitchBendCents);
+	SetMaxNumVoices(Settings.MaxVoices);
+	SetStartPointMs(Settings.StartPointOffsetMs);
+	SetFineTuneCents(Settings.FineTuneCents);
+	SetPan(Settings.PannerDetails);
+	KeyzoneSelectMode = Settings.KeyzoneSelectMode;
 
 	for (EAdsrIndex Idx : TEnumRange<EAdsrIndex>())
 	{
-		AdsrSettings[Idx].CopySettings(Preset.Adsrs[Idx]);
-		AdsrSettings[Idx].CopyCurveTables(Preset.Adsrs[Idx]);
+		AdsrSettings[Idx].CopySettings(Settings.Adsrs[Idx]);
+		AdsrSettings[Idx].CopyCurveTables(Settings.Adsrs[Idx]);
 	}
 
 	//-------------------------------------
 	// get the filter info
 	//------------------------------------
-	FilterSettings = Preset.Filter;
+	FilterSettings = Settings.Filter;
 
 	//------------------------------------
 	// Read in Lfo settings
@@ -726,14 +720,14 @@ void FFusionSampler::ApplyPatchPresetSettings()
 
 	for (ELfoIndex Idx : TEnumRange<ELfoIndex>())
 	{
-		LfoSettings[Idx].CopySettings(Preset.Lfos[Idx]);
+		LfoSettings[Idx].CopySettings(Settings.Lfos[Idx]);
 	}
 
 	//------------------------------------
 	// Read in portamento settings
 	//------------------------------------
 	{
-		const FPortamentoSettings& Portamento = Preset.Portamento;
+		const FPortamentoSettings& Portamento = Settings.Portamento;
 		SetIsPortamentoEnabled(Portamento.IsEnabled);
 		SetPortamentoMode(Portamento.Mode);
 		SetPortamentoTime(Portamento.Seconds);
@@ -742,7 +736,7 @@ void FFusionSampler::ApplyPatchPresetSettings()
 	// read in delay settings
 	//-------------------------------------
 	{
-		const FDelaySettings& DelaySettings = Preset.Delay;
+		const FDelaySettings& DelaySettings = Settings.Delay;
 		SetDelayEnabled(DelaySettings.IsEnabled);
 		Delay.SetTimeSyncOption(DelaySettings.TimeSyncOption);
 		Delay.SetDelaySeconds(DelaySettings.TimeSeconds);
@@ -768,7 +762,7 @@ void FFusionSampler::ApplyPatchPresetSettings()
 	// read in bitcrusher settings
 	//-------------------------------------
 	{
-		const FBitCrusherSettings& CrushSettings = Preset.BitCrusher;
+		const FBitCrusherSettings& CrushSettings = Settings.BitCrusher;
 		SetBitCrusherEnabled(CrushSettings.IsEnabled);
 		BitCrusher.Reset();
 		BitCrusher.SetWetGain(CrushSettings.WetGain);
@@ -780,7 +774,7 @@ void FFusionSampler::ApplyPatchPresetSettings()
 	// read in vocoder settings
 	//-------------------------------------
 	{
-		const FVocoderSettings& VocoderSettings = Preset.Vocoder;
+		const FVocoderSettings& VocoderSettings = Settings.Vocoder;
 		VocoderEnabled = VocoderSettings.IsEnabled;
 		Vocoder.SetBandConfig(VocoderSettings.BandConfig);
 		Vocoder.SetSoloing(VocoderSettings.Soloing);
@@ -800,7 +794,7 @@ void FFusionSampler::ApplyPatchPresetSettings()
 	// read in distortion settings
 	//-------------------------------------
 	{
-		const FDistortionSettingsV1& DistortionSettings = Preset.Distortion;
+		const FDistortionSettingsV1& DistortionSettings = Settings.Distortion;
 		DistortionEnabled = DistortionSettings.IsEnabled;
 		Distortion.SetSampleRate((int32)GetSamplesPerSecond());
 		Distortion.SetInputGainDb(DistortionSettings.InputGainDb);
@@ -1416,7 +1410,6 @@ void FFusionSampler::GetController(MidiConstants::EControllerID InController, in
 	{
 	case MidiConstants::EControllerID::BankSelection:
 	{
-		Msb = (int8)PresetIdx;
 		break;
 	}
 
@@ -1624,8 +1617,10 @@ void FFusionSampler::SetController(MidiConstants::EControllerID InController, fl
 	switch (InController)
 	{
 	case MidiConstants::EControllerID::BankSelection: 
-	{ 
-		PresetIdx = (int32)InValue; ApplyPatchPresetSettings(); 
+	{
+		// swap out the fusion patch with the new desired patch
+
+		// ApplyPatchSettings(); 
 		break;
 	}
 	case MidiConstants::EControllerID::Volume: 
diff --git a/Engine/Plugins/Runtime/Harmonix/Source/HarmonixDsp/Public/HarmonixDsp/FusionSampler/FusionPatch.h b/Engine/Plugins/Runtime/Harmonix/Source/HarmonixDsp/Public/HarmonixDsp/FusionSampler/FusionPatch.h
--- a/Engine/Plugins/Runtime/Harmonix/Source/HarmonixDsp/Public/HarmonixDsp/FusionSampler/FusionPatch.h
+++ b/Engine/Plugins/Runtime/Harmonix/Source/HarmonixDsp/Public/HarmonixDsp/FusionSampler/FusionPatch.h
@@ -27,44 +27,24 @@ public:
 
 	FFusionPatchData() {};
 
-	const TArray<FKeyzoneSettings>& GetKeyzones() const;
-	TArray<FKeyzoneSettings>& GetKeyzones();
-
-	int32 GetNumPresets() const;
-	int32 GetCurrentPresetIndex() const;
-	void SetCurrentPresetIndex(int32 Index);
-
-	bool IsCurrentPresetIndexValid() const { return Presets.IsValidIndex(CurrentPresetIndex); }
-	bool IsPresetIndexValid(int32 Index) const { return Presets.IsValidIndex(Index); }
-
-	bool HasPreset(FName PresetName) const;
-
-	// returns the index for the preset of the given name
-	// returns INDEX_NONE if there is no preset with the given name
-	int32 FindPresetIndex(FName PresetName) const;
-
-	const FFusionPatchSettings& GetPreset(int32 Index) const;
-	FFusionPatchSettings& GetPreset(int32 Index);
-	
-	const FFusionPatchSettings& GetCurrentPreset() const;
-	FFusionPatchSettings& GetCurrentPreset();
+	const TArray<FKeyzoneSettings>& GetKeyzones() const { return Keyzones; };
+	const FFusionPatchSettings& GetSettings() const { return Settings; }
+	void UpdateSettings(const FFusionPatchSettings& InSettings) { Settings = InSettings; }
 
 	void InitProxyData(const Audio::FProxyDataInitParams& InitParams);
-	void AddPreset(const FFusionPatchSettings& Preset);
-	void UpdatePreset(int32 Index, const FFusionPatchSettings& Preset);
 
 	void DisconnectSampler(const FFusionSampler* sampler);
 
 private:
 
-	UPROPERTY(EditAnywhere, Category="Presets")
+	UPROPERTY(EditAnywhere, Category="Keyzones")
 	TArray<FKeyzoneSettings> Keyzones;
 
-	UPROPERTY(EditAnywhere, Category="Presets")
-	TArray<FFusionPatchSettings> Presets;
-	
-	UPROPERTY(EditAnywhere, Category="Presets")
-	int32 CurrentPresetIndex = -1;
+	UPROPERTY(EditAnywhere, Category="Settings")
+	FFusionPatchSettings Settings;
+
+	UPROPERTY()
+	TArray<FFusionPatchSettings> Presets_DEPRECATED;
 };
 
 // This next macro does a few things. 
@@ -101,26 +81,10 @@ public:
 
 	UFusionPatch();
 
-	int32 GetCurrentPresetIndex() const { return FusionPatchData.GetCurrentPresetIndex(); }
-	void SetCurrentPresetIndex(int32 Index);
-	int32 GetNumPresets() const { return FusionPatchData.GetNumPresets(); }
-
-	bool IsCurrentPresetIndexValid() const { return FusionPatchData.IsCurrentPresetIndexValid(); }
-	bool IsPresetIndexValid(int32 Index) const { return FusionPatchData.IsPresetIndexValid(Index); }
-
-	bool HasPreset(FName PresetName) const { return FusionPatchData.HasPreset(PresetName); }
-
-	// returns the index for the preset of the given name
-	// returns INDEX_NONE if there is no preset with the given name
-	int32 FindPresetIndex(FName PresetName) const { return FusionPatchData.FindPresetIndex(PresetName); }
-
-	const FFusionPatchSettings& GetPreset(int32 Index) const { return FusionPatchData.GetPreset(Index); }
-	const FFusionPatchSettings& GetCurrentPreset() const { return FusionPatchData.GetCurrentPreset(); }
+	const FFusionPatchSettings& GetSettings() const { return FusionPatchData.GetSettings(); }
+	void UpdateSettings(const FFusionPatchSettings& InSettings);
 	const TArray<FKeyzoneSettings>& GetKeyzones() const { return FusionPatchData.GetKeyzones(); }
 
-	void AddPreset(FFusionPatchSettings& Preset);
-	void UpdatePreset(int32 Index, FFusionPatchSettings& Preset);
-
 #if WITH_EDITORONLY_DATA
 
 	UPROPERTY(VisibleAnywhere, Instanced, Category = ImportSettings)
diff --git a/Engine/Plugins/Runtime/Harmonix/Source/HarmonixDsp/Public/HarmonixDsp/FusionSampler/FusionSampler.h b/Engine/Plugins/Runtime/Harmonix/Source/HarmonixDsp/Public/HarmonixDsp/FusionSampler/FusionSampler.h
--- a/Engine/Plugins/Runtime/Harmonix/Source/HarmonixDsp/Public/HarmonixDsp/FusionSampler/FusionSampler.h
+++ b/Engine/Plugins/Runtime/Harmonix/Source/HarmonixDsp/Public/HarmonixDsp/FusionSampler/FusionSampler.h
@@ -168,9 +168,9 @@ protected:
 	bool RemoveChild(FVirtualInstrument* child);
 	void DumpAllChildren();
 
-	// apply all preset-specific settings from the patch.
-	// call whenever the preset or the patch change.
-	void ApplyPatchPresetSettings();
+	// apply all settings from the patch.
+	// call whenever the patch is set
+	void ApplyPatchSettings();
 
 	void  SetFineTuneCents(float cents);
 	float GetFineTuneCents() const;
@@ -405,8 +405,6 @@ protected:
 	TArray<FVirtualInstrument*> FusionPreChildren;
 	TArray<FVirtualInstrument*> FusionPostChildren;
 
-protected:
-
 	// we'll keep track of the 'current tempo' as it is
 	// needed by any 'beat sync' effects...
 	float CurrentTempoBPM = 120.0f;
@@ -414,8 +412,6 @@ protected:
 	float RawPitchMultiplier = 0.0f;
 	int32 Transposition = 0;
 
-	int32 PresetIdx = 0;
-
 	EKeyzoneSelectMode KeyzoneSelectMode = EKeyzoneSelectMode::Layers;
 
 private:
diff --git a/Engine/Plugins/Runtime/Harmonix/Source/HarmonixDsp/Public/HarmonixDsp/FusionSampler/Settings/FusionPatchSettings.h b/Engine/Plugins/Runtime/Harmonix/Source/HarmonixDsp/Public/HarmonixDsp/FusionSampler/Settings/FusionPatchSettings.h
--- a/Engine/Plugins/Runtime/Harmonix/Source/HarmonixDsp/Public/HarmonixDsp/FusionSampler/Settings/FusionPatchSettings.h
+++ b/Engine/Plugins/Runtime/Harmonix/Source/HarmonixDsp/Public/HarmonixDsp/FusionSampler/Settings/FusionPatchSettings.h
@@ -37,9 +37,6 @@ struct HARMONIXDSP_API FFusionPatchSettings
 public:
 
 	FFusionPatchSettings();
-
-	UPROPERTY(EditDefaultsOnly, Category = "Settings")
-	FName Name = FName("Default");
 	
 	UPROPERTY(EditDefaultsOnly, Category = "Settings", Meta = (ClampMin = "-96", ClampMax = "12", UIMin = "-96", UIMax = "12"))
 	float VolumeDb = 0.0f;
diff --git a/Engine/Plugins/Runtime/Harmonix/Source/HarmonixDspEditor/Private/Customization/FusionPatchDetailCustomization.cpp b/Engine/Plugins/Runtime/Harmonix/Source/HarmonixDspEditor/Private/Customization/FusionPatchDetailCustomization.cpp
--- a/Engine/Plugins/Runtime/Harmonix/Source/HarmonixDspEditor/Private/Customization/FusionPatchDetailCustomization.cpp
+++ b/Engine/Plugins/Runtime/Harmonix/Source/HarmonixDspEditor/Private/Customization/FusionPatchDetailCustomization.cpp
@@ -22,10 +22,8 @@ void FFusionPatchDetailCustomization::CustomizeDetails(IDetailLayoutBuilder& Det
 	//get a handle of Fusion Patch Data
 	TSharedPtr<IPropertyHandle> FusionPatchDataHandle = DetailLayout.GetProperty("FusionPatchData");
 	check(FusionPatchDataHandle);
-	TSharedPtr<IPropertyHandle> PresetsHandle = FusionPatchDataHandle->GetChildHandle("Presets");
-	check(PresetsHandle);
-	TSharedPtr<IPropertyHandle> CurrentPresetIndexHandle = FusionPatchDataHandle->GetChildHandle("CurrentPresetIndex");
-	check(CurrentPresetIndexHandle);
+	TSharedPtr<IPropertyHandle> SettingsHandle = FusionPatchDataHandle->GetChildHandle("Settings");
+	check(SettingsHandle);
 
 	//get a handle of the keyzones as array
 	TSharedPtr<IPropertyHandleArray> KeyzonesHandle = FusionPatchDataHandle->GetChildHandle("Keyzones")->AsArray();
@@ -89,9 +87,8 @@ void FFusionPatchDetailCustomization::CustomizeDetails(IDetailLayoutBuilder& Det
 
 	DrawSelectedKeyzoneProperties(FusionPatchDataCategory, KeyzonesHandle, NumKeyzones);
 
-	//add the presets properties since they're not customized 
-	FusionPatchDataCategory.AddProperty(PresetsHandle);
-	FusionPatchDataCategory.AddProperty(CurrentPresetIndexHandle);
+	//add the Settings properties back so they're visible
+	FusionPatchDataCategory.AddProperty(SettingsHandle);
 
 	//hide the non-customized versions of these properties
 	DetailLayout.HideProperty(FusionPatchDataHandle);
diff --git a/Engine/Plugins/Runtime/Harmonix/Source/HarmonixDspEditor/Private/FusionPatchJsonImporter.cpp b/Engine/Plugins/Runtime/Harmonix/Source/HarmonixDspEditor/Private/FusionPatchJsonImporter.cpp
--- a/Engine/Plugins/Runtime/Harmonix/Source/HarmonixDspEditor/Private/FusionPatchJsonImporter.cpp
+++ b/Engine/Plugins/Runtime/Harmonix/Source/HarmonixDspEditor/Private/FusionPatchJsonImporter.cpp
@@ -54,7 +54,7 @@ bool FFusionPatchJsonImporter::TryParseJson(TSharedPtr<FJsonObject> JsonObj, UFu
 	
 	FFusionPatchData& FusionPatchData = FusionPatch->FusionPatchData;
 
-	// import presets and keyzones first
+	// import patch settings and keyzones first
 	for (TSharedPtr<FJsonValue> ArrayValue : IterField(JsonObj, "presets"))
 	{
 		TSharedPtr<FJsonObject> PresetObject = nullptr;
@@ -76,20 +76,19 @@ bool FFusionPatchJsonImporter::TryParseJson(TSharedPtr<FJsonObject> JsonObj, UFu
 				continue;
 			}
 
-			FFusionPatchSettings Preset;
-			if (!FSettingsJsonImporter::TryParseJson(PresetJson, Preset))
+			FFusionPatchSettings PatchSettings;
+			if (!FSettingsJsonImporter::TryParseJson(PresetJson, PatchSettings))
 			{
 				UE_LOG(LogFusionPatchJsonImporter, Warning, TEXT("Unable to parse preset in json"));
 				continue;
 			}
 
-			Preset.Name = FName(Key);
-			FusionPatchData.AddPreset(Preset);
+			FusionPatchData.UpdateSettings(PatchSettings);
+			// we only support one "preset" and that's the settings
+			break;
 		}
 	}
-
-	FusionPatchData.CurrentPresetIndex = 0;
-
+	
 	// collect audio sample files for import
 	TArray<FString> AudioSampleFiles;
 	bool AllFilesExist = true;
diff --git a/Engine/Plugins/Runtime/Harmonix/Source/HarmonixDspEditor/Private/SettingsJsonImporter.cpp b/Engine/Plugins/Runtime/Harmonix/Source/HarmonixDspEditor/Private/SettingsJsonImporter.cpp
--- a/Engine/Plugins/Runtime/Harmonix/Source/HarmonixDspEditor/Private/SettingsJsonImporter.cpp
+++ b/Engine/Plugins/Runtime/Harmonix/Source/HarmonixDspEditor/Private/SettingsJsonImporter.cpp
@@ -119,21 +119,21 @@ bool FSettingsJsonImporter::TryParseJson(TSharedPtr<FJsonObject> JsonObj, FPanne
 	return false;
 }
 
-bool FSettingsJsonImporter::TryParseJson(TSharedPtr<FJsonObject> JsonObj, FFusionPatchSettings& Preset)
+bool FSettingsJsonImporter::TryParseJson(TSharedPtr<FJsonObject> JsonObj, FFusionPatchSettings& PatchSettings)
 {
 	FFusionPatchSettings Defaults;
-	TryGetNumberField(JsonObj, "volume", Preset.VolumeDb, HarmonixDsp::kDbSilence);
+	TryGetNumberField(JsonObj, "volume", PatchSettings.VolumeDb, HarmonixDsp::kDbSilence);
 	
-	TryParseJson(JsonObj, Preset.PannerDetails);
+	TryParseJson(JsonObj, PatchSettings.PannerDetails);
 
-	TryGetNumberField(JsonObj, "start_point", Preset.StartPointOffsetMs, Defaults.StartPointOffsetMs);
-	TryGetNumberField(JsonObj, "fine_tune", Preset.FineTuneCents, Defaults.FineTuneCents);
-	TryGetNumberField(JsonObj, "voice_limit", Preset.MaxVoices, 8);
-	TryGetEnumField(JsonObj, "layer_select_mode", Preset.KeyzoneSelectMode, EKeyzoneSelectMode::Layers);
+	TryGetNumberField(JsonObj, "start_point", PatchSettings.StartPointOffsetMs, Defaults.StartPointOffsetMs);
+	TryGetNumberField(JsonObj, "fine_tune", PatchSettings.FineTuneCents, Defaults.FineTuneCents);
+	TryGetNumberField(JsonObj, "voice_limit", PatchSettings.MaxVoices, 8);
+	TryGetEnumField(JsonObj, "layer_select_mode", PatchSettings.KeyzoneSelectMode, EKeyzoneSelectMode::Layers);
 
 	// apply defaults just in case
-	Preset.DownPitchBendCents = -700.0f;
-	Preset.UpPitchBendCents = 700.0f;
+	PatchSettings.DownPitchBendCents = -700.0f;
+	PatchSettings.UpPitchBendCents = 700.0f;
 	TSharedPtr<FJsonObject> PitchBendObj;
 	if (TryGetObjectField(JsonObj, "pitch_bend", PitchBendObj))
 	{
@@ -142,8 +142,8 @@ bool FSettingsJsonImporter::TryParseJson(TSharedPtr<FJsonObject> JsonObj, FFusio
 		{
 			if (ensure(RangeValues->Num() == 2))
 			{
-				TryGetNumber((*RangeValues)[0], Preset.DownPitchBendCents, Defaults.DownPitchBendCents);
-				TryGetNumber((*RangeValues)[1], Preset.UpPitchBendCents, Defaults.UpPitchBendCents);
+				TryGetNumber((*RangeValues)[0], PatchSettings.DownPitchBendCents, Defaults.DownPitchBendCents);
+				TryGetNumber((*RangeValues)[1], PatchSettings.UpPitchBendCents, Defaults.UpPitchBendCents);
 			}
 		}
 	}
@@ -161,7 +161,7 @@ bool FSettingsJsonImporter::TryParseJson(TSharedPtr<FJsonObject> JsonObj, FFusio
 		if (!TryGetObjectField(AdsrObj, "adsr", AdsrValue))
 			continue;
 
-		FAdsrSettings& AdsrSettings = AdsrIdx == 0 ? Preset.Adsrs.Volume : Preset.Adsrs.Assignable;
+		FAdsrSettings& AdsrSettings = AdsrIdx == 0 ? PatchSettings.Adsrs.Volume : PatchSettings.Adsrs.Assignable;
 		if (!TryParseJson(AdsrValue, AdsrSettings))
 			continue;
 
@@ -181,7 +181,7 @@ bool FSettingsJsonImporter::TryParseJson(TSharedPtr<FJsonObject> JsonObj, FFusio
 	{
 		++LfoIdx;
 
-		if (!ensure(LfoIdx < Preset.Lfos.Num()))
+		if (!ensure(LfoIdx < PatchSettings.Lfos.Num()))
 			break;
 
 		TSharedPtr<FJsonObject> LfoObj;
@@ -192,7 +192,7 @@ bool FSettingsJsonImporter::TryParseJson(TSharedPtr<FJsonObject> JsonObj, FFusio
 		if (!TryGetObjectField(LfoObj, "lfo", LfoValue))
 			continue;
 
-		FLfoSettings& LfoSettings = Preset.Lfos[LfoIdx];
+		FLfoSettings& LfoSettings = PatchSettings.Lfos[LfoIdx];
 
 		if (!TryParseJson(LfoValue, LfoSettings))
 			continue;
@@ -211,11 +211,11 @@ bool FSettingsJsonImporter::TryParseJson(TSharedPtr<FJsonObject> JsonObj, FFusio
 	TSharedPtr<FJsonObject> FilterObj;
 	if (TryGetObjectField(JsonObj, "filter", FilterObj))
 	{
-		TryParseJson(FilterObj, Preset.Filter);
+		TryParseJson(FilterObj, PatchSettings.Filter);
 	}
 	else
 	{
-		Preset.Filter.IsEnabled = false;
+		PatchSettings.Filter.IsEnabled = false;
 	}
 
 	{
@@ -223,7 +223,7 @@ bool FSettingsJsonImporter::TryParseJson(TSharedPtr<FJsonObject> JsonObj, FFusio
 	for (TSharedPtr<FJsonValue> ArrayValue : IterField(JsonObj, "randomizers"))
 	{
 		++RandomizerIdx;
-		if (!ensure(RandomizerIdx < Preset.Randomizers.Num()))
+		if (!ensure(RandomizerIdx < PatchSettings.Randomizers.Num()))
 			break;
 
 		TSharedPtr<FJsonObject> RandomizerObj;
@@ -234,7 +234,7 @@ bool FSettingsJsonImporter::TryParseJson(TSharedPtr<FJsonObject> JsonObj, FFusio
 		if (!TryGetObjectField(RandomizerObj, "modulator", RandomizerValue))
 			continue;
 
-		TryParseJson(RandomizerValue, Preset.Randomizers[RandomizerIdx]);
+		TryParseJson(RandomizerValue, PatchSettings.Randomizers[RandomizerIdx]);
 	}
 	}
 
@@ -243,7 +243,7 @@ bool FSettingsJsonImporter::TryParseJson(TSharedPtr<FJsonObject> JsonObj, FFusio
 	for (TSharedPtr<FJsonValue> ArrayValue : IterField(JsonObj, "velocity_mods"))
 	{
 		++ModIdx;
-		if (!ensure(ModIdx < Preset.VelocityModulators.Num()))
+		if (!ensure(ModIdx < PatchSettings.VelocityModulators.Num()))
 			break;
 
 		TSharedPtr<FJsonObject> VelocityModObj;
@@ -254,7 +254,7 @@ bool FSettingsJsonImporter::TryParseJson(TSharedPtr<FJsonObject> JsonObj, FFusio
 		if (!TryGetObjectField(VelocityModObj, "modulator", VelocityModValue))
 			continue;
 
-		TryParseJson(VelocityModValue, Preset.VelocityModulators[ModIdx]);
+		TryParseJson(VelocityModValue, PatchSettings.VelocityModulators[ModIdx]);
 	}
 	}
 
@@ -262,11 +262,11 @@ bool FSettingsJsonImporter::TryParseJson(TSharedPtr<FJsonObject> JsonObj, FFusio
 	TSharedPtr<FJsonObject> DelayObj;
 	if (TryGetObjectField(JsonObj, "delay_effect", DelayObj))
 	{
-		TryParseJson(DelayObj, Preset.Delay);
+		TryParseJson(DelayObj, PatchSettings.Delay);
 	}
 	else
 	{
-		Preset.Delay.IsEnabled = false;
+		PatchSettings.Delay.IsEnabled = false;
 	}
 	}
 
@@ -274,11 +274,11 @@ bool FSettingsJsonImporter::TryParseJson(TSharedPtr<FJsonObject> JsonObj, FFusio
 	TSharedPtr<FJsonObject> BitcrushObj;
 	if (TryGetObjectField(JsonObj, "bitcrush_effect", BitcrushObj))
 	{
-		TryParseJson(BitcrushObj, Preset.BitCrusher);
+		TryParseJson(BitcrushObj, PatchSettings.BitCrusher);
 	}
 	else
 	{
-		Preset.BitCrusher.IsEnabled = false;
+		PatchSettings.BitCrusher.IsEnabled = false;
 	}
 	}
 
@@ -286,11 +286,11 @@ bool FSettingsJsonImporter::TryParseJson(TSharedPtr<FJsonObject> JsonObj, FFusio
 	TSharedPtr<FJsonObject> VocoderObj;
 	if (TryGetObjectField(JsonObj, "vocoder_effect", VocoderObj))
 	{
-		TryParseJson(VocoderObj, Preset.Vocoder);
+		TryParseJson(VocoderObj, PatchSettings.Vocoder);
 	}
 	else
 	{
-		Preset.Vocoder.IsEnabled = false;
+		PatchSettings.Vocoder.IsEnabled = false;
 	}
 	}
 
@@ -298,11 +298,11 @@ bool FSettingsJsonImporter::TryParseJson(TSharedPtr<FJsonObject> JsonObj, FFusio
 	TSharedPtr<FJsonObject> DistortionObj;
 	if (TryGetObjectField(JsonObj, "distortion_effect", DistortionObj))
 	{
-		TryParseJson(DistortionObj, Preset.Distortion);
+		TryParseJson(DistortionObj, PatchSettings.Distortion);
 	}
 	else
 	{
-		Preset.Distortion.IsEnabled = false;
+		PatchSettings.Distortion.IsEnabled = false;
 	}
 	}
 
@@ -310,13 +310,13 @@ bool FSettingsJsonImporter::TryParseJson(TSharedPtr<FJsonObject> JsonObj, FFusio
 	TSharedPtr<FJsonObject> PortamentoObj;
 	if (TryGetObjectField(JsonObj, "portamento", PortamentoObj))
 	{
-		TryParseJson(PortamentoObj, Preset.Portamento);
+		TryParseJson(PortamentoObj, PatchSettings.Portamento);
 	}
 	else
 	{
-		Preset.Portamento.IsEnabled = false;
-		Preset.Portamento.Mode = EPortamentoMode::Legato;
-		Preset.Portamento.Seconds = 0.0f;
+		PatchSettings.Portamento.IsEnabled = false;
+		PatchSettings.Portamento.Mode = EPortamentoMode::Legato;
+		PatchSettings.Portamento.Seconds = 0.0f;
 	}
 	}
 
diff --git a/Engine/Plugins/Runtime/Harmonix/Source/HarmonixDspEditor/Private/SettingsJsonImporter.h b/Engine/Plugins/Runtime/Harmonix/Source/HarmonixDspEditor/Private/SettingsJsonImporter.h
--- a/Engine/Plugins/Runtime/Harmonix/Source/HarmonixDspEditor/Private/SettingsJsonImporter.h
+++ b/Engine/Plugins/Runtime/Harmonix/Source/HarmonixDspEditor/Private/SettingsJsonImporter.h
@@ -16,7 +16,7 @@ public:
 
 	static bool TryParseJson(TSharedPtr<FJsonObject> JsonObj, struct FTimeStretchConfig& TimeStretchConfig);
 
-	static bool TryParseJson(TSharedPtr<FJsonObject> JsonObj, struct FFusionPatchSettings& Preset);
+	static bool TryParseJson(TSharedPtr<FJsonObject> JsonObj, struct FFusionPatchSettings& PatchSettings);
 
 	static bool TryParseJson(TSharedPtr<FJsonObject> JsonObj, struct FAdsrSettings& AdsrSettings);
 
diff --git a/Engine/Plugins/Runtime/Harmonix/Source/HarmonixMetasound/Private/HarmonixMetasound/Nodes/FusionSamplerNode.cpp b/Engine/Plugins/Runtime/Harmonix/Source/HarmonixMetasound/Private/HarmonixMetasound/Nodes/FusionSamplerNode.cpp
--- a/Engine/Plugins/Runtime/Harmonix/Source/HarmonixMetasound/Private/HarmonixMetasound/Nodes/FusionSamplerNode.cpp
+++ b/Engine/Plugins/Runtime/Harmonix/Source/HarmonixMetasound/Private/HarmonixMetasound/Nodes/FusionSamplerNode.cpp
@@ -229,7 +229,7 @@ namespace HarmonixMetasound
 			}
 			else if (FusionPatchDataPtr)
 			{
-				SetController(MidiConstants::LFO0Frequency, FusionPatchDataPtr->GetPreset(0).Lfos[0].Freq);
+				SetController(MidiConstants::LFO0Frequency, FusionPatchDataPtr->GetSettings().Lfos[0].Freq);
 			}
 			if (*Lfo0DepthInPin >= 0.0f)
 			{
@@ -237,7 +237,7 @@ namespace HarmonixMetasound
 			}
 			else if (FusionPatchDataPtr)
 			{
-				SetController(MidiConstants::LFO0Depth, FusionPatchDataPtr->GetPreset(0).Lfos[0].Depth);
+				SetController(MidiConstants::LFO0Depth, FusionPatchDataPtr->GetSettings().Lfos[0].Depth);
 			}
 			if (*Lfo1RateInPin > 0.0f)
 			{
@@ -245,7 +245,7 @@ namespace HarmonixMetasound
 			}
 			else if (FusionPatchDataPtr)
 			{
-				SetController(MidiConstants::LFO1Frequency, FusionPatchDataPtr->GetPreset(0).Lfos[1].Freq);
+				SetController(MidiConstants::LFO1Frequency, FusionPatchDataPtr->GetSettings().Lfos[1].Freq);
 			}
 			if (*Lfo1DepthInPin >= 0.0f)
 			{
@@ -253,7 +253,7 @@ namespace HarmonixMetasound
 			}
 			else if (FusionPatchDataPtr)
 			{
-				SetController(MidiConstants::LFO1Depth, FusionPatchDataPtr->GetPreset(0).Lfos[1].Depth);
+				SetController(MidiConstants::LFO1Depth, FusionPatchDataPtr->GetSettings().Lfos[1].Depth);
 			}
 
 			float FineTuneCentsTotal = 0;
@@ -263,7 +263,7 @@ namespace HarmonixMetasound
 			}
 			if (FusionPatchDataPtr)
 			{
-				FineTuneCentsTotal += FusionPatchDataPtr->GetPreset(0).FineTuneCents;
+				FineTuneCentsTotal += FusionPatchDataPtr->GetSettings().FineTuneCents;
 			}
 			SetFineTuneCents(FineTuneCentsTotal);
 		}
