[Harmonix] Move TempoMapMetronome to internal

diff --git a/Engine/Plugins/Runtime/Harmonix/Source/HarmonixMetasound/Private/HarmonixMetasound/Nodes/TempoMapMetronomeNode.cpp b/Engine/Plugins/Runtime/Harmonix/Source/HarmonixMetasound/Private/HarmonixMetasound/Nodes/TempoMapMetronomeNode.cpp
deleted file mode 100644
diff --git a/Engine/Plugins/Runtime/Harmonix/Source/HarmonixMetasound/Public/HarmonixMetasound/DataTypes/MusicTransport.h b/Engine/Plugins/Runtime/Harmonix/Source/HarmonixMetasound/Public/HarmonixMetasound/DataTypes/MusicTransport.h
--- a/Engine/Plugins/Runtime/Harmonix/Source/HarmonixMetasound/Public/HarmonixMetasound/DataTypes/MusicTransport.h
+++ b/Engine/Plugins/Runtime/Harmonix/Source/HarmonixMetasound/Public/HarmonixMetasound/DataTypes/MusicTransport.h
@@ -82,7 +82,7 @@ namespace HarmonixMetasound
 		NumStates
 	};
 
-	class FMusicTransportControllable
+	class HARMONIXMETASOUND_API FMusicTransportControllable
 	{
 	public:
 		FMusicTransportControllable(EMusicPlayerTransportState InitialState);
diff --git a/Engine/Plugins/Runtime/Harmonix/Source/HarmonixMetasoundTests/Private/Nodes/TempoMapMetronomeNodeTests.cpp b/Engine/Plugins/Runtime/Harmonix/Source/HarmonixMetasoundTests/Private/Nodes/TempoMapMetronomeNodeTests.cpp
deleted file mode 100644
