Rename Metasound Param macros in the HarmonixMetasound module so they're parsed correctly by the text gatherer.

Also, fix up some naming conflicts by adding some Aliases, and making tooltp/description text match identically to any redefinitions.

New names for the various harmonix macros

DECLARE_METASOUND_PARAM_EXTERN
DEFINE_INPUT_METASOUND_PARAM
DEFINE_OUTPUT_METASOUND_PARAM
DEFINE_METASOUND_PARAM_ALIAS

This naming convention takes advantage of the fact that the text gatherer ignores prefixes when matching macro names, but won't match if there's a suffix. It also more closely matches the UE standard for similar "Declare" and "Define" macros

I created a separate INPUT and OUTPUT macro so that I could make the LOCTEXT for those definitions prefix with "Input" and "Output" respectively to avoid LOCEXT key/name conflicts without refactoring/renaming our inputs.

This did not, however, fix the text gatherer generating METASOUND_PARAM conflicts, which I had to fix by cleaning up the remaining names.

diff --git a/Engine/Plugins/Runtime/Harmonix/Source/HarmonixMetasound/Private/HarmonixMetasound/Common.cpp b/Engine/Plugins/Runtime/Harmonix/Source/HarmonixMetasound/Private/HarmonixMetasound/Common.cpp
--- a/Engine/Plugins/Runtime/Harmonix/Source/HarmonixMetasound/Private/HarmonixMetasound/Common.cpp
+++ b/Engine/Plugins/Runtime/Harmonix/Source/HarmonixMetasound/Private/HarmonixMetasound/Common.cpp
@@ -20,87 +20,89 @@ namespace HarmonixMetasound
 	{
 		namespace Inputs
 		{
-			EXTERN_METASOUND_PARAM_D(Enable,                        "Enable",                             "Enables processing.")
-			EXTERN_METASOUND_PARAM_D(MidiStream,                    "Midi Stream",                        "Midi event input stream.")
-			EXTERN_METASOUND_PARAM_D(MidiChannelNumber,             "Midi Channel",                       "Midi channel to render (0 for all/omni).")
-			EXTERN_METASOUND_PARAM_D(MidiChannelFilterSpecifier,    "Midi Channel Filter",                "Midi channels to process. '*' for all, or a string like '1,3,4-8' to identify specific channels.")
-			EXTERN_METASOUND_PARAM_D(MidiTrackNumber,               "Track Number",                       "Track number (1 based).")
-			EXTERN_METASOUND_PARAM_D(MidiTrackIndexFilterSpecifier, "Midi Track Filter",                  "Midi tracks to process. '*' for all, or a string like '1,3,4-8' to identify specific tracks.")
-			EXTERN_METASOUND_PARAM_D(MinMidiNote,                   "Min Note #",                         "Minimum midi note number (0-127).")
-			EXTERN_METASOUND_PARAM_D(MaxMidiNote,                   "Max Note #",                         "Maximum midi note number (0-127).")
-			EXTERN_METASOUND_PARAM_D(MinMidiVelocity,               "Min Velocity",                       "Minimum Velocity (0-127).")
-			EXTERN_METASOUND_PARAM_D(MaxMidiVelocity,               "Max Velocity",                       "Maximum Velocity (0-127).")
-			EXTERN_METASOUND_PARAM_D(MidiFileAsset,                 "Midi File",                          "Standard Midi file.")
-			EXTERN_METASOUND_PARAM_D(Transport,                     "Transport",                          "Play, pause, continue, stop, etc.")
-			EXTERN_METASOUND_PARAM_D(MidiClock,                     "Midi Clock",                         "Midi timing information.")
-			EXTERN_METASOUND_PARAM_D(Tempo,                         "Tempo",                              "Tempo in BPM. This is MIDI bpm... ALWAYS QUARTER NOTES PER MINUTE regardless of time signature!")
-			EXTERN_METASOUND_PARAM_D(TimeSigNumerator,	            "Time Sig. Numerator",                "Time Signature Numerator")
-			EXTERN_METASOUND_PARAM_D(TimeSigDenominator,            "Time Sig. Denominator",              "TimeSignature Denominator")
-			EXTERN_METASOUND_PARAM_D(Speed,                         "Speed",                              "Playback speed multiplier.")
-			EXTERN_METASOUND_PARAM_D(Loop,				            "Simple Loop",			              "Whether looping is enabled for this node. Loop length will be rounded to the nearest bar.")
-			EXTERN_METASOUND_PARAM_D(LoopLengthBars,                "Loop Length (Bars)",                 "The number of bars to loop for simple looping")
-			EXTERN_METASOUND_PARAM_D(TransportPrepare,              "Prepare",                            "Causes receiving nodes to prepare to play (pre-buffer, cache, etc.)")
-			EXTERN_METASOUND_PARAM_D(TransportPlay,                 "Play",                               "Causes a transport play request.")
-			EXTERN_METASOUND_PARAM_D(TransportPause,                "Pause",                              "Causes a transport pause request.")
-			EXTERN_METASOUND_PARAM_D(TransportContinue,             "Continue",                           "Causes a transport contiunue request.")
-			EXTERN_METASOUND_PARAM_D(TransportStop,                 "Stop",                               "Causes a transport stop request.")
-			EXTERN_METASOUND_PARAM_D(TransportKill,                 "Kill",                               "Causes a transport kill request.")
-			EXTERN_METASOUND_PARAM_D(TriggerSeek,                   "Trigger Seek",                       "Causes a transport seek request. Seek destination provided on a separate input.")
-			EXTERN_METASOUND_PARAM_D(SeekDestination,               "Seek Target",                        "Destination for the next seek when it is triggered.")
-			EXTERN_METASOUND_PARAM_D(PrerollBars,                   "Pre-roll Bars",                      "Number of bars to look back when seeking midi to find note-on messages that should hold over.")
-			EXTERN_METASOUND_PARAM_D(GridSizeUnits,                 "Base Grid Size",                     "Base size of the qrid squares in musical terms.")
-			EXTERN_METASOUND_PARAM_D(GridSizeMult,                  "Grid Size Multiplier",               "Grid size multiplier.")
-			EXTERN_METASOUND_PARAM_D(OffsetUnits,                   "Base Offset Size ",                  "Base size of the qrid offset.")
-			EXTERN_METASOUND_PARAM_D(OffsetMult,                    "Offset Multiplier",                  "Offset multiplier.")
-			EXTERN_METASOUND_PARAM_D(Bar,		   	                "Bar",                                "Musical bar (or measure) index (0 based!)")
-			EXTERN_METASOUND_PARAM_D(FloatBeat,			            "Beat",                               "Musical beat index. 0 based and in units of the time signature denomninator! The value is rounded to the nearest 32nd note.")
-			EXTERN_METASOUND_PARAM_D(SynthPatch,                    "Patch",                              "Synthesizer patch asset.")
-			EXTERN_METASOUND_PARAM_D(Transposition,                 "Transposition",                      "Transposition in semitones.")
-			EXTERN_METASOUND_PARAM_D(ClockSpeedToPitch,             "Midi Clock Speed Affects Pitch",     "When true the clock speed of the incoming midi stream will affect the pitch of the synth.")
-			EXTERN_METASOUND_PARAM_D(ClockSpeedToFrequency,         "Midi Clock Speed Affects Frequency", "When true the clock speed of the incoming midi stream will affect the calculated frequency.")
-			EXTERN_METASOUND_PARAM_D(LFOSyncType,					"LFO Sync Type",					  "Specifies if and how the LFO should sync to a clock");
-			EXTERN_METASOUND_PARAM_D(LFOFrequency,					"LFO Frequency",					  "The period of the LFO. If Sync Type is TempoSync, the unit is cycles per quarter note. Otherwise, the unit is Hz (cycles per second).");
-			EXTERN_METASOUND_PARAM_D(LFOInvert,						"Invert LFO",						  "Toggle to invert the LFO");
-			EXTERN_METASOUND_PARAM_D(LFOShape,						"LFO Shape",						  "The shape of the LFO. 0.0-1.0 = morph between square and triangle. 1.0-2.0 = morph between triangle and sawtooth");
-			EXTERN_METASOUND_PARAM_D(AudioMono,						"Audio Mono",						  "The mono audio input");
+			DEFINE_INPUT_METASOUND_PARAM(Enable,                        "Enable",                             "Enables processing.")
+			DEFINE_INPUT_METASOUND_PARAM(MidiStream,                    "Midi Stream",						  "Midi event stream.")
+			DEFINE_INPUT_METASOUND_PARAM(MidiChannelNumber,             "Midi Channel",                       "Midi channel to render (0 for all/omni).")
+			DEFINE_INPUT_METASOUND_PARAM(MidiChannelFilterSpecifier,    "Midi Channel Filter",                "Midi channels to process. '*' for all, or a string like '1,3,4-8' to identify specific channels.")
+			DEFINE_INPUT_METASOUND_PARAM(MidiTrackNumber,               "Track Number",                       "Track number (1 based).")
+			DEFINE_INPUT_METASOUND_PARAM(MidiTrackIndexFilterSpecifier, "Midi Track Filter",                  "Midi tracks to process. '*' for all, or a string like '1,3,4-8' to identify specific tracks.")
+			DEFINE_INPUT_METASOUND_PARAM(MinMidiNote,                   "Min Note #",                         "Minimum midi note number (0-127).")
+			DEFINE_INPUT_METASOUND_PARAM(MaxMidiNote,                   "Max Note #",                         "Maximum midi note number (0-127).")
+			DEFINE_INPUT_METASOUND_PARAM(MinMidiVelocity,               "Min Velocity",                       "Minimum Velocity (0-127).")
+			DEFINE_INPUT_METASOUND_PARAM(MaxMidiVelocity,               "Max Velocity",                       "Maximum Velocity (0-127).")
+			DEFINE_INPUT_METASOUND_PARAM(MidiFileAsset,                 "Midi File",                          "Standard Midi file.")
+			DEFINE_INPUT_METASOUND_PARAM(Transport,                     "Transport",                          "Play, pause, continue, stop, etc.")
+			DEFINE_INPUT_METASOUND_PARAM(MidiClock,                     "Midi Clock",                         "Midi timing information.")
+			DEFINE_INPUT_METASOUND_PARAM(Tempo,                         "Tempo",                              "BPM. This is MIDI bpm... ALWAYS QUARTER NOTES PER MINUTE regardless of time signature!")
+			DEFINE_INPUT_METASOUND_PARAM(TimeSigNumerator,	            "Time Sig. Numerator",                "Time Signature Numerator")
+			DEFINE_INPUT_METASOUND_PARAM(TimeSigDenominator,            "Time Sig. Denominator",              "TimeSignature Denominator")
+			DEFINE_INPUT_METASOUND_PARAM(Speed,                         "Speed",                              "Playback speed multiplier.")
+			DEFINE_INPUT_METASOUND_PARAM(Loop,				            "Simple Loop",						  "Whether looping is enabled for this node. Loop length will be rounded to the nearest bar.")
+			DEFINE_INPUT_METASOUND_PARAM(LoopLengthBars,                "Loop Length (Bars)",                 "The number of bars to loop for simple looping")
+			DEFINE_INPUT_METASOUND_PARAM(TransportPrepare,              "Prepare",                            "Causes receiving nodes to prepare to play (pre-buffer, cache, etc.)")
+			DEFINE_INPUT_METASOUND_PARAM(TransportPlay,                 "Play",                               "Causes a transport play request.")
+			DEFINE_INPUT_METASOUND_PARAM(TransportPause,                "Pause",                              "Causes a transport pause request.")
+			DEFINE_INPUT_METASOUND_PARAM(TransportContinue,             "Continue",                           "Causes a transport continue request.")
+			DEFINE_INPUT_METASOUND_PARAM(TransportStop,                 "Stop",                               "Causes a transport stop request.")
+			DEFINE_INPUT_METASOUND_PARAM(TransportKill,                 "Kill",                               "Causes a transport kill request.")
+			DEFINE_INPUT_METASOUND_PARAM(TriggerSeek,                   "Trigger Seek",                       "Causes a transport seek request. Seek destination provided on a separate input.")
+			DEFINE_INPUT_METASOUND_PARAM(SeekDestination,               "Seek Target",                        "Destination for the next seek when it is triggered.")
+			DEFINE_INPUT_METASOUND_PARAM(PrerollBars,                   "Pre-roll Bars",                      "Number of bars to look back when seeking midi to find note-on messages that should hold over.")
+			DEFINE_INPUT_METASOUND_PARAM(GridSizeUnits,                 "Base Grid Size",                     "Base size of the qrid squares in musical terms.")
+			DEFINE_INPUT_METASOUND_PARAM(GridSizeMult,                  "Grid Size Multiplier",               "Grid size multiplier.")
+			DEFINE_INPUT_METASOUND_PARAM(OffsetUnits,                   "Base Offset Size ",                  "Base size of the qrid offset.")
+			DEFINE_INPUT_METASOUND_PARAM(OffsetMult,                    "Offset Multiplier",                  "Offset multiplier.")
+			DEFINE_INPUT_METASOUND_PARAM(Bar,		   					"Bar",                                "Musical bar (or measure) index (0 based!)")
+			DEFINE_INPUT_METASOUND_PARAM(FloatBeat,			            "Beat",                               "Musical beat index. 0 based and in units of the time signature denomninator! The value is rounded to the nearest 32nd note.")
+			DEFINE_INPUT_METASOUND_PARAM(SynthPatch,                    "Patch",                              "Synthesizer patch asset.")
+			DEFINE_INPUT_METASOUND_PARAM(Transposition,                 "Transposition",                      "Transposition in semitones.")
+			DEFINE_INPUT_METASOUND_PARAM(ClockSpeedToPitch,             "Midi Clock Speed Affects Pitch",     "When true the clock speed of the incoming midi stream will affect the pitch of the synth.")
+			DEFINE_INPUT_METASOUND_PARAM(ClockSpeedToFrequency,         "Midi Clock Speed Affects Frequency", "When true the clock speed of the incoming midi stream will affect the calculated frequency.")
+			DEFINE_INPUT_METASOUND_PARAM(LFOSyncType,					"LFO Sync Type",					  "Specifies if and how the LFO should sync to a clock");
+			DEFINE_INPUT_METASOUND_PARAM(LFOFrequency,				    "LFO Frequency",					  "The period of the LFO. If Sync Type is TempoSync, the unit is cycles per quarter note. Otherwise, the unit is Hz (cycles per second).");
+			DEFINE_INPUT_METASOUND_PARAM(LFOInvert,					    "Invert LFO",						  "Toggle to invert the LFO");
+			DEFINE_INPUT_METASOUND_PARAM(LFOShape,					    "LFO Shape",						  "The shape of the LFO. 0.0-1.0 = morph between square and triangle. 1.0-2.0 = morph between triangle and sawtooth");
+			DEFINE_INPUT_METASOUND_PARAM(AudioMono,					    "Audio Mono",					      "The mono audio.");
+			DEFINE_INPUT_METASOUND_PARAM(AudioLeft,					    "Audio Left",					      "The left channel audio.");
+			DEFINE_INPUT_METASOUND_PARAM(AudioRight,					"Audio Right",					      "The right channel audio.");
 		}
 		namespace Outputs
 		{
-			EXTERN_METASOUND_PARAM_D(Transport,          "Transport",                   "Play, pause, continue, stop, etc.")
-			EXTERN_METASOUND_PARAM_D(MidiClock,          "Midi Clock",                  "Midi timing information.")
-			EXTERN_METASOUND_PARAM_D(MidiFileAsset,      "Midi File",                   "Standard Midi file.")
-			EXTERN_METASOUND_PARAM_D(MidiStream,         "Midi Stream",                 "Midi event output stream.")
-			EXTERN_METASOUND_PARAM_D(NoteOn,             "Note On",                     "Midi note on.")
-			EXTERN_METASOUND_PARAM_D(NoteOff,            "Note Off",                    "Midi note off.")
-			EXTERN_METASOUND_PARAM_D(MidiNoteNumber,     "Midi Note #",                 "Midi note number (0 - 127).")
-			EXTERN_METASOUND_PARAM_D(Frequency,          "Frequency",                   "Midi note number converted to frequency.")
-			EXTERN_METASOUND_PARAM_D(MidiVelocity,       "Velocity",                    "Midi velocity (0 - 127).")
-			EXTERN_METASOUND_PARAM_D(NormalizedVelocity, "Normalized Velocity",         "Velocity (0.0 - 1.0).")
-			EXTERN_METASOUND_PARAM_D(TransportPrepare,   "Prepare",                     "Causes receiving nodes to prepare to play (pre-buffer, cache, etc.)")
-			EXTERN_METASOUND_PARAM_D(TransportPlay,      "Play",                        "Causes a transport play request.")
-			EXTERN_METASOUND_PARAM_D(TransportPause,     "Pause",                       "Causes a transport pause request.")
-			EXTERN_METASOUND_PARAM_D(TransportContinue,  "Continue",                    "Causes a transport continue request.")
-			EXTERN_METASOUND_PARAM_D(TransportStop,      "Stop",                        "Causes a transport stop request.")
-			EXTERN_METASOUND_PARAM_D(TransportKill,      "Kill",                        "Causes a transport kill request.")
-			EXTERN_METASOUND_PARAM_D(MusicTimestamp,     "Music Timestamp",             "A single structure representing a musical time (bar & beat).")
-			EXTERN_METASOUND_PARAM_D(SeekTarget,         "Seek Target",                 "A well formed musical seek target.")
-			EXTERN_METASOUND_PARAM_D(AudioMono,          "Audio Mono",                  "The mono audio output.")
-			EXTERN_METASOUND_PARAM_D(AudioLeft,          "Audio Left",                  "The left channel audio output.")
-			EXTERN_METASOUND_PARAM_D(AudioRight,         "Audio Right",                 "The right channel audio output.")
-			EXTERN_METASOUND_PARAM_D(AudioCenter,        "Audio Center",                "5.1/7.1: The center channel audio output.")
-			EXTERN_METASOUND_PARAM_D(AudioLFE,           "Audio LFE",                   "5.1/7.1: The LFE channel audio output.")
-			EXTERN_METASOUND_PARAM_D(AudioLeftSide,      "Audio Left Side",             "Quad/5.1: Left Surround, 7.1: Left Side.")
-			EXTERN_METASOUND_PARAM_D(AudioRightSide,     "Audio Right Side",            "Quad/5.1: Right Surround, 7.1: Right Side.")
-			EXTERN_METASOUND_PARAM_D(AudioLeftBack,      "Audio Left Back - Surround",  "7.1: Left Back.")
-			EXTERN_METASOUND_PARAM_D(AudioRightBack,     "Audio Right Back - Surround", "7.1: Right Back.")
-			EXTERN_METASOUND_PARAM_D(MusicTimespanBar,   "Bar",                         "Bar (Measure). 1 based!")
-			EXTERN_METASOUND_PARAM_D(MusicTimespanBeat,  "Beat",                        "Beat. 1 based!")
-			EXTERN_METASOUND_PARAM_D(Tempo,              "Tempo",                       "BPM. This is MIDI bpm... ALWAYS QUARTER NOTES PER MINUTE regardless of time signature!")
-			EXTERN_METASOUND_PARAM_D(Speed,              "Speed",                       "Speed Multiplier")
-			EXTERN_METASOUND_PARAM_D(TimeSigNumerator,   "Time Sig. Numerator",         "Time Signature Numerator")
-			EXTERN_METASOUND_PARAM_D(TimeSigDenominator, "Time Sig. Denominator",       "TimeSignature Denominator")
-			EXTERN_METASOUND_PARAM_D(SecsPerQuarter,     "Secs Per Quarter Note",       "Seconds per quarter note.")
-			EXTERN_METASOUND_PARAM_D(SecsPerBeat,        "Secs Per Beat",               "Seconds per beat. Equals seconds per beat if time signature denominator is 4.")
+			DEFINE_OUTPUT_METASOUND_PARAM(Transport,          "Transport",                   "Play, pause, continue, stop, etc.")
+			DEFINE_OUTPUT_METASOUND_PARAM(MidiClock,          "Midi Clock",                  "Midi timing information.")
+			DEFINE_OUTPUT_METASOUND_PARAM(MidiFileAsset,      "Midi File",                   "Standard Midi file.")
+			DEFINE_OUTPUT_METASOUND_PARAM(MidiStream,         "Midi Stream",				 "Midi event stream.")
+			DEFINE_OUTPUT_METASOUND_PARAM(NoteOn,             "Note On",                     "Midi note on.")
+			DEFINE_OUTPUT_METASOUND_PARAM(NoteOff,            "Note Off",                    "Midi note off.")
+			DEFINE_OUTPUT_METASOUND_PARAM(MidiNoteNumber,     "Midi Note #",                 "Midi note number (0 - 127).")
+			DEFINE_OUTPUT_METASOUND_PARAM(Frequency,          "Frequency",                   "Midi note number converted to frequency.")
+			DEFINE_OUTPUT_METASOUND_PARAM(MidiVelocity,       "Velocity",                    "Midi velocity (0 - 127).")
+			DEFINE_OUTPUT_METASOUND_PARAM(NormalizedVelocity, "Normalized Velocity",         "Velocity (0.0 - 1.0).")
+			DEFINE_OUTPUT_METASOUND_PARAM(TransportPrepare,   "Prepare",                     "Causes receiving nodes to prepare to play (pre-buffer, cache, etc.)")
+			DEFINE_OUTPUT_METASOUND_PARAM(TransportPlay,      "Play",                        "Causes a transport play request.")
+			DEFINE_OUTPUT_METASOUND_PARAM(TransportPause,     "Pause",                       "Causes a transport pause request.")
+			DEFINE_OUTPUT_METASOUND_PARAM(TransportContinue,  "Continue",                    "Causes a transport continue request.")
+			DEFINE_OUTPUT_METASOUND_PARAM(TransportStop,      "Stop",                        "Causes a transport stop request.")
+			DEFINE_OUTPUT_METASOUND_PARAM(TransportKill,      "Kill",                        "Causes a transport kill request.")
+			DEFINE_OUTPUT_METASOUND_PARAM(MusicTimestamp,     "Music Timestamp",             "A single structure representing a musical time (bar & beat).")
+			DEFINE_OUTPUT_METASOUND_PARAM(SeekTarget,         "Seek Target",                 "A well formed musical seek target.")
+			DEFINE_OUTPUT_METASOUND_PARAM(AudioMono,          "Audio Mono",					 "The mono audio.")
+			DEFINE_OUTPUT_METASOUND_PARAM(AudioLeft,          "Audio Left",				     "The left channel audio.")
+			DEFINE_OUTPUT_METASOUND_PARAM(AudioRight,         "Audio Right",                 "The right channel audio.")
+			DEFINE_OUTPUT_METASOUND_PARAM(AudioCenter,        "Audio Center",                "5.1/7.1: The center channel audio output.")
+			DEFINE_OUTPUT_METASOUND_PARAM(AudioLFE,           "Audio LFE",                   "5.1/7.1: The LFE channel audio output.")
+			DEFINE_OUTPUT_METASOUND_PARAM(AudioLeftSide,      "Audio Left Side",             "Quad/5.1: Left Surround, 7.1: Left Side.")
+			DEFINE_OUTPUT_METASOUND_PARAM(AudioRightSide,     "Audio Right Side",            "Quad/5.1: Right Surround, 7.1: Right Side.")
+			DEFINE_OUTPUT_METASOUND_PARAM(AudioLeftBack,      "Audio Left Back - Surround",  "7.1: Left Back.")
+			DEFINE_OUTPUT_METASOUND_PARAM(AudioRightBack,     "Audio Right Back - Surround", "7.1: Right Back.")
+			DEFINE_OUTPUT_METASOUND_PARAM(MusicTimespanBar,   "Bar",                         "Bar (Measure). 1 based!")
+			DEFINE_OUTPUT_METASOUND_PARAM(MusicTimespanBeat,  "Beat",                        "Beat. 1 based!")
+			DEFINE_OUTPUT_METASOUND_PARAM(Tempo,              "Tempo",                       "BPM. This is MIDI bpm... ALWAYS QUARTER NOTES PER MINUTE regardless of time signature!")
+			DEFINE_OUTPUT_METASOUND_PARAM(Speed,              "Speed",                       "Playback speed multiplier.")
+			DEFINE_OUTPUT_METASOUND_PARAM(TimeSigNumerator,   "Time Sig. Numerator",         "Time Signature Numerator")
+			DEFINE_OUTPUT_METASOUND_PARAM(TimeSigDenominator, "Time Sig. Denominator",       "TimeSignature Denominator")
+			DEFINE_OUTPUT_METASOUND_PARAM(SecsPerQuarter,     "Secs Per Quarter Note",       "Seconds per quarter note.")
+			DEFINE_OUTPUT_METASOUND_PARAM(SecsPerBeat,        "Secs Per Beat",               "Seconds per beat. Equals seconds per beat if time signature denominator is 4.")
 		}
 	}
 }
diff --git a/Engine/Plugins/Runtime/Harmonix/Source/HarmonixMetasound/Private/HarmonixMetasound/Nodes/DelayNode.cpp b/Engine/Plugins/Runtime/Harmonix/Source/HarmonixMetasound/Private/HarmonixMetasound/Nodes/DelayNode.cpp
--- a/Engine/Plugins/Runtime/Harmonix/Source/HarmonixMetasound/Private/HarmonixMetasound/Nodes/DelayNode.cpp
+++ b/Engine/Plugins/Runtime/Harmonix/Source/HarmonixMetasound/Private/HarmonixMetasound/Nodes/DelayNode.cpp
@@ -24,32 +24,32 @@ namespace HarmonixMetasound::DelayNode
 	
 	namespace Inputs
 	{
-		EXTERN_METASOUND_PARAM_D(AudioLeft, "Audio Left", "The left channel of audio to process");
-		EXTERN_METASOUND_PARAM_D(AudioRight, "Audio Right", "The right channel of audio to process");
-		EXTERN_METASOUND_PARAM_D(MidiClock, "MIDI Clock", "MIDI timing information")
-		EXTERN_METASOUND_PARAM_D(DelayTimeType, "Delay Sync Type", "Specifies if and how the delay should sync to a clock");
-		EXTERN_METASOUND_PARAM_D(DelayTime, "Delay Time", "The period of the delay. If Delay Sync Type is TempoSync, the unit is quarter notes. Otherwise, the unit is seconds.");
-		EXTERN_METASOUND_PARAM_D(Feedback, "Feedback", "The feedback amount, range [0.0, 1.0]");
-		EXTERN_METASOUND_PARAM_D(DryLevel, "Dry Gain", "The amount of dry signal to pass through, range [0.0, 1.0]");
-		EXTERN_METASOUND_PARAM_D(WetLevel, "Wet Gain", "The amount of wet signal to output, range [0.0, 1.0]");
-		EXTERN_METASOUND_PARAM_D(WetFilterEnabled, "Wet Filter Enabled", "Enables or disables the filter on the wet signal");
-		EXTERN_METASOUND_PARAM_D(FeedbackFilterEnabled, "Feedback Filter Enabled", "Enables or disables the filter on the feedback signal");
-		EXTERN_METASOUND_PARAM_D(FilterType, "Filter Type", "The type of filter to apply to the signal");
-		EXTERN_METASOUND_PARAM_D(FilterCutoff, "Filter Cutoff", "The cutoff frequency of the filter, in Hz");
-		EXTERN_METASOUND_PARAM_D(FilterQ, "Filter Q", "The Q of the filter");
-		EXTERN_METASOUND_PARAM_D(LFOEnabled, "LFO Enabled", "Enables or disables the LFO on the delay time");
-		EXTERN_METASOUND_PARAM_D(LFOTimeType, "LFO Sync Type", "Specifies if and how the LFO should sync to a clock");
-		EXTERN_METASOUND_PARAM_D(LFOFrequency, "LFO Frequency", "The period of the LFO. If LFO Sync Type is TempoSync, the unit is cycles per quarter note. Otherwise, the unit is Hz (cycles per second).");
-		EXTERN_METASOUND_PARAM_D(LFODepth, "LFO Depth", "The amount of delay to be added by the LFO. If Delay Sync Type is TempoSync, the unit is quarter notes. Otherwise, the unit is seconds.");
-		EXTERN_METASOUND_PARAM_D(StereoType, "Stereo Type", "Specifies method for panning and spreading the sound");
-		EXTERN_METASOUND_PARAM_D(StereoSpreadLeft, "Stereo Spread Left", "The amount of stereo spread on the left channel");
-		EXTERN_METASOUND_PARAM_D(StereoSpreadRight, "Stereo Spread Right", "The amount of stereo spread on the right channel");
+		DEFINE_METASOUND_PARAM_ALIAS(AudioLeft, CommonPinNames::Inputs::AudioLeft);
+		DEFINE_METASOUND_PARAM_ALIAS(AudioRight, CommonPinNames::Inputs::AudioRight);
+		DEFINE_METASOUND_PARAM_ALIAS(MidiClock, CommonPinNames::Inputs::MidiClock);
+		DEFINE_INPUT_METASOUND_PARAM(DelayTimeType, "Delay Sync Type", "Specifies if and how the delay should sync to a clock");
+		DEFINE_INPUT_METASOUND_PARAM(DelayTime, "Delay Time", "The period of the delay. If Delay Sync Type is TempoSync, the unit is quarter notes. Otherwise, the unit is seconds.");
+		DEFINE_INPUT_METASOUND_PARAM(Feedback, "Feedback", "The feedback amount, range [0.0, 1.0]");
+		DEFINE_INPUT_METASOUND_PARAM(DryLevel, "Dry Gain", "The amount of dry signal to pass through, range [0.0, 1.0]");
+		DEFINE_INPUT_METASOUND_PARAM(WetLevel, "Wet Gain", "The amount of wet signal to output, range [0.0, 1.0]");
+		DEFINE_INPUT_METASOUND_PARAM(WetFilterEnabled, "Wet Filter Enabled", "Enables or disables the filter on the wet signal");
+		DEFINE_INPUT_METASOUND_PARAM(FeedbackFilterEnabled, "Feedback Filter Enabled", "Enables or disables the filter on the feedback signal");
+		DEFINE_INPUT_METASOUND_PARAM(FilterType, "Filter Type", "The type of filter to apply to the signal");
+		DEFINE_INPUT_METASOUND_PARAM(FilterCutoff, "Filter Cutoff", "The cutoff frequency of the filter, in Hz");
+		DEFINE_INPUT_METASOUND_PARAM(FilterQ, "Filter Q", "The Q of the filter");
+		DEFINE_INPUT_METASOUND_PARAM(LFOEnabled, "LFO Enabled", "Enables or disables the LFO on the delay time");
+		DEFINE_INPUT_METASOUND_PARAM(LFOTimeType, "LFO Sync Type", "Specifies if and how the LFO should sync to a clock");
+		DEFINE_METASOUND_PARAM_ALIAS(LFOFrequency, CommonPinNames::Inputs::LFOFrequency);
+		DEFINE_INPUT_METASOUND_PARAM(LFODepth, "LFO Depth", "The amount of delay to be added by the LFO. If Delay Sync Type is TempoSync, the unit is quarter notes. Otherwise, the unit is seconds.");
+		DEFINE_INPUT_METASOUND_PARAM(StereoType, "Stereo Type", "Specifies method for panning and spreading the sound");
+		DEFINE_INPUT_METASOUND_PARAM(StereoSpreadLeft, "Stereo Spread Left", "The amount of stereo spread on the left channel");
+		DEFINE_INPUT_METASOUND_PARAM(StereoSpreadRight, "Stereo Spread Right", "The amount of stereo spread on the right channel");
 	}
 
 	namespace Outputs
 	{
-		EXTERN_METASOUND_PARAM_D(AudioLeft, "Audio Left", "The left channel of processed audio");
-		EXTERN_METASOUND_PARAM_D(AudioRight, "Audio Right", "The right channel of processed audio");
+		DEFINE_METASOUND_PARAM_ALIAS(AudioLeft, CommonPinNames::Outputs::AudioLeft);
+		DEFINE_METASOUND_PARAM_ALIAS(AudioRight, CommonPinNames::Outputs::AudioRight);
 	}
 	
 	class FDelayOperator final : public TExecutableOperator<FDelayOperator>
diff --git a/Engine/Plugins/Runtime/Harmonix/Source/HarmonixMetasound/Private/HarmonixMetasound/Nodes/DjFilterNode.cpp b/Engine/Plugins/Runtime/Harmonix/Source/HarmonixMetasound/Private/HarmonixMetasound/Nodes/DjFilterNode.cpp
--- a/Engine/Plugins/Runtime/Harmonix/Source/HarmonixMetasound/Private/HarmonixMetasound/Nodes/DjFilterNode.cpp
+++ b/Engine/Plugins/Runtime/Harmonix/Source/HarmonixMetasound/Private/HarmonixMetasound/Nodes/DjFilterNode.cpp
@@ -29,19 +29,19 @@ namespace HarmonixMetasound::Nodes::DjFilter
 
 	namespace Inputs
 	{
-		EXTERN_METASOUND_PARAM_ALIAS(AudioMono, CommonPinNames::Inputs::AudioMono);
-		EXTERN_METASOUND_PARAM_D(Amount, "Amount", "The knob position. -1 to DeadZoneSize will low-pass the signal. DeadZoneSize to 1 will high-pass the signal. Within the dead zone will fade to dry.");
-		EXTERN_METASOUND_PARAM_D(Resonance, "Resonance", "The filter resonance");
-		EXTERN_METASOUND_PARAM_D(LowPassMinFrequency, "Low-Pass Min Frequency", "The frequency the low-pass will be at when Amount is set to -1");
-		EXTERN_METASOUND_PARAM_D(LowPassMaxFrequency, "Low-Pass Max Frequency", "The frequency the low-pass will be at when Amount is set to -DeadZoneSize");
-		EXTERN_METASOUND_PARAM_D(HighPassMinFrequency, "High-Pass Min Frequency", "The frequency the high-pass will be at when Amount is set to DeadZoneSize");
-		EXTERN_METASOUND_PARAM_D(HighPassMaxFrequency, "High-Pass Max Frequency", "The frequency the high-pass will be at when Amount is set to 1");
-		EXTERN_METASOUND_PARAM_D(DeadZoneSize, "Dead Zone Size", "The portion of the Amount that will cross-fade between the filtered signal and the dry signal");
+		DEFINE_METASOUND_PARAM_ALIAS(AudioMono, CommonPinNames::Inputs::AudioMono);
+		DEFINE_INPUT_METASOUND_PARAM(Amount, "Amount", "The knob position. -1 to DeadZoneSize will low-pass the signal. DeadZoneSize to 1 will high-pass the signal. Within the dead zone will fade to dry.");
+		DEFINE_INPUT_METASOUND_PARAM(Resonance, "Resonance", "The filter resonance");
+		DEFINE_INPUT_METASOUND_PARAM(LowPassMinFrequency, "Low-Pass Min Frequency", "The frequency the low-pass will be at when Amount is set to -1");
+		DEFINE_INPUT_METASOUND_PARAM(LowPassMaxFrequency, "Low-Pass Max Frequency", "The frequency the low-pass will be at when Amount is set to -DeadZoneSize");
+		DEFINE_INPUT_METASOUND_PARAM(HighPassMinFrequency, "High-Pass Min Frequency", "The frequency the high-pass will be at when Amount is set to DeadZoneSize");
+		DEFINE_INPUT_METASOUND_PARAM(HighPassMaxFrequency, "High-Pass Max Frequency", "The frequency the high-pass will be at when Amount is set to 1");
+		DEFINE_INPUT_METASOUND_PARAM(DeadZoneSize, "Dead Zone Size", "The portion of the Amount that will cross-fade between the filtered signal and the dry signal");
 	}
 
 	namespace Outputs
 	{
-		EXTERN_METASOUND_PARAM_ALIAS(AudioMono, CommonPinNames::Outputs::AudioMono);
+		DEFINE_METASOUND_PARAM_ALIAS(AudioMono, CommonPinNames::Outputs::AudioMono);
 	}
 
 	class FOp final : public Metasound::TExecutableOperator<FOp>
diff --git a/Engine/Plugins/Runtime/Harmonix/Source/HarmonixMetasound/Private/HarmonixMetasound/Nodes/FFTAnalyzerNode.cpp b/Engine/Plugins/Runtime/Harmonix/Source/HarmonixMetasound/Private/HarmonixMetasound/Nodes/FFTAnalyzerNode.cpp
--- a/Engine/Plugins/Runtime/Harmonix/Source/HarmonixMetasound/Private/HarmonixMetasound/Nodes/FFTAnalyzerNode.cpp
+++ b/Engine/Plugins/Runtime/Harmonix/Source/HarmonixMetasound/Private/HarmonixMetasound/Nodes/FFTAnalyzerNode.cpp
@@ -34,7 +34,7 @@ namespace HarmonixMetasound
 	{
 		namespace Inputs
 		{
-			METASOUND_PARAM(Enable, "Enable", "When set, enables the analyzer, otherwise skips processing.")
+			DEFINE_METASOUND_PARAM_ALIAS(Enable, CommonPinNames::Inputs::Enable);
 			METASOUND_PARAM(Audio, "In", "Audio Input");
 			METASOUND_PARAM(FFTSize, "FFT Size", "The size of the FFT window, which determines the frequency resolution of the resulting spectrum analysis (a larger size will result in higher frequency resolution).");
 			METASOUND_PARAM(MinFrequencyHz, "Min Frequency Hz", "The minimum frequency that will be included in the spectrum analysis.");
diff --git a/Engine/Plugins/Runtime/Harmonix/Source/HarmonixMetasound/Private/HarmonixMetasound/Nodes/MidiCCTriggerNode.cpp b/Engine/Plugins/Runtime/Harmonix/Source/HarmonixMetasound/Private/HarmonixMetasound/Nodes/MidiCCTriggerNode.cpp
--- a/Engine/Plugins/Runtime/Harmonix/Source/HarmonixMetasound/Private/HarmonixMetasound/Nodes/MidiCCTriggerNode.cpp
+++ b/Engine/Plugins/Runtime/Harmonix/Source/HarmonixMetasound/Private/HarmonixMetasound/Nodes/MidiCCTriggerNode.cpp
@@ -35,18 +35,18 @@ namespace HarmonixMetasound::Nodes::MidiCCTriggerNode
 
 	namespace Inputs
 	{
-		EXTERN_METASOUND_PARAM_ALIAS(Enable, CommonPinNames::Inputs::Enable);
-		EXTERN_METASOUND_PARAM_ALIAS(MidiTrackNumber, CommonPinNames::Inputs::MidiTrackNumber);
-		EXTERN_METASOUND_PARAM_ALIAS(MidiChannelNumber, CommonPinNames::Inputs::MidiChannelNumber);
-		EXTERN_METASOUND_PARAM_D(InputMidiControllerID, "Standard Midi Controller ID", "Standard Midi Controller ID (0-127)");
-		EXTERN_METASOUND_PARAM_ALIAS(MidiStream, CommonPinNames::Inputs::MidiStream);
+		DEFINE_METASOUND_PARAM_ALIAS(Enable, CommonPinNames::Inputs::Enable);
+		DEFINE_METASOUND_PARAM_ALIAS(MidiTrackNumber, CommonPinNames::Inputs::MidiTrackNumber);
+		DEFINE_METASOUND_PARAM_ALIAS(MidiChannelNumber, CommonPinNames::Inputs::MidiChannelNumber);
+		DEFINE_INPUT_METASOUND_PARAM(InputMidiControllerID, "Standard Midi Controller ID", "Standard Midi Controller ID (0-127)");
+		DEFINE_METASOUND_PARAM_ALIAS(MidiStream, CommonPinNames::Inputs::MidiStream);
 	}
 
 	namespace Outputs
 	{
-		EXTERN_METASOUND_PARAM_D(OutputControlChangeValueInt32, "Midi Control Change Value (Int32)", "Midi Control Change value (0-127)");
-		EXTERN_METASOUND_PARAM_D(OutputControlChangeValueFloat, "Midi Control Change Value (Float)", "normalized Midi Control Change value (0.0-1.0)");
-		EXTERN_METASOUND_PARAM_D(OutputTrigger, "Trigger Out", "A trigger when a Midi Control Change message is found");
+		DEFINE_OUTPUT_METASOUND_PARAM(OutputControlChangeValueInt32, "Midi Control Change Value (Int32)", "Midi Control Change value (0-127)");
+		DEFINE_OUTPUT_METASOUND_PARAM(OutputControlChangeValueFloat, "Midi Control Change Value (Float)", "normalized Midi Control Change value (0.0-1.0)");
+		DEFINE_OUTPUT_METASOUND_PARAM(OutputTrigger, "Trigger Out", "A trigger when a Midi Control Change message is found");
 	}
 
 	class FOp final : public Metasound::TExecutableOperator<FOp>
diff --git a/Engine/Plugins/Runtime/Harmonix/Source/HarmonixMetasound/Private/HarmonixMetasound/Nodes/MorphingLfoNode.cpp b/Engine/Plugins/Runtime/Harmonix/Source/HarmonixMetasound/Private/HarmonixMetasound/Nodes/MorphingLfoNode.cpp
--- a/Engine/Plugins/Runtime/Harmonix/Source/HarmonixMetasound/Private/HarmonixMetasound/Nodes/MorphingLfoNode.cpp
+++ b/Engine/Plugins/Runtime/Harmonix/Source/HarmonixMetasound/Private/HarmonixMetasound/Nodes/MorphingLfoNode.cpp
@@ -18,16 +18,16 @@ namespace HarmonixMetasound::Nodes::MorphingLFO
 {
 	namespace Inputs
 	{
-		EXTERN_METASOUND_PARAM_ALIAS(MidiClock, CommonPinNames::Inputs::MidiClock);
-		EXTERN_METASOUND_PARAM_ALIAS(LFOSyncType, CommonPinNames::Inputs::LFOSyncType);
-		EXTERN_METASOUND_PARAM_ALIAS(LFOFrequency, CommonPinNames::Inputs::LFOFrequency);
-		EXTERN_METASOUND_PARAM_ALIAS(LFOInvert, CommonPinNames::Inputs::LFOInvert);
-		EXTERN_METASOUND_PARAM_ALIAS(LFOShape, CommonPinNames::Inputs::LFOShape);
+		DEFINE_METASOUND_PARAM_ALIAS(MidiClock, CommonPinNames::Inputs::MidiClock);
+		DEFINE_METASOUND_PARAM_ALIAS(LFOSyncType, CommonPinNames::Inputs::LFOSyncType);
+		DEFINE_METASOUND_PARAM_ALIAS(LFOFrequency, CommonPinNames::Inputs::LFOFrequency);
+		DEFINE_METASOUND_PARAM_ALIAS(LFOInvert, CommonPinNames::Inputs::LFOInvert);
+		DEFINE_METASOUND_PARAM_ALIAS(LFOShape, CommonPinNames::Inputs::LFOShape);
 	}
 
 	namespace Outputs
 	{
-		EXTERN_METASOUND_PARAM_D(LFO, "LFO Output", "The output of the LFO, range [0.0, 1.0]")
+		DEFINE_OUTPUT_METASOUND_PARAM(LFO, "LFO Output", "The output of the LFO, range [0.0, 1.0]")
 	}
 
 	template<typename OutputDataType>
diff --git a/Engine/Plugins/Runtime/Harmonix/Source/HarmonixMetasound/Private/HarmonixMetasound/Nodes/MultibandAnalyzerNode.cpp b/Engine/Plugins/Runtime/Harmonix/Source/HarmonixMetasound/Private/HarmonixMetasound/Nodes/MultibandAnalyzerNode.cpp
--- a/Engine/Plugins/Runtime/Harmonix/Source/HarmonixMetasound/Private/HarmonixMetasound/Nodes/MultibandAnalyzerNode.cpp
+++ b/Engine/Plugins/Runtime/Harmonix/Source/HarmonixMetasound/Private/HarmonixMetasound/Nodes/MultibandAnalyzerNode.cpp
@@ -33,18 +33,18 @@ namespace HarmonixMetasound::Nodes::MultibandAnalyzer
 	
 	namespace Inputs
 	{
-		EXTERN_METASOUND_PARAM_ALIAS(Enable, CommonPinNames::Inputs::Enable);
-		EXTERN_METASOUND_PARAM_ALIAS(AudioMono, CommonPinNames::Inputs::AudioMono);
-		EXTERN_METASOUND_PARAM_D(CrossoverFrequencies, "Crossover Frequencies", "The frequencies at which the analysis bands should be split");
-		EXTERN_METASOUND_PARAM_D(ApplySmoothing, "Apply Smoothing", "If enabled, will use an envelope follower on each band to smooth the output levels. Otherwise, the raw peak value for each band will be reported.");
-		EXTERN_METASOUND_PARAM_D(AttackTime, "Attack Time", "The amount of smoothing on rising amplitude, in seconds");
-		EXTERN_METASOUND_PARAM_D(ReleaseTime, "Release Time", "The amount of smoothing on falling amplitudes, in seconds");
-		EXTERN_METASOUND_PARAM_D(PeakMode, "Peak Mode", "The method for reporting output levels for each band");
+		DEFINE_METASOUND_PARAM_ALIAS(Enable, CommonPinNames::Inputs::Enable);
+		DEFINE_METASOUND_PARAM_ALIAS(AudioMono, CommonPinNames::Inputs::AudioMono);
+		DEFINE_INPUT_METASOUND_PARAM(CrossoverFrequencies, "Crossover Frequencies", "The frequencies at which the analysis bands should be split");
+		DEFINE_INPUT_METASOUND_PARAM(ApplySmoothing, "Apply Smoothing", "If enabled, will use an envelope follower on each band to smooth the output levels. Otherwise, the raw peak value for each band will be reported.");
+		DEFINE_INPUT_METASOUND_PARAM(AttackTime, "Attack Time", "The amount of smoothing on rising amplitude, in seconds");
+		DEFINE_INPUT_METASOUND_PARAM(ReleaseTime, "Release Time", "The amount of smoothing on falling amplitudes, in seconds");
+		DEFINE_INPUT_METASOUND_PARAM(PeakMode, "Peak Mode", "The method for reporting output levels for each band");
 	}
 
 	namespace Outputs
 	{
-		EXTERN_METASOUND_PARAM_D(BandLevels, "Band Levels", "The levels for each band");
+		DEFINE_OUTPUT_METASOUND_PARAM(BandLevels, "Band Levels", "The levels for each band");
 	}
 
 	/**
diff --git a/Engine/Plugins/Runtime/Harmonix/Source/HarmonixMetasound/Private/HarmonixMetasound/Nodes/PeakNode.cpp b/Engine/Plugins/Runtime/Harmonix/Source/HarmonixMetasound/Private/HarmonixMetasound/Nodes/PeakNode.cpp
--- a/Engine/Plugins/Runtime/Harmonix/Source/HarmonixMetasound/Private/HarmonixMetasound/Nodes/PeakNode.cpp
+++ b/Engine/Plugins/Runtime/Harmonix/Source/HarmonixMetasound/Private/HarmonixMetasound/Nodes/PeakNode.cpp
@@ -29,13 +29,13 @@ namespace HarmonixMetasound::Nodes::Peak
 
 	namespace Inputs
 	{
-		EXTERN_METASOUND_PARAM_ALIAS(AudioMono, CommonPinNames::Inputs::AudioMono);
+		DEFINE_METASOUND_PARAM_ALIAS(AudioMono, CommonPinNames::Inputs::AudioMono);
 		METASOUND_PARAM(ParamEnable, "Enable", "Enable the peak node.")
 	}
 
 	namespace Outputs
 	{
-		EXTERN_METASOUND_PARAM_D(Peak, "Peak", "The peak for the latest block")
+		DEFINE_OUTPUT_METASOUND_PARAM(Peak, "Peak", "The peak for the latest block")
 	}
 
 	class FOp final : public Metasound::TExecutableOperator<FOp>
diff --git a/Engine/Plugins/Runtime/Harmonix/Source/HarmonixMetasound/Public/HarmonixMetasound/Common.h b/Engine/Plugins/Runtime/Harmonix/Source/HarmonixMetasound/Public/HarmonixMetasound/Common.h
--- a/Engine/Plugins/Runtime/Harmonix/Source/HarmonixMetasound/Public/HarmonixMetasound/Common.h
+++ b/Engine/Plugins/Runtime/Harmonix/Source/HarmonixMetasound/Public/HarmonixMetasound/Common.h
@@ -22,27 +22,35 @@ namespace HarmonixMetasound
 		HARMONIXMETASOUND_API extern const FText Modulation;
 	}
 
-#define EXTERN_METASOUND_PARAM(NAME)      \
+#define DECLARE_METASOUND_PARAM_EXTERN(NAME)      \
 	HARMONIXMETASOUND_API extern const TCHAR* NAME##Name;       \
 	HARMONIXMETASOUND_API extern const FText NAME##Tooltip;     \
 	HARMONIXMETASOUND_API extern const FText NAME##DisplayName; 
 
 #if WITH_EDITOR
 #define LOC_DEFINE_REGION
-#define EXTERN_METASOUND_PARAM_D(NAME, NAME_TEXT, TOOLTIP_TEXT)              \
-	const TCHAR* NAME##Name = TEXT(NAME_TEXT);                               \
-	const FText NAME##Tooltip = LOCTEXT(#NAME "Tooltip", TOOLTIP_TEXT);      \
-	const FText NAME##DisplayName = LOCTEXT(#NAME "DisplayName", NAME_TEXT); 
+#define DEFINE_INPUT_METASOUND_PARAM(NAME, NAME_TEXT, TOOLTIP_TEXT)						\
+	const TCHAR* NAME##Name = TEXT(NAME_TEXT);											\
+	const FText NAME##Tooltip = LOCTEXT(#NAME "InputTooltip", TOOLTIP_TEXT);			\
+	const FText NAME##DisplayName = LOCTEXT(#NAME "InputDisplayName", NAME_TEXT);
+#define DEFINE_OUTPUT_METASOUND_PARAM(NAME, NAME_TEXT, TOOLTIP_TEXT)					\
+	const TCHAR* NAME##Name = TEXT(NAME_TEXT);											\
+	const FText NAME##Tooltip = LOCTEXT(#NAME "OutputTooltip", TOOLTIP_TEXT);			\
+    const FText NAME##DisplayName = LOCTEXT(#NAME "OutputDisplayName", NAME_TEXT);
 #undef LOC_DEFINE_REGION
 #else 
-#define EXTERN_METASOUND_PARAM_D(NAME, NAME_TEXT, TOOLTIP_TEXT) \
-	const TCHAR* NAME##Name = TEXT(NAME_TEXT);                  \
-	const FText NAME##Tooltip = FText::GetEmpty();              \
+#define DEFINE_INPUT_METASOUND_PARAM(NAME, NAME_TEXT, TOOLTIP_TEXT)		\
+	const TCHAR* NAME##Name = TEXT(NAME_TEXT);							\
+	const FText NAME##Tooltip = FText::GetEmpty();						\
+	const FText NAME##DisplayName = FText::GetEmpty();
+#define DEFINE_OUTPUT_METASOUND_PARAM(NAME, NAME_TEXT, TOOLTIP_TEXT)	\
+	const TCHAR* NAME##Name = TEXT(NAME_TEXT);							\
+	const FText NAME##Tooltip = FText::GetEmpty();						\
 	const FText NAME##DisplayName = FText::GetEmpty();
 #endif // WITH_EDITOR
 
 // Alias an existing param definition. Helpful for re-using definitions while providing a public node API.
-#define EXTERN_METASOUND_PARAM_ALIAS(ALIAS_NAME, TARGET_NAME) \
+#define DEFINE_METASOUND_PARAM_ALIAS(ALIAS_NAME, TARGET_NAME) \
 	const TCHAR* ALIAS_NAME##Name = TARGET_NAME##Name; \
 	const FText ALIAS_NAME##Tooltip = TARGET_NAME##Tooltip; \
 	const FText ALIAS_NAME##DisplayName = TARGET_NAME##DisplayName; 
@@ -51,87 +59,89 @@ namespace HarmonixMetasound
 	{
 		namespace Inputs
 		{
-			EXTERN_METASOUND_PARAM(Enable)
-			EXTERN_METASOUND_PARAM(MidiStream)
-			EXTERN_METASOUND_PARAM(MidiChannelNumber)
-			EXTERN_METASOUND_PARAM(MidiChannelFilterSpecifier)
-			EXTERN_METASOUND_PARAM(MidiTrackNumber)
-			EXTERN_METASOUND_PARAM(MidiTrackIndexFilterSpecifier)
-			EXTERN_METASOUND_PARAM(MinMidiNote)
-			EXTERN_METASOUND_PARAM(MaxMidiNote)
-			EXTERN_METASOUND_PARAM(MinMidiVelocity)
-			EXTERN_METASOUND_PARAM(MaxMidiVelocity)
-			EXTERN_METASOUND_PARAM(MidiFileAsset)
-			EXTERN_METASOUND_PARAM(Transport)
-			EXTERN_METASOUND_PARAM(MidiClock)
-			EXTERN_METASOUND_PARAM(Tempo)
-			EXTERN_METASOUND_PARAM(TimeSigNumerator)
-			EXTERN_METASOUND_PARAM(TimeSigDenominator)
-			EXTERN_METASOUND_PARAM(Speed)
-			EXTERN_METASOUND_PARAM(Loop)
-			EXTERN_METASOUND_PARAM(LoopLengthBars)
-			EXTERN_METASOUND_PARAM(TransportPrepare)
-			EXTERN_METASOUND_PARAM(TransportPlay)
-			EXTERN_METASOUND_PARAM(TransportPause)
-			EXTERN_METASOUND_PARAM(TransportContinue)
-			EXTERN_METASOUND_PARAM(TransportStop)
-			EXTERN_METASOUND_PARAM(TransportKill)
-			EXTERN_METASOUND_PARAM(TriggerSeek)
-			EXTERN_METASOUND_PARAM(SeekDestination)
-			EXTERN_METASOUND_PARAM(PrerollBars)
-			EXTERN_METASOUND_PARAM(GridSizeUnits)
-			EXTERN_METASOUND_PARAM(GridSizeMult)
-			EXTERN_METASOUND_PARAM(OffsetUnits)
-			EXTERN_METASOUND_PARAM(OffsetMult)
-			EXTERN_METASOUND_PARAM(Bar)
-			EXTERN_METASOUND_PARAM(FloatBeat)
-			EXTERN_METASOUND_PARAM(SynthPatch)
-			EXTERN_METASOUND_PARAM(Transposition)
-			EXTERN_METASOUND_PARAM(ClockSpeedToPitch)
-			EXTERN_METASOUND_PARAM(ClockSpeedToFrequency)
-			EXTERN_METASOUND_PARAM(LFOSyncType);
-			EXTERN_METASOUND_PARAM(LFOFrequency);
-			EXTERN_METASOUND_PARAM(LFOInvert);
-			EXTERN_METASOUND_PARAM(LFOShape);
-			EXTERN_METASOUND_PARAM(AudioMono);
+			DECLARE_METASOUND_PARAM_EXTERN(Enable)
+			DECLARE_METASOUND_PARAM_EXTERN(MidiStream)
+			DECLARE_METASOUND_PARAM_EXTERN(MidiChannelNumber)
+			DECLARE_METASOUND_PARAM_EXTERN(MidiChannelFilterSpecifier)
+			DECLARE_METASOUND_PARAM_EXTERN(MidiTrackNumber)
+			DECLARE_METASOUND_PARAM_EXTERN(MidiTrackIndexFilterSpecifier)
+			DECLARE_METASOUND_PARAM_EXTERN(MinMidiNote)
+			DECLARE_METASOUND_PARAM_EXTERN(MaxMidiNote)
+			DECLARE_METASOUND_PARAM_EXTERN(MinMidiVelocity)
+			DECLARE_METASOUND_PARAM_EXTERN(MaxMidiVelocity)
+			DECLARE_METASOUND_PARAM_EXTERN(MidiFileAsset)
+			DECLARE_METASOUND_PARAM_EXTERN(Transport)
+			DECLARE_METASOUND_PARAM_EXTERN(MidiClock)
+			DECLARE_METASOUND_PARAM_EXTERN(Tempo)
+			DECLARE_METASOUND_PARAM_EXTERN(TimeSigNumerator)
+			DECLARE_METASOUND_PARAM_EXTERN(TimeSigDenominator)
+			DECLARE_METASOUND_PARAM_EXTERN(Speed)
+			DECLARE_METASOUND_PARAM_EXTERN(Loop)
+			DECLARE_METASOUND_PARAM_EXTERN(LoopLengthBars)
+			DECLARE_METASOUND_PARAM_EXTERN(TransportPrepare)
+			DECLARE_METASOUND_PARAM_EXTERN(TransportPlay)
+			DECLARE_METASOUND_PARAM_EXTERN(TransportPause)
+			DECLARE_METASOUND_PARAM_EXTERN(TransportContinue)
+			DECLARE_METASOUND_PARAM_EXTERN(TransportStop)
+			DECLARE_METASOUND_PARAM_EXTERN(TransportKill)
+			DECLARE_METASOUND_PARAM_EXTERN(TriggerSeek)
+			DECLARE_METASOUND_PARAM_EXTERN(SeekDestination)
+			DECLARE_METASOUND_PARAM_EXTERN(PrerollBars)
+			DECLARE_METASOUND_PARAM_EXTERN(GridSizeUnits)
+			DECLARE_METASOUND_PARAM_EXTERN(GridSizeMult)
+			DECLARE_METASOUND_PARAM_EXTERN(OffsetUnits)
+			DECLARE_METASOUND_PARAM_EXTERN(OffsetMult)
+			DECLARE_METASOUND_PARAM_EXTERN(Bar)
+			DECLARE_METASOUND_PARAM_EXTERN(FloatBeat)
+			DECLARE_METASOUND_PARAM_EXTERN(SynthPatch)
+			DECLARE_METASOUND_PARAM_EXTERN(Transposition)
+			DECLARE_METASOUND_PARAM_EXTERN(ClockSpeedToPitch)
+			DECLARE_METASOUND_PARAM_EXTERN(ClockSpeedToFrequency)
+			DECLARE_METASOUND_PARAM_EXTERN(LFOSyncType);
+			DECLARE_METASOUND_PARAM_EXTERN(LFOFrequency);
+			DECLARE_METASOUND_PARAM_EXTERN(LFOInvert);
+			DECLARE_METASOUND_PARAM_EXTERN(LFOShape);
+			DECLARE_METASOUND_PARAM_EXTERN(AudioMono);
+			DECLARE_METASOUND_PARAM_EXTERN(AudioLeft);
+			DECLARE_METASOUND_PARAM_EXTERN(AudioRight);
 		}
 		namespace Outputs
 		{
-			EXTERN_METASOUND_PARAM(Transport)
-			EXTERN_METASOUND_PARAM(MidiClock)
-			EXTERN_METASOUND_PARAM(MidiFileAsset)
-			EXTERN_METASOUND_PARAM(MidiStream)
-			EXTERN_METASOUND_PARAM(NoteOn)
-			EXTERN_METASOUND_PARAM(NoteOff)
-			EXTERN_METASOUND_PARAM(MidiNoteNumber)
-			EXTERN_METASOUND_PARAM(Frequency)
-			EXTERN_METASOUND_PARAM(MidiVelocity)
-			EXTERN_METASOUND_PARAM(NormalizedVelocity)
-			EXTERN_METASOUND_PARAM(TransportPrepare)
-			EXTERN_METASOUND_PARAM(TransportPlay)
-			EXTERN_METASOUND_PARAM(TransportPause)
-			EXTERN_METASOUND_PARAM(TransportContinue)
-			EXTERN_METASOUND_PARAM(TransportStop)
-			EXTERN_METASOUND_PARAM(TransportKill)
-			EXTERN_METASOUND_PARAM(MusicTimestamp)
-			EXTERN_METASOUND_PARAM(SeekTarget)
-			EXTERN_METASOUND_PARAM(AudioMono)
-			EXTERN_METASOUND_PARAM(AudioLeft)
-			EXTERN_METASOUND_PARAM(AudioRight)
-			EXTERN_METASOUND_PARAM(AudioCenter)
-			EXTERN_METASOUND_PARAM(AudioLFE)
-			EXTERN_METASOUND_PARAM(AudioLeftBack)
-			EXTERN_METASOUND_PARAM(AudioRightBack)
-			EXTERN_METASOUND_PARAM(AudioLeftSide)
-			EXTERN_METASOUND_PARAM(AudioRightSide)
-			EXTERN_METASOUND_PARAM(MusicTimespanBar)
-			EXTERN_METASOUND_PARAM(MusicTimespanBeat)
-			EXTERN_METASOUND_PARAM(Tempo)
-			EXTERN_METASOUND_PARAM(Speed)
-			EXTERN_METASOUND_PARAM(TimeSigNumerator)
-			EXTERN_METASOUND_PARAM(TimeSigDenominator)
-			EXTERN_METASOUND_PARAM(SecsPerQuarter)
-			EXTERN_METASOUND_PARAM(SecsPerBeat)
+			DECLARE_METASOUND_PARAM_EXTERN(Transport)
+			DECLARE_METASOUND_PARAM_EXTERN(MidiClock)
+			DECLARE_METASOUND_PARAM_EXTERN(MidiFileAsset)
+			DECLARE_METASOUND_PARAM_EXTERN(MidiStream)
+			DECLARE_METASOUND_PARAM_EXTERN(NoteOn)
+			DECLARE_METASOUND_PARAM_EXTERN(NoteOff)
+			DECLARE_METASOUND_PARAM_EXTERN(MidiNoteNumber)
+			DECLARE_METASOUND_PARAM_EXTERN(Frequency)
+			DECLARE_METASOUND_PARAM_EXTERN(MidiVelocity)
+			DECLARE_METASOUND_PARAM_EXTERN(NormalizedVelocity)
+			DECLARE_METASOUND_PARAM_EXTERN(TransportPrepare)
+			DECLARE_METASOUND_PARAM_EXTERN(TransportPlay)
+			DECLARE_METASOUND_PARAM_EXTERN(TransportPause)
+			DECLARE_METASOUND_PARAM_EXTERN(TransportContinue)
+			DECLARE_METASOUND_PARAM_EXTERN(TransportStop)
+			DECLARE_METASOUND_PARAM_EXTERN(TransportKill)
+			DECLARE_METASOUND_PARAM_EXTERN(MusicTimestamp)
+			DECLARE_METASOUND_PARAM_EXTERN(SeekTarget)
+			DECLARE_METASOUND_PARAM_EXTERN(AudioMono)
+			DECLARE_METASOUND_PARAM_EXTERN(AudioLeft)
+			DECLARE_METASOUND_PARAM_EXTERN(AudioRight)
+			DECLARE_METASOUND_PARAM_EXTERN(AudioCenter)
+			DECLARE_METASOUND_PARAM_EXTERN(AudioLFE)
+			DECLARE_METASOUND_PARAM_EXTERN(AudioLeftBack)
+			DECLARE_METASOUND_PARAM_EXTERN(AudioRightBack)
+			DECLARE_METASOUND_PARAM_EXTERN(AudioLeftSide)
+			DECLARE_METASOUND_PARAM_EXTERN(AudioRightSide)
+			DECLARE_METASOUND_PARAM_EXTERN(MusicTimespanBar)
+			DECLARE_METASOUND_PARAM_EXTERN(MusicTimespanBeat)
+			DECLARE_METASOUND_PARAM_EXTERN(Tempo)
+			DECLARE_METASOUND_PARAM_EXTERN(Speed)
+			DECLARE_METASOUND_PARAM_EXTERN(TimeSigNumerator)
+			DECLARE_METASOUND_PARAM_EXTERN(TimeSigDenominator)
+			DECLARE_METASOUND_PARAM_EXTERN(SecsPerQuarter)
+			DECLARE_METASOUND_PARAM_EXTERN(SecsPerBeat)
 		}
 	}
 }
diff --git a/Engine/Plugins/Runtime/Harmonix/Source/HarmonixMetasound/Public/HarmonixMetasound/Nodes/DelayNode.h b/Engine/Plugins/Runtime/Harmonix/Source/HarmonixMetasound/Public/HarmonixMetasound/Nodes/DelayNode.h
--- a/Engine/Plugins/Runtime/Harmonix/Source/HarmonixMetasound/Public/HarmonixMetasound/Nodes/DelayNode.h
+++ b/Engine/Plugins/Runtime/Harmonix/Source/HarmonixMetasound/Public/HarmonixMetasound/Nodes/DelayNode.h
@@ -14,31 +14,31 @@ namespace HarmonixMetasound::DelayNode
 
 	namespace Inputs
 	{
-		EXTERN_METASOUND_PARAM(AudioLeft);
-		EXTERN_METASOUND_PARAM(AudioRight);
-		EXTERN_METASOUND_PARAM(MidiClock)
-		EXTERN_METASOUND_PARAM(DelayTimeType);
-		EXTERN_METASOUND_PARAM(DelayTime);
-		EXTERN_METASOUND_PARAM(Feedback);
-		EXTERN_METASOUND_PARAM(DryLevel);
-		EXTERN_METASOUND_PARAM(WetLevel);
-		EXTERN_METASOUND_PARAM(WetFilterEnabled);
-		EXTERN_METASOUND_PARAM(FeedbackFilterEnabled);
-		EXTERN_METASOUND_PARAM(FilterType);
-		EXTERN_METASOUND_PARAM(FilterCutoff);
-		EXTERN_METASOUND_PARAM(FilterQ);
-		EXTERN_METASOUND_PARAM(LFOEnabled);
-		EXTERN_METASOUND_PARAM(LFOTimeType);
-		EXTERN_METASOUND_PARAM(LFOFrequency);
-		EXTERN_METASOUND_PARAM(LFODepth);
-		EXTERN_METASOUND_PARAM(StereoType);
-		EXTERN_METASOUND_PARAM(StereoSpreadLeft);
-		EXTERN_METASOUND_PARAM(StereoSpreadRight);
+		DECLARE_METASOUND_PARAM_EXTERN(AudioLeft);
+		DECLARE_METASOUND_PARAM_EXTERN(AudioRight);
+		DECLARE_METASOUND_PARAM_EXTERN(MidiClock)
+		DECLARE_METASOUND_PARAM_EXTERN(DelayTimeType);
+		DECLARE_METASOUND_PARAM_EXTERN(DelayTime);
+		DECLARE_METASOUND_PARAM_EXTERN(Feedback);
+		DECLARE_METASOUND_PARAM_EXTERN(DryLevel);
+		DECLARE_METASOUND_PARAM_EXTERN(WetLevel);
+		DECLARE_METASOUND_PARAM_EXTERN(WetFilterEnabled);
+		DECLARE_METASOUND_PARAM_EXTERN(FeedbackFilterEnabled);
+		DECLARE_METASOUND_PARAM_EXTERN(FilterType);
+		DECLARE_METASOUND_PARAM_EXTERN(FilterCutoff);
+		DECLARE_METASOUND_PARAM_EXTERN(FilterQ);
+		DECLARE_METASOUND_PARAM_EXTERN(LFOEnabled);
+		DECLARE_METASOUND_PARAM_EXTERN(LFOTimeType);
+		DECLARE_METASOUND_PARAM_EXTERN(LFOFrequency);
+		DECLARE_METASOUND_PARAM_EXTERN(LFODepth);
+		DECLARE_METASOUND_PARAM_EXTERN(StereoType);
+		DECLARE_METASOUND_PARAM_EXTERN(StereoSpreadLeft);
+		DECLARE_METASOUND_PARAM_EXTERN(StereoSpreadRight);
 	}
 
 	namespace Outputs
 	{
-		EXTERN_METASOUND_PARAM(AudioLeft);
-		EXTERN_METASOUND_PARAM(AudioRight);
+		DECLARE_METASOUND_PARAM_EXTERN(AudioLeft);
+		DECLARE_METASOUND_PARAM_EXTERN(AudioRight);
 	}
 }
diff --git a/Engine/Plugins/Runtime/Harmonix/Source/HarmonixMetasound/Public/HarmonixMetasound/Nodes/DjFilterNode.h b/Engine/Plugins/Runtime/Harmonix/Source/HarmonixMetasound/Public/HarmonixMetasound/Nodes/DjFilterNode.h
--- a/Engine/Plugins/Runtime/Harmonix/Source/HarmonixMetasound/Public/HarmonixMetasound/Nodes/DjFilterNode.h
+++ b/Engine/Plugins/Runtime/Harmonix/Source/HarmonixMetasound/Public/HarmonixMetasound/Nodes/DjFilterNode.h
@@ -12,18 +12,18 @@ namespace HarmonixMetasound::Nodes::DjFilter
 	
 	namespace Inputs
 	{
-		EXTERN_METASOUND_PARAM(AudioMono);
-		EXTERN_METASOUND_PARAM(Amount);
-		EXTERN_METASOUND_PARAM(Resonance);
-		EXTERN_METASOUND_PARAM(LowPassMinFrequency);
-		EXTERN_METASOUND_PARAM(LowPassMaxFrequency);
-		EXTERN_METASOUND_PARAM(HighPassMinFrequency);
-		EXTERN_METASOUND_PARAM(HighPassMaxFrequency);
-		EXTERN_METASOUND_PARAM(DeadZoneSize);
+		DECLARE_METASOUND_PARAM_EXTERN(AudioMono);
+		DECLARE_METASOUND_PARAM_EXTERN(Amount);
+		DECLARE_METASOUND_PARAM_EXTERN(Resonance);
+		DECLARE_METASOUND_PARAM_EXTERN(LowPassMinFrequency);
+		DECLARE_METASOUND_PARAM_EXTERN(LowPassMaxFrequency);
+		DECLARE_METASOUND_PARAM_EXTERN(HighPassMinFrequency);
+		DECLARE_METASOUND_PARAM_EXTERN(HighPassMaxFrequency);
+		DECLARE_METASOUND_PARAM_EXTERN(DeadZoneSize);
 	}
 
 	namespace Outputs
 	{
-		EXTERN_METASOUND_PARAM(AudioMono);
+		DECLARE_METASOUND_PARAM_EXTERN(AudioMono);
 	}
 }
diff --git a/Engine/Plugins/Runtime/Harmonix/Source/HarmonixMetasound/Public/HarmonixMetasound/Nodes/MidiCCTriggerNode.h b/Engine/Plugins/Runtime/Harmonix/Source/HarmonixMetasound/Public/HarmonixMetasound/Nodes/MidiCCTriggerNode.h
--- a/Engine/Plugins/Runtime/Harmonix/Source/HarmonixMetasound/Public/HarmonixMetasound/Nodes/MidiCCTriggerNode.h
+++ b/Engine/Plugins/Runtime/Harmonix/Source/HarmonixMetasound/Public/HarmonixMetasound/Nodes/MidiCCTriggerNode.h
@@ -12,17 +12,17 @@ namespace HarmonixMetasound::Nodes::MidiCCTriggerNode
 
 	namespace Inputs
 	{
-		EXTERN_METASOUND_PARAM(Enable);
-		EXTERN_METASOUND_PARAM(MidiTrackNumber);
-		EXTERN_METASOUND_PARAM(MidiChannelNumber);
-		EXTERN_METASOUND_PARAM(InputMidiControllerID);
-		EXTERN_METASOUND_PARAM(MidiStream);
+		DECLARE_METASOUND_PARAM_EXTERN(Enable);
+		DECLARE_METASOUND_PARAM_EXTERN(MidiTrackNumber);
+		DECLARE_METASOUND_PARAM_EXTERN(MidiChannelNumber);
+		DECLARE_METASOUND_PARAM_EXTERN(InputMidiControllerID);
+		DECLARE_METASOUND_PARAM_EXTERN(MidiStream);
 	}
 
 	namespace Outputs
 	{
-		EXTERN_METASOUND_PARAM(OutputControlChangeValueInt32);
-		EXTERN_METASOUND_PARAM(OutputControlChangeValueFloat);
-		EXTERN_METASOUND_PARAM(OutputTrigger);
+		DECLARE_METASOUND_PARAM_EXTERN(OutputControlChangeValueInt32);
+		DECLARE_METASOUND_PARAM_EXTERN(OutputControlChangeValueFloat);
+		DECLARE_METASOUND_PARAM_EXTERN(OutputTrigger);
 	}
 }
\ No newline at end of file
diff --git a/Engine/Plugins/Runtime/Harmonix/Source/HarmonixMetasound/Public/HarmonixMetasound/Nodes/MorphingLfoNode.h b/Engine/Plugins/Runtime/Harmonix/Source/HarmonixMetasound/Public/HarmonixMetasound/Nodes/MorphingLfoNode.h
--- a/Engine/Plugins/Runtime/Harmonix/Source/HarmonixMetasound/Public/HarmonixMetasound/Nodes/MorphingLfoNode.h
+++ b/Engine/Plugins/Runtime/Harmonix/Source/HarmonixMetasound/Public/HarmonixMetasound/Nodes/MorphingLfoNode.h
@@ -22,15 +22,15 @@ namespace HarmonixMetasound::Nodes::MorphingLFO
 	
 	namespace Inputs
 	{
-		EXTERN_METASOUND_PARAM(MidiClock);
-		EXTERN_METASOUND_PARAM(LFOSyncType);
-		EXTERN_METASOUND_PARAM(LFOFrequency);
-		EXTERN_METASOUND_PARAM(LFOInvert);
-		EXTERN_METASOUND_PARAM(LFOShape)
+		DECLARE_METASOUND_PARAM_EXTERN(MidiClock);
+		DECLARE_METASOUND_PARAM_EXTERN(LFOSyncType);
+		DECLARE_METASOUND_PARAM_EXTERN(LFOFrequency);
+		DECLARE_METASOUND_PARAM_EXTERN(LFOInvert);
+		DECLARE_METASOUND_PARAM_EXTERN(LFOShape)
 	}
 
 	namespace Outputs
 	{
-		EXTERN_METASOUND_PARAM(LFO);
+		DECLARE_METASOUND_PARAM_EXTERN(LFO);
 	}
 }
\ No newline at end of file
diff --git a/Engine/Plugins/Runtime/Harmonix/Source/HarmonixMetasound/Public/HarmonixMetasound/Nodes/MultibandAnalyzerNode.h b/Engine/Plugins/Runtime/Harmonix/Source/HarmonixMetasound/Public/HarmonixMetasound/Nodes/MultibandAnalyzerNode.h
--- a/Engine/Plugins/Runtime/Harmonix/Source/HarmonixMetasound/Public/HarmonixMetasound/Nodes/MultibandAnalyzerNode.h
+++ b/Engine/Plugins/Runtime/Harmonix/Source/HarmonixMetasound/Public/HarmonixMetasound/Nodes/MultibandAnalyzerNode.h
@@ -23,17 +23,17 @@ namespace HarmonixMetasound::Nodes::MultibandAnalyzer
 	
 	namespace Inputs
 	{
-		EXTERN_METASOUND_PARAM(Enable);
-		EXTERN_METASOUND_PARAM(AudioMono);
-		EXTERN_METASOUND_PARAM(CrossoverFrequencies);
-		EXTERN_METASOUND_PARAM(ApplySmoothing);
-		EXTERN_METASOUND_PARAM(AttackTime);
-		EXTERN_METASOUND_PARAM(ReleaseTime);
-		EXTERN_METASOUND_PARAM(PeakMode);
+		DECLARE_METASOUND_PARAM_EXTERN(Enable);
+		DECLARE_METASOUND_PARAM_EXTERN(AudioMono);
+		DECLARE_METASOUND_PARAM_EXTERN(CrossoverFrequencies);
+		DECLARE_METASOUND_PARAM_EXTERN(ApplySmoothing);
+		DECLARE_METASOUND_PARAM_EXTERN(AttackTime);
+		DECLARE_METASOUND_PARAM_EXTERN(ReleaseTime);
+		DECLARE_METASOUND_PARAM_EXTERN(PeakMode);
 	}
 
 	namespace Outputs
 	{
-		EXTERN_METASOUND_PARAM(BandLevels);
+		DECLARE_METASOUND_PARAM_EXTERN(BandLevels);
 	}
 }
diff --git a/Engine/Plugins/Runtime/Harmonix/Source/HarmonixMetasound/Public/HarmonixMetasound/Nodes/PeakNode.h b/Engine/Plugins/Runtime/Harmonix/Source/HarmonixMetasound/Public/HarmonixMetasound/Nodes/PeakNode.h
--- a/Engine/Plugins/Runtime/Harmonix/Source/HarmonixMetasound/Public/HarmonixMetasound/Nodes/PeakNode.h
+++ b/Engine/Plugins/Runtime/Harmonix/Source/HarmonixMetasound/Public/HarmonixMetasound/Nodes/PeakNode.h
@@ -12,11 +12,11 @@ namespace HarmonixMetasound::Nodes::Peak
 	
 	namespace Inputs
 	{
-		EXTERN_METASOUND_PARAM(AudioMono);
+		DECLARE_METASOUND_PARAM_EXTERN(AudioMono);
 	}
 
 	namespace Outputs
 	{
-		EXTERN_METASOUND_PARAM(Peak);
+		DECLARE_METASOUND_PARAM_EXTERN(Peak);
 	}
 }
