Don't reference UFusionPatch static class in ShutDown module.  Causes breakages in ASAN

#tests PIE

diff --git a/Engine/Plugins/Runtime/Harmonix/Source/HarmonixDspEditor/Private/HarmonixDspEditorModule.cpp b/Engine/Plugins/Runtime/Harmonix/Source/HarmonixDspEditor/Private/HarmonixDspEditorModule.cpp
--- a/Engine/Plugins/Runtime/Harmonix/Source/HarmonixDspEditor/Private/HarmonixDspEditorModule.cpp
+++ b/Engine/Plugins/Runtime/Harmonix/Source/HarmonixDspEditor/Private/HarmonixDspEditorModule.cpp
@@ -34,7 +34,7 @@ void FHarmonixDspEditorModule::ShutdownModule()
 {
 
 	FPropertyEditorModule& PropertyEditorModule = FModuleManager::LoadModuleChecked<FPropertyEditorModule>("PropertyEditor");
-	PropertyEditorModule.UnregisterCustomClassLayout(UFusionPatch::StaticClass()->GetFName());
+	PropertyEditorModule.UnregisterCustomClassLayout("FusionPatch");
 	PropertyEditorModule.UnregisterCustomClassLayout("StretcherAndPitchShifterFactoryConfig");
 	PropertyEditorModule.UnregisterCustomPropertyTypeLayout("PannerDetails");
 	PropertyEditorModule.UnregisterCustomPropertyTypeLayout("PitchShifterName");
