[Harmonix] Change display name: Midi Stream Logger -> MIDI Logger

diff --git a/Engine/Plugins/Runtime/Harmonix/Source/HarmonixMetasound/Private/HarmonixMetasound/Nodes/MidiStreamLoggerNode.cpp b/Engine/Plugins/Runtime/Harmonix/Source/HarmonixMetasound/Private/HarmonixMetasound/Nodes/MidiStreamLoggerNode.cpp
--- a/Engine/Plugins/Runtime/Harmonix/Source/HarmonixMetasound/Private/HarmonixMetasound/Nodes/MidiStreamLoggerNode.cpp
+++ b/Engine/Plugins/Runtime/Harmonix/Source/HarmonixMetasound/Private/HarmonixMetasound/Nodes/MidiStreamLoggerNode.cpp
@@ -69,8 +69,8 @@ namespace HarmonixMetasound
 			Info.ClassName        = { HarmonixNodeNamespace, TEXT("MidiStreamLogger"), TEXT("")};
 			Info.MajorVersion     = 0;
 			Info.MinorVersion     = 1;
-			Info.DisplayName      = METASOUND_LOCTEXT("MidiStreamLoggerNode_DisplayName", "Midi Stream Logger");
-			Info.Description      = METASOUND_LOCTEXT("MidiStreamLoggerNode_Description", "Receives midi messages and writes them to the log.");
+			Info.DisplayName      = METASOUND_LOCTEXT("MidiStreamLoggerNode_DisplayName", "MIDI Logger");
+			Info.Description      = METASOUND_LOCTEXT("MidiStreamLoggerNode_Description", "Receives MIDI messages and writes them to the log.");
 			Info.Author           = PluginAuthor;
 			Info.PromptIfMissing  = PluginNodeMissingPrompt;
 			Info.DefaultInterface = GetVertexInterface();
