Code cleanup and prep for public release.
- Renamed GetSongRemainingTime -> GetSongRemainingMs
- Changed return type of GetNumSections()
- Moved some UPROPERTIES from public to private so that C++ can't read them. Safe from Blueprint because blueprint gets them through getter functions.
- Made some functions const.

#virtualized

diff --git a/Engine/Plugins/Runtime/Harmonix/Config/BaseHarmonix.ini b/Engine/Plugins/Runtime/Harmonix/Config/BaseHarmonix.ini
--- a/Engine/Plugins/Runtime/Harmonix/Config/BaseHarmonix.ini
+++ b/Engine/Plugins/Runtime/Harmonix/Config/BaseHarmonix.ini
@@ -15,6 +15,7 @@
 +FunctionRedirects=(OldName="HarmonixBlueprintUtil.GetMeasuredAudioLatencyMs",NewName="GetMeasuredUserExperienceAndReactionToAudioRenderOffsetMs")
 +FunctionRedirects=(OldName="HarmonixBlueprintUtil.SetMeasuredVideoLatencyMs",NewName="SetMeasuredVideoToAudioRenderOffsetMs")
 +FunctionRedirects=(OldName="HarmonixBlueprintUtil.GetMeasuredVideoLatencyMs",NewName="GetMeasuredVideoToAudioRenderOffsetMs")
++FunctionRedirects=(OldName="MusicClockComponent.GetSongRemainingTime", NewName="MusicClockComponent.GetSongRemainingMs")
 
 +PropertyRedirects=(OldName="MidiFileData.LastTick",NewName="MidiFileData.LastEventTick")
 
diff --git a/Engine/Plugins/Runtime/Harmonix/Source/HarmonixMetasound/Private/HarmonixMetasound/Components/MusicClockComponent.cpp b/Engine/Plugins/Runtime/Harmonix/Source/HarmonixMetasound/Private/HarmonixMetasound/Components/MusicClockComponent.cpp
--- a/Engine/Plugins/Runtime/Harmonix/Source/HarmonixMetasound/Private/HarmonixMetasound/Components/MusicClockComponent.cpp
+++ b/Engine/Plugins/Runtime/Harmonix/Source/HarmonixMetasound/Private/HarmonixMetasound/Components/MusicClockComponent.cpp
@@ -72,7 +72,7 @@ void UMusicClockComponent::ConnectToWallClockForMidi(UMidiFile* InTempoMap)
 	ConnectToWallClock();
 }
 
-void UMusicClockComponent::UpdateClockIfNeeded()
+void UMusicClockComponent::EnsureClockIsValidForGameFrame() const
 {
 	//	Not for use outside the game thread.
 	if (ensureMsgf(
@@ -90,16 +90,15 @@ void UMusicClockComponent::UpdateClockIfNeeded()
 	//	Run the actual clock update.
 	if (State == EMusicClockState::Running && ClockDriver)
 	{
-		// Make note of what will soon become the "previous position"...	
-		PrevAudioRenderSongPos = CurrentSmoothedAudioRenderSongPos;
-		PrevPlayerExperiencedSongPos = CurrentPlayerExperiencedSongPos;
-		PrevVideoRenderSongPos = CurrentVideoRenderSongPos;
-
-		if (ClockDriver->RefreshCurrentSongPos())
-		{
-			CalcDelta();
-			LastUpdateFrame = GFrameCounter;
-		}
+		// NOTE: This is a little naughty here. Even though this function is const, this next call
+		// to the clock driver is non-const, AND it can reach back into this UMusicClockComponent and
+		// mutate some current state. It is currently the best approach we have. Future refactoring
+		// might eliminate this anomaly. Bottom line is... Callers to "outer functions" that call this
+		// function can/should/need-to call into UMusicClockComponent through a const reference, and
+		// that is reasonable. But sometimes we have to update our internal state before returning
+		// from those functions. All of those state changes happen as a result of this call to the
+		// current ClockDriver.
+		ClockDriver->EnsureClockIsValidForGameFrame();
 	}
 }
 
@@ -144,11 +143,11 @@ void UMusicClockComponent::ConnectToWallClock()
 	ClockDriver = MakeUnique<FWallClockMusicClockDriver>(this, TempoMap);
 }
 
-FMidiSongPos UMusicClockComponent::CalculateSongPosWithOffset(float MsOffset, ECalibratedMusicTimebase Timebase)
+FMidiSongPos UMusicClockComponent::CalculateSongPosWithOffset(float MsOffset, ECalibratedMusicTimebase Timebase) const
 {
 	check(ClockDriver);
 
-	UpdateClockIfNeeded();
+	EnsureClockIsValidForGameFrame();
 
 	FMidiSongPos Result;
 	if (ClockDriver->CalculateSongPosWithOffset(MsOffset, Timebase, Result))
@@ -197,9 +196,8 @@ void UMusicClockComponent::EndPlay(const EEndPlayReason::Type EndPlayReason)
 void UMusicClockComponent::TickComponent(float DeltaTime, enum ELevelTick TickType, FActorComponentTickFunction* ThisTickFunction)
 {
 	Super::TickComponent(DeltaTime, TickType, ThisTickFunction);
-
-	UpdateClockIfNeeded();
-	Broadcast();
+	EnsureClockIsValidForGameFrame();
+	BroadcastSongPosChanges();
 }
 
 void UMusicClockComponent::Start()
@@ -266,9 +264,9 @@ void UMusicClockComponent::Stop()
 	PlayStateEvent.Broadcast(State);
 }
 
-float UMusicClockComponent::GetSecondsIncludingCountIn(ECalibratedMusicTimebase Timebase)
+float UMusicClockComponent::GetSecondsIncludingCountIn(ECalibratedMusicTimebase Timebase) const
 {
-	UpdateClockIfNeeded();
+	EnsureClockIsValidForGameFrame();
 	switch (Timebase)
 	{
 	case ECalibratedMusicTimebase::AudioRenderTime:	return CurrentSmoothedAudioRenderSongPos.SecondsIncludingCountIn;
@@ -278,9 +276,9 @@ float UMusicClockComponent::GetSecondsIncludingCountIn(ECalibratedMusicTimebase
 	}
 }
 
-float UMusicClockComponent::GetSecondsFromBarOne(ECalibratedMusicTimebase Timebase)
+float UMusicClockComponent::GetSecondsFromBarOne(ECalibratedMusicTimebase Timebase) const
 {
-	UpdateClockIfNeeded();
+	EnsureClockIsValidForGameFrame();
 	switch (Timebase)
 	{
 	case ECalibratedMusicTimebase::AudioRenderTime:	return CurrentSmoothedAudioRenderSongPos.SecondsFromBarOne;
@@ -290,9 +288,9 @@ float UMusicClockComponent::GetSecondsFromBarOne(ECalibratedMusicTimebase Timeba
 	}
 }
 
-float UMusicClockComponent::GetBarsIncludingCountIn(ECalibratedMusicTimebase Timebase)
+float UMusicClockComponent::GetBarsIncludingCountIn(ECalibratedMusicTimebase Timebase) const
 {
-	UpdateClockIfNeeded();
+	EnsureClockIsValidForGameFrame();
 	switch (Timebase)
 	{
 	case ECalibratedMusicTimebase::AudioRenderTime:	return CurrentSmoothedAudioRenderSongPos.BarsIncludingCountIn;
@@ -302,9 +300,9 @@ float UMusicClockComponent::GetBarsIncludingCountIn(ECalibratedMusicTimebase Tim
 	}
 }
 
-float UMusicClockComponent::GetBeatsIncludingCountIn(ECalibratedMusicTimebase Timebase)
+float UMusicClockComponent::GetBeatsIncludingCountIn(ECalibratedMusicTimebase Timebase) const
 {
-	UpdateClockIfNeeded();
+	EnsureClockIsValidForGameFrame();
 	switch (Timebase)
 	{
 	case ECalibratedMusicTimebase::AudioRenderTime: return CurrentSmoothedAudioRenderSongPos.BeatsIncludingCountIn;
@@ -314,9 +312,9 @@ float UMusicClockComponent::GetBeatsIncludingCountIn(ECalibratedMusicTimebase Ti
 	}
 }
 
-FMusicTimestamp UMusicClockComponent::GetCurrentTimestamp(ECalibratedMusicTimebase Timebase)
+FMusicTimestamp UMusicClockComponent::GetCurrentTimestamp(ECalibratedMusicTimebase Timebase) const
 {
-	UpdateClockIfNeeded();
+	EnsureClockIsValidForGameFrame();
 	switch (Timebase)
 	{
 	case ECalibratedMusicTimebase::AudioRenderTime:	return CurrentSmoothedAudioRenderSongPos.Timestamp;
@@ -326,63 +324,63 @@ FMusicTimestamp UMusicClockComponent::GetCurrentTimestamp(ECalibratedMusicTimeba
 	}
 }
 
-FString UMusicClockComponent::GetCurrentSectionName(ECalibratedMusicTimebase Timebase /*= ECalibratedMusicTimebase::VideoRenderTime*/)
+FString UMusicClockComponent::GetCurrentSectionName(ECalibratedMusicTimebase Timebase) const
 {
 	const FMidiSongPos& SongPos = GetSongPos(Timebase);
 	return SongPos.CurrentSongSection.Name;
 }
 
-int32 UMusicClockComponent::GetCurrentSectionIndex(ECalibratedMusicTimebase Timebase /*= ECalibratedMusicTimebase::VideoRenderTime*/)
+int32 UMusicClockComponent::GetCurrentSectionIndex(ECalibratedMusicTimebase Timebase) const
 {
 	const FMidiSongPos& SongPos = GetSongPos(Timebase);
 	return GetSongMaps().GetSectionMap().TickToSectionIndex(SongPos.CurrentSongSection.StartTick);
 }
 
-float UMusicClockComponent::GetCurrentSectionStartMs(ECalibratedMusicTimebase Timebase /*= ECalibratedMusicTimebase::VideoRenderTime*/)
+float UMusicClockComponent::GetCurrentSectionStartMs(ECalibratedMusicTimebase Timebase) const
 {
 	const FMidiSongPos& SongPos = GetSongPos(Timebase);
 	return GetSongMaps().TickToMs(SongPos.CurrentSongSection.StartTick);
 }
 
-float UMusicClockComponent::GetCurrentSectionLengthMs(ECalibratedMusicTimebase Timebase /*= ECalibratedMusicTimebase::VideoRenderTime*/)
+float UMusicClockComponent::GetCurrentSectionLengthMs(ECalibratedMusicTimebase Timebase) const
 {
 	const FMidiSongPos& SongPos = GetSongPos(Timebase);
 	return GetSongMaps().TickToMs(SongPos.CurrentSongSection.LengthTicks);
 }
 
-float UMusicClockComponent::GetDistanceFromCurrentBeat(ECalibratedMusicTimebase Timebase /*= ECalibratedMusicTimebase::VideoRenderTime*/)
+float UMusicClockComponent::GetDistanceFromCurrentBeat(ECalibratedMusicTimebase Timebase) const
 {
 	return FMath::Fractional(GetSongPos(Timebase).BeatsIncludingCountIn);
 }
 
-float UMusicClockComponent::GetDistanceToNextBeat(ECalibratedMusicTimebase Timebase /*= ECalibratedMusicTimebase::VideoRenderTime*/)
+float UMusicClockComponent::GetDistanceToNextBeat(ECalibratedMusicTimebase Timebase) const
 {
 	return 1 - GetDistanceFromCurrentBeat(Timebase);
 }
 
-float UMusicClockComponent::GetDistanceToClosestBeat(ECalibratedMusicTimebase Timebase /*= ECalibratedMusicTimebase::VideoRenderTime*/)
+float UMusicClockComponent::GetDistanceToClosestBeat(ECalibratedMusicTimebase Timebase) const
 {
 	return FMath::Min(GetDistanceFromCurrentBeat(Timebase), GetDistanceToNextBeat(Timebase));
 }
 
-float UMusicClockComponent::GetDistanceFromCurrentBar(ECalibratedMusicTimebase Timebase /*= ECalibratedMusicTimebase::VideoRenderTime*/)
+float UMusicClockComponent::GetDistanceFromCurrentBar(ECalibratedMusicTimebase Timebase) const
 {
 	return FMath::Fractional(GetSongPos(Timebase).BarsIncludingCountIn);
 }
 
-float UMusicClockComponent::GetDistanceToNextBar(ECalibratedMusicTimebase Timebase /*= ECalibratedMusicTimebase::VideoRenderTime*/)
+float UMusicClockComponent::GetDistanceToNextBar(ECalibratedMusicTimebase Timebase) const
 {
 	return 1 - GetDistanceFromCurrentBar(Timebase);
 }
 
-float UMusicClockComponent::GetDistanceToClosestBar(ECalibratedMusicTimebase Timebase /*= ECalibratedMusicTimebase::VideoRenderTime*/)
+float UMusicClockComponent::GetDistanceToClosestBar(ECalibratedMusicTimebase Timebase) const
 {
 	return FMath::Min(GetDistanceFromCurrentBar(Timebase), GetDistanceToNextBar(Timebase));
 }
 
-float UMusicClockComponent::GetDeltaBar(ECalibratedMusicTimebase Timebase)
+float UMusicClockComponent::GetDeltaBar(ECalibratedMusicTimebase Timebase) const
 {
-	UpdateClockIfNeeded();
+	EnsureClockIsValidForGameFrame();
 	switch (Timebase)
 	{
 	case ECalibratedMusicTimebase::AudioRenderTime: return AudioRenderDeltaBarF;
@@ -392,9 +390,9 @@ float UMusicClockComponent::GetDeltaBar(ECalibratedMusicTimebase Timebase)
 	}
 }
 
-float UMusicClockComponent::GetDeltaBeat(ECalibratedMusicTimebase Timebase)
+float UMusicClockComponent::GetDeltaBeat(ECalibratedMusicTimebase Timebase) const
 {
-	UpdateClockIfNeeded();
+	EnsureClockIsValidForGameFrame();
 	switch (Timebase)
 	{
 	case ECalibratedMusicTimebase::AudioRenderTime: return AudioRenderDeltaBeatF;
@@ -474,7 +472,7 @@ float UMusicClockComponent::GetSectionEndMsAtMs(float Ms) const
 	return GetSongMaps().GetSectionEndMsAtMs(Ms);
 }
 
-float UMusicClockComponent::GetNumSections() const
+int32 UMusicClockComponent::GetNumSections() const
 {
 	return GetSongMaps().GetSectionMap().GetNumSections();
 }
@@ -494,7 +492,7 @@ int32 UMusicClockComponent::GetSongLengthBars() const
 	return GetSongMaps().GetSongLengthBars();
 }
 
-float UMusicClockComponent::GetSongRemainingTime(ECalibratedMusicTimebase Timebase /*= ECalibratedMusicTimebase::VideoRenderTime*/)
+float UMusicClockComponent::GetSongRemainingMs(ECalibratedMusicTimebase Timebase) const
 {
 	const float SongLengthMs = GetSongMaps().GetSongLengthMs();
 	return SongLengthMs <= 0.f ? 0.f : SongLengthMs - (GetSongPos(Timebase).SecondsIncludingCountIn * 1000.0f);
@@ -506,9 +504,9 @@ const FSongMaps& UMusicClockComponent::GetSongMaps() const
 	return SongMaps ? *SongMaps : DefaultMaps;
 }
 
-const FMidiSongPos& UMusicClockComponent::GetSongPos(ECalibratedMusicTimebase Timebase)
+const FMidiSongPos& UMusicClockComponent::GetSongPos(ECalibratedMusicTimebase Timebase) const
 {
-	UpdateClockIfNeeded();
+	EnsureClockIsValidForGameFrame();
 
 	switch (Timebase)
 	{
@@ -522,24 +520,24 @@ const FMidiSongPos& UMusicClockComponent::GetSongPos(ECalibratedMusicTimebase Ti
 	}
 }
 
-FMidiSongPos UMusicClockComponent::GetCurrentSmoothedAudioRenderSongPos()
+FMidiSongPos UMusicClockComponent::GetCurrentSmoothedAudioRenderSongPos() const
 {
 	return GetSongPos(ECalibratedMusicTimebase::AudioRenderTime);
 }
 
-FMidiSongPos UMusicClockComponent::GetCurrentVideoRenderSongPos()
+FMidiSongPos UMusicClockComponent::GetCurrentVideoRenderSongPos() const
 {
 	return GetSongPos(ECalibratedMusicTimebase::VideoRenderTime);
 }
 
-FMidiSongPos UMusicClockComponent::GetCurrentPlayerExperiencedSongPos()
+FMidiSongPos UMusicClockComponent::GetCurrentPlayerExperiencedSongPos() const
 {
 	return GetSongPos(ECalibratedMusicTimebase::ExperiencedTime);
 }
 
-float UMusicClockComponent::MeasureSpanProgress(const FMusicalTimeSpan& Span, ECalibratedMusicTimebase Timebase)
+float UMusicClockComponent::MeasureSpanProgress(const FMusicalTimeSpan& Span, ECalibratedMusicTimebase Timebase) const
 {
-	UpdateClockIfNeeded();
+	EnsureClockIsValidForGameFrame();
 
 	const FSongMaps* Maps = ClockDriver ? ClockDriver->GetCurrentSongMaps() : &DefaultMaps;
 	if (!Maps)
@@ -550,17 +548,7 @@ float UMusicClockComponent::MeasureSpanProgress(const FMusicalTimeSpan& Span, EC
 	return Span.CalcPositionInSpan(CurrentSmoothedAudioRenderSongPos, *Maps);
 }
 
-void UMusicClockComponent::CalcDelta()
-{
-	AudioRenderDeltaBarF = CurrentSmoothedAudioRenderSongPos.BarsIncludingCountIn - PrevAudioRenderSongPos.BarsIncludingCountIn;
-	AudioRenderDeltaBeatF = CurrentSmoothedAudioRenderSongPos.BeatsIncludingCountIn - PrevAudioRenderSongPos.BeatsIncludingCountIn;
-	PlayerExperienceDeltaBarF = CurrentPlayerExperiencedSongPos.BarsIncludingCountIn - PrevPlayerExperiencedSongPos.BarsIncludingCountIn;
-	PlayerExperienceDeltaBeatF = CurrentPlayerExperiencedSongPos.BeatsIncludingCountIn - PrevPlayerExperiencedSongPos.BeatsIncludingCountIn;
-	VideoRenderDeltaBarF = CurrentVideoRenderSongPos.BarsIncludingCountIn - PrevVideoRenderSongPos.BarsIncludingCountIn;
-	VideoRenderDeltaBeatF = CurrentVideoRenderSongPos.BeatsIncludingCountIn - PrevVideoRenderSongPos.BeatsIncludingCountIn;
-}
-
-void UMusicClockComponent::Broadcast()
+void UMusicClockComponent::BroadcastSongPosChanges()
 {
 	const FMidiSongPos* Basis = &CurrentVideoRenderSongPos;
 	switch (TimebaseForBarAndBeatEvents)
@@ -599,3 +587,25 @@ void UMusicClockComponent::MakeDefaultSongMap()
 	DefaultMaps.GetTempoMap().AddTempoInfoPoint(MidiConstants::BPMToMidiTempo(Tempo), 0);
 	DefaultMaps.GetBarMap().AddTimeSignatureAtBarIncludingCountIn(0, TimeSignatureNum, TimeSignatureDenom);
 }
+
+void FMusicClockDriverBase::EnsureClockIsValidForGameFrame()
+{
+	// Here the ClockDriver reaches back up and mutates its owning UMusicClockComponent
+	// to make sure its current state is appropriate to the current musical time. See 
+	// UMusicClockComponent::EnsureClockIsValidForGameFrame for more details as to why 
+	// this is so.
+	Clock->PrevAudioRenderSongPos = Clock->CurrentSmoothedAudioRenderSongPos;
+	Clock->PrevPlayerExperiencedSongPos = Clock->CurrentPlayerExperiencedSongPos;
+	Clock->PrevVideoRenderSongPos = Clock->CurrentVideoRenderSongPos;
+
+	if (RefreshCurrentSongPos())
+	{
+		Clock->AudioRenderDeltaBarF = Clock->CurrentSmoothedAudioRenderSongPos.BarsIncludingCountIn - Clock->PrevAudioRenderSongPos.BarsIncludingCountIn;
+		Clock->AudioRenderDeltaBeatF = Clock->CurrentSmoothedAudioRenderSongPos.BeatsIncludingCountIn - Clock->PrevAudioRenderSongPos.BeatsIncludingCountIn;
+		Clock->PlayerExperienceDeltaBarF = Clock->CurrentPlayerExperiencedSongPos.BarsIncludingCountIn - Clock->PrevPlayerExperiencedSongPos.BarsIncludingCountIn;
+		Clock->PlayerExperienceDeltaBeatF = Clock->CurrentPlayerExperiencedSongPos.BeatsIncludingCountIn - Clock->PrevPlayerExperiencedSongPos.BeatsIncludingCountIn;
+		Clock->VideoRenderDeltaBarF = Clock->CurrentVideoRenderSongPos.BarsIncludingCountIn - Clock->PrevVideoRenderSongPos.BarsIncludingCountIn;
+		Clock->VideoRenderDeltaBeatF = Clock->CurrentVideoRenderSongPos.BeatsIncludingCountIn - Clock->PrevVideoRenderSongPos.BeatsIncludingCountIn;
+		Clock->LastUpdateFrame = GFrameCounter;
+	}
+}
diff --git a/Engine/Plugins/Runtime/Harmonix/Source/HarmonixMetasound/Private/HarmonixMetasound/Components/MusicTempometerComponent.cpp b/Engine/Plugins/Runtime/Harmonix/Source/HarmonixMetasound/Private/HarmonixMetasound/Components/MusicTempometerComponent.cpp
--- a/Engine/Plugins/Runtime/Harmonix/Source/HarmonixMetasound/Private/HarmonixMetasound/Components/MusicTempometerComponent.cpp
+++ b/Engine/Plugins/Runtime/Harmonix/Source/HarmonixMetasound/Private/HarmonixMetasound/Components/MusicTempometerComponent.cpp
@@ -42,9 +42,9 @@ void UMusicTempometerComponent::UpdateCachedSongPos() const
 	LastFrameCounter = GFrameCounter;
 
 	// Cache or clear various FSongPos
-	if (const UMusicClockComponent* Clock = GetClockNoMutex())
+	if (UMusicClockComponent* Clock = GetMutableClockNoMutex())
 	{
-		SongPos = Clock->CurrentVideoRenderSongPos;
+		SongPos = Clock->GetCurrentVideoRenderSongPos();
 	}
 	else
 	{
diff --git a/Engine/Plugins/Runtime/Harmonix/Source/HarmonixMetasound/Public/HarmonixMetasound/Components/MusicClockComponent.h b/Engine/Plugins/Runtime/Harmonix/Source/HarmonixMetasound/Public/HarmonixMetasound/Components/MusicClockComponent.h
--- a/Engine/Plugins/Runtime/Harmonix/Source/HarmonixMetasound/Public/HarmonixMetasound/Components/MusicClockComponent.h
+++ b/Engine/Plugins/Runtime/Harmonix/Source/HarmonixMetasound/Public/HarmonixMetasound/Components/MusicClockComponent.h
@@ -107,25 +107,25 @@ public:
 	// NOTE: INCLUDES time for count-in and pickup bars.
 	// Note: Not const as it might cause the clock to update from its source.
 	UFUNCTION(BlueprintPure, Category = "MusicClock")
-	float GetSecondsIncludingCountIn(ECalibratedMusicTimebase Timebase = ECalibratedMusicTimebase::VideoRenderTime);
+	float GetSecondsIncludingCountIn(ECalibratedMusicTimebase Timebase = ECalibratedMusicTimebase::VideoRenderTime) const;
 
 	// Time from Bar 1 Beat 1. The classic "start of the song".
 	// NOTE: DOES NOT INCLUDE time for count-in and pickup bars.
 	// Note: Not const as it might cause the clock to update from its source.
 	UFUNCTION(BlueprintPure, Category = "MusicClock")
-	float GetSecondsFromBarOne(ECalibratedMusicTimebase Timebase = ECalibratedMusicTimebase::VideoRenderTime);
+	float GetSecondsFromBarOne(ECalibratedMusicTimebase Timebase = ECalibratedMusicTimebase::VideoRenderTime) const;
 
 	// Returns the fractional total bars from the beginning of the authored music content.
 	// NOTE: INCLUDES time for count-in and pickup bars.
 	// Note: Not const as it might cause the clock to update from its source.
 	UFUNCTION(BlueprintPure, Category = "MusicClock")
-	float GetBarsIncludingCountIn(ECalibratedMusicTimebase Timebase = ECalibratedMusicTimebase::VideoRenderTime);
+	float GetBarsIncludingCountIn(ECalibratedMusicTimebase Timebase = ECalibratedMusicTimebase::VideoRenderTime) const;
 
 	// Returns the fractional total beats from the beginning of the authored music content.
 	// NOTE: INCLUDES time for count-in and pickup bars.
 	// Note: Not const as it might cause the clock to update from its source.
 	UFUNCTION(BlueprintPure, Category = "MusicClock")
-	float GetBeatsIncludingCountIn(ECalibratedMusicTimebase Timebase = ECalibratedMusicTimebase::VideoRenderTime); 
+	float GetBeatsIncludingCountIn(ECalibratedMusicTimebase Timebase = ECalibratedMusicTimebase::VideoRenderTime) const;
 
 	// Returns the "classic" musical timestamp in the form Bar (int) & Beat (float). In this form...
 	//    - Bar 1, Beat 1.0 is the "beginning of the song" AFTER count-in/pickups
@@ -133,17 +133,17 @@ public:
 	//    - While Bar can be positive or negative, Beat is always >= 1.0 and is read as "beat in the bar". Again... '1' based!
 	// Note: Not const as it might cause the clock to update from its source.
 	UFUNCTION(BlueprintPure, Category = "MusicClock")
-	FMusicTimestamp GetCurrentTimestamp(ECalibratedMusicTimebase Timebase = ECalibratedMusicTimebase::VideoRenderTime);
+	FMusicTimestamp GetCurrentTimestamp(ECalibratedMusicTimebase Timebase = ECalibratedMusicTimebase::VideoRenderTime) const;
 
 	/** Returns the name of the section that we're currently in (intro, chorus, outro) */
 	// Note: Not const as it might cause the clock to update from its source.
 	UFUNCTION(BlueprintPure, Category = "MusicClock")
-	FString GetCurrentSectionName(ECalibratedMusicTimebase Timebase = ECalibratedMusicTimebase::VideoRenderTime);
+	FString GetCurrentSectionName(ECalibratedMusicTimebase Timebase = ECalibratedMusicTimebase::VideoRenderTime) const;
 
 	/** Returns the index of the current section for the provided time base. [0, Num-1] */
 	// Note: Not const as it might cause the clock to update from its source.
 	UFUNCTION(BlueprintPure, Category = "MusicClock")
-	int32 GetCurrentSectionIndex(ECalibratedMusicTimebase Timebase = ECalibratedMusicTimebase::VideoRenderTime);
+	int32 GetCurrentSectionIndex(ECalibratedMusicTimebase Timebase = ECalibratedMusicTimebase::VideoRenderTime) const;
 
 	UFUNCTION(BlueprintPure, Category = "MusicClock")
 	const TArray<FSongSection>& GetSongSections() const;
@@ -151,54 +151,59 @@ public:
 	/** Returns the start time of the current section in milliseconds */
 	// Note: Not const as it might cause the clock to update from its source.
 	UFUNCTION(BlueprintPure, Category = "MusicClock")
-	float GetCurrentSectionStartMs(ECalibratedMusicTimebase Timebase = ECalibratedMusicTimebase::VideoRenderTime);
+	float GetCurrentSectionStartMs(ECalibratedMusicTimebase Timebase = ECalibratedMusicTimebase::VideoRenderTime) const;
 
 	/** Returns the length of the current section in milliseconds */
 	// Note: Not const as it might cause the clock to update from its source.
 	UFUNCTION(BlueprintPure, Category = "MusicClock")
-	float GetCurrentSectionLengthMs(ECalibratedMusicTimebase Timebase = ECalibratedMusicTimebase::VideoRenderTime);
+	float GetCurrentSectionLengthMs(ECalibratedMusicTimebase Timebase = ECalibratedMusicTimebase::VideoRenderTime) const;
 
-	/** Gets a value expressed in beats between 0-1 that indicates how much progress we made towards the current beat to the next one */
+	/** Gets a value expressed in beats between 0-1 that indicates how much progress we made in the current beat */
 	// Note: Not const as it might cause the clock to update from its source.
 	UFUNCTION(BlueprintPure, Category = "MusicClock")
-	float GetDistanceFromCurrentBeat(ECalibratedMusicTimebase Timebase = ECalibratedMusicTimebase::VideoRenderTime);
+	float GetDistanceFromCurrentBeat(ECalibratedMusicTimebase Timebase = ECalibratedMusicTimebase::VideoRenderTime) const;
 
 	/** Gets a value expressed in beats between 0-1 that indicates how close we are to the next beat. */
 	// Note: Not const as it might cause the clock to update from its source.
 	UFUNCTION(BlueprintPure, Category = "MusicClock")
-	float GetDistanceToNextBeat(ECalibratedMusicTimebase Timebase = ECalibratedMusicTimebase::VideoRenderTime);
+	float GetDistanceToNextBeat(ECalibratedMusicTimebase Timebase = ECalibratedMusicTimebase::VideoRenderTime) const;
 
 	/** Gets a value expressed in beats between 0-1 that indicates how close we are to the closest beat (current beat or next beat). */
 	// Note: Not const as it might cause the clock to update from its source.
 	UFUNCTION(BlueprintPure, Category = "MusicClock")
-	float GetDistanceToClosestBeat(ECalibratedMusicTimebase Timebase = ECalibratedMusicTimebase::VideoRenderTime);
+	float GetDistanceToClosestBeat(ECalibratedMusicTimebase Timebase = ECalibratedMusicTimebase::VideoRenderTime) const;
 
 	/** Gets a value expressed in bars between 0-1 that indicates how much progress we made towards the current bar to the next one. */
 	// Note: Not const as it might cause the clock to update from its source.
 	UFUNCTION(BlueprintPure, Category = "MusicClock")
-	float GetDistanceFromCurrentBar(ECalibratedMusicTimebase Timebase = ECalibratedMusicTimebase::VideoRenderTime);
+	float GetDistanceFromCurrentBar(ECalibratedMusicTimebase Timebase = ECalibratedMusicTimebase::VideoRenderTime) const;
 
 	/** Gets a value expressed in bars between 0-1 that indicates how close we are to the next bar. */
 	// Note: Not const as it might cause the clock to update from its source.
 	UFUNCTION(BlueprintPure, Category = "MusicClock")
-	float GetDistanceToNextBar(ECalibratedMusicTimebase Timebase = ECalibratedMusicTimebase::VideoRenderTime);
+	float GetDistanceToNextBar(ECalibratedMusicTimebase Timebase = ECalibratedMusicTimebase::VideoRenderTime) const;
 
 	/** Gets value expressed in bars between 0-1 that indicates how close we are to the closest bar (current bar or next bar). */
 	// Note: Not const as it might cause the clock to update from its source.
 	UFUNCTION(BlueprintPure, Category = "MusicClock")
-	float GetDistanceToClosestBar(ECalibratedMusicTimebase Timebase = ECalibratedMusicTimebase::VideoRenderTime);
+	float GetDistanceToClosestBar(ECalibratedMusicTimebase Timebase = ECalibratedMusicTimebase::VideoRenderTime) const;
 
 	// Note: Not const as it might cause the clock to update from its source.
 	UFUNCTION(BlueprintPure, Category = "MusicClock")
-	float GetDeltaBar(ECalibratedMusicTimebase Timebase = ECalibratedMusicTimebase::VideoRenderTime);
+	float GetDeltaBar(ECalibratedMusicTimebase Timebase = ECalibratedMusicTimebase::VideoRenderTime) const;
 
 	// Note: Not const as it might cause the clock to update from its source.
 	UFUNCTION(BlueprintPure, Category = "MusicClock")
-	float GetDeltaBeat(ECalibratedMusicTimebase Timebase = ECalibratedMusicTimebase::VideoRenderTime);
+	float GetDeltaBeat(ECalibratedMusicTimebase Timebase = ECalibratedMusicTimebase::VideoRenderTime) const;
 
 	// Note: Not const as it might cause the clock to update from its source.
 	UFUNCTION(BlueprintPure, Category = "MusicClock")
-	const FMidiSongPos& GetSongPos(ECalibratedMusicTimebase Timebase = ECalibratedMusicTimebase::VideoRenderTime);
+	const FMidiSongPos& GetSongPos(ECalibratedMusicTimebase Timebase = ECalibratedMusicTimebase::VideoRenderTime) const;
+
+	/** Returns the remaining time until the end of the MIDI in milliseconds based on the timestamp corresponding to the passed Timebase */
+	// Note: Not const as it might cause the clock to update from its source.
+	UFUNCTION(BlueprintPure, Category = "MusicClock")
+	float GetSongRemainingMs(ECalibratedMusicTimebase Timebase = ECalibratedMusicTimebase::VideoRenderTime) const;
 
 	UFUNCTION(BlueprintPure, Category = "Count In")
 	float GetCountInSeconds() const;
@@ -237,7 +242,7 @@ public:
 	float GetSectionEndMsAtMs(float Ms) const;
 
 	UFUNCTION(BlueprintPure, Category = "Section")
-	float GetNumSections() const;
+	int32 GetNumSections() const;
 
 	UFUNCTION(BlueprintPure, Category = "Song Data")
 	float GetSongLengthMs() const;
@@ -248,11 +253,6 @@ public:
 	UFUNCTION(BlueprintPure, Category = "Song Data")
 	int32 GetSongLengthBars() const;
 
-	/** Returns the remaining time until the end of the MIDI in milliseconds based on the timestamp corresponding to the passed Timebase */
-	// Note: Not const as it might cause the clock to update from its source.
-	UFUNCTION(BlueprintPure, Category = "MusicClock")
-	float GetSongRemainingTime(ECalibratedMusicTimebase Timebase = ECalibratedMusicTimebase::VideoRenderTime);
-
 	const FSongMaps& GetSongMaps() const;
 
 	UPROPERTY(EditAnywhere, BlueprintReadWrite, Category = "MusicClock")
@@ -276,27 +276,30 @@ public:
 	UPROPERTY(BlueprintAssignable, Category = "MusicClock")
 	FMusicClockDisconnected MusicClockDisconnectedEvent;
 
+private:
+	// Don't let C++ access these directly! They are a blueprint convenience and only work because 
+	// they specify getter functions!
 	UPROPERTY(BlueprintGetter = GetCurrentSmoothedAudioRenderSongPos, Category = "MusicClock")
 	FMidiSongPos CurrentSmoothedAudioRenderSongPos;
-
-	UFUNCTION(BlueprintCallable, BlueprintPure, Category = "MusicClock")
-	FMidiSongPos GetCurrentSmoothedAudioRenderSongPos();
-
 	UPROPERTY(BlueprintGetter = GetCurrentVideoRenderSongPos, Category = "MusicClock")
 	FMidiSongPos CurrentVideoRenderSongPos;
+	UPROPERTY(BlueprintGetter = GetCurrentPlayerExperiencedSongPos, Category = "MusicClock")
+	FMidiSongPos CurrentPlayerExperiencedSongPos;
 
+public:
+	// Getter functions for the Blueprint properties exposed above...
 	UFUNCTION(BlueprintCallable, BlueprintPure, Category = "MusicClock")
-	FMidiSongPos GetCurrentVideoRenderSongPos();
+	FMidiSongPos GetCurrentSmoothedAudioRenderSongPos() const;
 
-	UPROPERTY(BlueprintGetter = GetCurrentPlayerExperiencedSongPos, Category = "MusicClock")
-	FMidiSongPos CurrentPlayerExperiencedSongPos;
+	UFUNCTION(BlueprintCallable, BlueprintPure, Category = "MusicClock")
+	FMidiSongPos GetCurrentVideoRenderSongPos() const;
 
 	UFUNCTION(BlueprintCallable, BlueprintPure, Category = "MusicClock")
-	FMidiSongPos GetCurrentPlayerExperiencedSongPos();
+	FMidiSongPos GetCurrentPlayerExperiencedSongPos() const;
 
 	// Note: Not const as it might cause the clock to update from its source.
 	UFUNCTION(BlueprintPure, Category = "MusicClock")
-	float MeasureSpanProgress(const FMusicalTimeSpan& Span, ECalibratedMusicTimebase Timebase = ECalibratedMusicTimebase::VideoRenderTime);
+	float MeasureSpanProgress(const FMusicalTimeSpan& Span, ECalibratedMusicTimebase Timebase = ECalibratedMusicTimebase::VideoRenderTime) const;
 
 	UFUNCTION(BlueprintCallable, Category = "Audio|MusicClock", meta = (WorldContext = "WorldContextObject", DefaultToSelf = "WorldContextObject"))
 	static UMusicClockComponent* CreateMetasoundDrivenMusicClock(UObject* WorldContextObject, UAudioComponent* InAudioComponent, FName MetasoundOuputPinName = "Midi Clock", bool Start = true);
@@ -311,7 +314,7 @@ public:
 	void ConnectToWallClockForMidi(UMidiFile* InTempoMap);
 
 	// Note: Not const as it might cause the clock to update from its source.
-	FMidiSongPos CalculateSongPosWithOffset(float MsOffset, ECalibratedMusicTimebase Timebase = ECalibratedMusicTimebase::VideoRenderTime);
+	FMidiSongPos CalculateSongPosWithOffset(float MsOffset, ECalibratedMusicTimebase Timebase = ECalibratedMusicTimebase::VideoRenderTime) const;
 
 	const FMidiSongPos& GetRawUnsmoothedAudioRenderPos() const { return RawUnsmoothedAudioRenderPos; }
 
@@ -330,14 +333,14 @@ private:
 	FSongMaps DefaultMaps;
 	FMidiSongPos RawUnsmoothedAudioRenderPos;
 
-	float     AudioRenderDeltaBarF = 0.0f;
-	float     AudioRenderDeltaBeatF = 0.0f;
-	float     PlayerExperienceDeltaBarF = 0.0f;
-	float     PlayerExperienceDeltaBeatF = 0.0f;
-	float     VideoRenderDeltaBarF = 0.0f;
-	float     VideoRenderDeltaBeatF = 0.0f;
-	int32     LastBroadcastBar = -1;
-	int32     LastBroadcastBeat = -1;
+	float AudioRenderDeltaBarF = 0.0f;
+	float AudioRenderDeltaBeatF = 0.0f;
+	float PlayerExperienceDeltaBarF = 0.0f;
+	float PlayerExperienceDeltaBeatF = 0.0f;
+	float VideoRenderDeltaBarF = 0.0f;
+	float VideoRenderDeltaBeatF = 0.0f;
+	int32 LastBroadcastBar = -1;
+	int32 LastBroadcastBeat = -1;
 	FSongSection LastBroadcastSongSection;
 
 	FMidiSongPos PrevAudioRenderSongPos;
@@ -347,14 +350,14 @@ private:
 	uint64    LastUpdateFrame = 0;
 
 	void CreateClockDriver();
-	void CalcDelta();
-	void Broadcast();
+	void BroadcastSongPosChanges();
 	void MakeDefaultSongMap();
 	bool ConnectToMetasound();
 	void ConnectToWallClock();
 
 	// Ensures the clock will be updated once per frame.  Should only get called on the game thread.
-	void UpdateClockIfNeeded();
+	void EnsureClockIsValidForGameFrame() const;
+	//void CalcDelta() const;
 };
 
 struct FMusicClockDriverBase
@@ -366,10 +369,11 @@ public:
 	{}
 	virtual ~FMusicClockDriverBase() = default;
 
+	void EnsureClockIsValidForGameFrame();
+
 	virtual bool CalculateSongPosWithOffset(float MsOffset, ECalibratedMusicTimebase Timebase, FMidiSongPos& OutResult) const = 0;
 
 	virtual void Disconnect() = 0;
-	virtual bool RefreshCurrentSongPos() = 0;
 	virtual void OnStart() = 0;
 	virtual void OnPause() = 0;
 	virtual void OnContinue() = 0;
@@ -378,5 +382,7 @@ public:
 
 protected:
 	UMusicClockComponent* Clock;
+private:
+	virtual bool RefreshCurrentSongPos() = 0;
 };
 
diff --git a/Engine/Plugins/Runtime/Harmonix/Source/HarmonixMetasound/Public/HarmonixMetasound/Components/MusicTempometerComponent.h b/Engine/Plugins/Runtime/Harmonix/Source/HarmonixMetasound/Public/HarmonixMetasound/Components/MusicTempometerComponent.h
--- a/Engine/Plugins/Runtime/Harmonix/Source/HarmonixMetasound/Public/HarmonixMetasound/Components/MusicTempometerComponent.h
+++ b/Engine/Plugins/Runtime/Harmonix/Source/HarmonixMetasound/Public/HarmonixMetasound/Components/MusicTempometerComponent.h
@@ -298,6 +298,15 @@ private:
 		return MusicClock;
 	}
 
+	UMusicClockComponent* GetMutableClockNoMutex() const
+	{
+		if (!IsValid(MusicClock.Get()))
+		{
+			SetOwnerClock();
+		}
+		return MusicClock;
+	}
+
 	void UpdateCachedSongPos() const;
 	void UpdateCachedSongPosIfNeeded() const
 	{
