[Harmonix] Clean up merge node and make it work with Live Update

#tests automated tests,, Patchwork PIE

diff --git a/Engine/Plugins/Runtime/Harmonix/Source/HarmonixMetasound/Private/HarmonixMetasound/Nodes/MidiStreamMergeNode.cpp b/Engine/Plugins/Runtime/Harmonix/Source/HarmonixMetasound/Private/HarmonixMetasound/Nodes/MidiStreamMergeNode.cpp
--- a/Engine/Plugins/Runtime/Harmonix/Source/HarmonixMetasound/Private/HarmonixMetasound/Nodes/MidiStreamMergeNode.cpp
+++ b/Engine/Plugins/Runtime/Harmonix/Source/HarmonixMetasound/Private/HarmonixMetasound/Nodes/MidiStreamMergeNode.cpp
@@ -12,8 +12,6 @@
 #include "HarmonixMetasound/DataTypes/MidiStream.h"
 #include "HarmonixMetasound/DataTypes/MusicTransport.h"
 
-DEFINE_LOG_CATEGORY_STATIC(LogMidiStreamMerge, Log, All);
-
 #define LOCTEXT_NAMESPACE "HarmonixMetaSound"
 
 namespace HarmonixMetasound
@@ -33,8 +31,6 @@ namespace HarmonixMetasound
 
 		virtual void BindInputs(FInputVertexInterfaceData& InVertexData) override;
 		virtual void BindOutputs(FOutputVertexInterfaceData& InVertexData) override;
-		virtual FDataReferenceCollection GetInputs() const override;
-		virtual FDataReferenceCollection GetOutputs() const override;
 
 		void Reset(const FResetParams& ResetParams);
 		
@@ -47,6 +43,8 @@ namespace HarmonixMetasound
 
 		//** OUTPUTS
 		FMidiStreamWriteRef MidiStreamOutPin;
+
+		bool bNeedsClockRefresh { true };
 	};
 
 	class FMidiStreamMergeNode : public FNodeFacade
@@ -135,6 +133,8 @@ namespace HarmonixMetasound
 
 		InVertexData.BindReadVertex(METASOUND_GET_PARAM_NAME(InputMidiStreamA), MidiStreamAInPin);
 		InVertexData.BindReadVertex(METASOUND_GET_PARAM_NAME(InputMidiStreamB), MidiStreamBInPin);
+
+		bNeedsClockRefresh = true;
 	}
 
 	void FMidiStreamMergeOperator::BindOutputs(FOutputVertexInterfaceData& InVertexData)
@@ -142,45 +142,38 @@ namespace HarmonixMetasound
 		using namespace CommonPinNames;
 
 		InVertexData.BindReadVertex(METASOUND_GET_PARAM_NAME(Outputs::MidiStream), MidiStreamOutPin);
-	}
 
-	FDataReferenceCollection FMidiStreamMergeOperator::GetInputs() const
-	{
-		// This should never be called. Bind(...) is called instead. This method
-		// exists as a stop-gap until the API can be deprecated and removed.
-		checkNoEntry();
-		return {};
+		bNeedsClockRefresh = true;
 	}
 
-	FDataReferenceCollection FMidiStreamMergeOperator::GetOutputs() const
+	void FMidiStreamMergeOperator::Reset(const FResetParams&)
 	{
-		// This should never be called. Bind(...) is called instead. This method
-		// exists as a stop-gap until the API can be deprecated and removed.
-		checkNoEntry();
-		return {};
+		bNeedsClockRefresh = true;
 	}
 
-	void FMidiStreamMergeOperator::Reset(const FResetParams& ResetParams)
+	void FMidiStreamMergeOperator::Execute()
 	{
-		MidiStreamOutPin->PrepareBlock();
-		
-		if (MidiStreamAInPin.Get() && MidiStreamAInPin->GetMidiClockSource())
+		// Refresh the clock for the output if inputs may have changed
+		if (bNeedsClockRefresh)
 		{
-			MidiStreamOutPin->SetClockSource(*(MidiStreamAInPin->GetMidiClockSource()));
+			MidiStreamOutPin->ResetClockSource();
+			
+			// If there is a clock on the first stream, use that
+			if (const FMidiClockReadRef* ClockA = MidiStreamAInPin->GetMidiClockSource())
+			{
+				MidiStreamOutPin->SetClockSource(*ClockA);
+			}
+			// If not, and there's a clock on the second stream, use that
+			else if (const FMidiClockReadRef* ClockB = MidiStreamBInPin->GetMidiClockSource())
+			{
+				MidiStreamOutPin->SetClockSource(*ClockB);
+			}
+
+			bNeedsClockRefresh = false;
 		}
-		else if (MidiStreamBInPin.Get() && MidiStreamBInPin->GetMidiClockSource())
-		{
-			MidiStreamOutPin->SetClockSource(*(MidiStreamBInPin->GetMidiClockSource()));
-		}
-	}
-
-	void FMidiStreamMergeOperator::Execute()
-	{
+		
 		MidiStreamOutPin->PrepareBlock();
-
-		TArray<FMidiStreamReadRef> MidiStreams = { MidiStreamAInPin, MidiStreamBInPin };
-
-		MidiStreamOutPin->Copy(MidiStreams);
+		MidiStreamOutPin->Copy({ MidiStreamAInPin, MidiStreamBInPin });
 	}
 }
 
diff --git a/Engine/Plugins/Runtime/Harmonix/Source/HarmonixMetasound/Public/HarmonixMetasound/DataTypes/MidiStream.h b/Engine/Plugins/Runtime/Harmonix/Source/HarmonixMetasound/Public/HarmonixMetasound/DataTypes/MidiStream.h
--- a/Engine/Plugins/Runtime/Harmonix/Source/HarmonixMetasound/Public/HarmonixMetasound/DataTypes/MidiStream.h
+++ b/Engine/Plugins/Runtime/Harmonix/Source/HarmonixMetasound/Public/HarmonixMetasound/DataTypes/MidiStream.h
@@ -94,6 +94,11 @@ namespace HarmonixMetasound
 			MidiClockSource.Emplace(ClockSource);
 		}
 
+		void ResetClockSource()
+		{
+			MidiClockSource.Reset();
+		}
+
 		using MidiSinkFunction = std::function<void()>;
 		void ExecuteBlock(MidiSinkFunction Destination) const
 		{
diff --git a/Engine/Plugins/Runtime/Harmonix/Source/HarmonixMetasoundTests/Private/Nodes/MidiStreamMergeNodeTests.cpp b/Engine/Plugins/Runtime/Harmonix/Source/HarmonixMetasoundTests/Private/Nodes/MidiStreamMergeNodeTests.cpp
--- a/Engine/Plugins/Runtime/Harmonix/Source/HarmonixMetasoundTests/Private/Nodes/MidiStreamMergeNodeTests.cpp
+++ b/Engine/Plugins/Runtime/Harmonix/Source/HarmonixMetasoundTests/Private/Nodes/MidiStreamMergeNodeTests.cpp
@@ -1,8 +1,9 @@
 // Copyright Epic Games, Inc. All Rights Reserved.
 
 #include "NodeTestGraphBuilder.h"
-#include "HarmonixDsp/AudioBuffer.h"
 #include "HarmonixMetasound/Common.h"
+#include "HarmonixMetasound/DataTypes/MidiStream.h"
+
 #include "Misc/AutomationTest.h"
 
 #if WITH_DEV_AUTOMATION_TESTS
@@ -16,7 +17,7 @@ namespace HarmonixMetasoundTests::MidiStreamMergeNode
 
 	IMPLEMENT_SIMPLE_AUTOMATION_TEST(
 		FMidiStreamMergeCreateNodeTest,
-		"Harmonix.Metasound.Nodes.MidiStreamMergeNode.CreateNode",
+		"Harmonix.Metasound.Nodes.MidiStreamMergeNode.Merge",
 		EAutomationTestFlags::EditorContext | EAutomationTestFlags::EngineFilter)
 		bool FMidiStreamMergeCreateNodeTest::RunTest(const FString&)
 	{
@@ -29,13 +30,56 @@ namespace HarmonixMetasoundTests::MidiStreamMergeNode
 			NumSamplesPerBlock);
 		UTEST_TRUE("Graph successfully built", Generator.IsValid());
 
+		// add some messages to the inputs
+		Generator->ApplyToInputValue<FMidiStream>("Midi Stream A", [](FMidiStream& Stream)
+		{
+			const FMidiStreamEvent Event
+			{
+				uint32{ 0 },
+				FMidiMsg::CreateNoteOn(2, 34, 56)
+			};
+			Stream.AddMidiEvent(Event);
+		});
+		Generator->ApplyToInputValue<FMidiStream>("Midi Stream B", [](FMidiStream& Stream)
+		{
+			const FMidiStreamEvent Event
+			{
+				uint32{ 0 },
+				FMidiMsg::CreateNoteOff(1, 23)
+			};
+			Stream.AddMidiEvent(Event);
+		});
+		
 		// execute a block
 		{
-			TAudioBuffer<float> Buffer{ Generator->GetNumChannels(), NumSamplesPerBlock, EAudioBufferCleanupMode::Delete};
-			Generator->OnGenerateAudio(Buffer.GetRawChannelData(0), Buffer.GetNumTotalValidSamples());
+			FAudioBuffer Buffer{ Generator->OperatorSettings };
+			Generator->OnGenerateAudio(Buffer.GetData(), Buffer.Num());
 		}
 
-		// No output to validate, finish test
+		// expect to see those messages at the output
+		TOptional<TDataReadReference<FMidiStream>> OutputStream = Generator->GetOutputReadReference<FMidiStream>("Midi Stream");
+		UTEST_TRUE("Got output ref", OutputStream.IsSet());
+		const TArray<FMidiStreamEvent> Events = (*OutputStream)->GetEventsInBlock();
+		UTEST_EQUAL("Got the right number of events", Events.Num(), 2);
+
+		for (const FMidiStreamEvent& Event : Events)
+		{
+			if (Event.MidiMessage.IsNoteOn())
+			{
+				UTEST_EQUAL("Right channel", Event.MidiMessage.GetStdChannel(), 2);
+				UTEST_EQUAL("Right note number", Event.MidiMessage.GetStdData1(), 34);
+				UTEST_EQUAL("Right velocity", Event.MidiMessage.GetStdData2(), 56);
+			}
+			else if (Event.MidiMessage.IsNoteOff())
+			{
+				UTEST_EQUAL("Right channel", Event.MidiMessage.GetStdChannel(), 1);
+				UTEST_EQUAL("Right note number", Event.MidiMessage.GetStdData1(), 23);
+			}
+			else
+			{
+				UTEST_TRUE("Wrong message", false);
+			}
+		}
 
 		return true;
 	}
