[Harmonix] Clean up MIDI Clock State node

diff --git a/Engine/Plugins/Runtime/Harmonix/Source/HarmonixMetasound/Private/HarmonixMetasound/Nodes/MidiClockStateNode.cpp b/Engine/Plugins/Runtime/Harmonix/Source/HarmonixMetasound/Private/HarmonixMetasound/Nodes/MidiClockStateNode.cpp
--- a/Engine/Plugins/Runtime/Harmonix/Source/HarmonixMetasound/Private/HarmonixMetasound/Nodes/MidiClockStateNode.cpp
+++ b/Engine/Plugins/Runtime/Harmonix/Source/HarmonixMetasound/Private/HarmonixMetasound/Nodes/MidiClockStateNode.cpp
@@ -19,7 +19,6 @@ namespace HarmonixMetasound
 {
 	using namespace Metasound;
 
-
 	class FMidiClockStateOperator : public TExecutableOperator<FMidiClockStateOperator>
 	{
 	public:
@@ -34,19 +33,9 @@ namespace HarmonixMetasound
 		virtual ~FMidiClockStateOperator() override;
 
 		virtual void BindInputs(FInputVertexInterfaceData& InVertexData) override;
-		virtual FDataReferenceCollection GetInputs() const override
-		{
-			checkNoEntry();
-			return {};
-		}
-
+		
 		virtual void BindOutputs(FOutputVertexInterfaceData& InVertexData) override;
-		virtual FDataReferenceCollection GetOutputs() const override
-		{
-			checkNoEntry();
-			return {};
-		}
-
+		
 		void Reset(const FResetParams& ResetParams);
 
 		void Execute();
@@ -92,8 +81,8 @@ namespace HarmonixMetasound
 			Info.ClassName        = { HarmonixNodeNamespace, TEXT("MidiClockState"), TEXT("")};
 			Info.MajorVersion     = 0;
 			Info.MinorVersion     = 1;
-			Info.DisplayName      = METASOUND_LOCTEXT("MidiClockStateNode_DisplayName", "Midi Clock State");
-			Info.Description      = METASOUND_LOCTEXT("MidiClockStateNode_Description", "Prvoides the current tempo, current speed, and current position of the attached midi clock.");
+			Info.DisplayName      = METASOUND_LOCTEXT("MidiClockStateNode_DisplayName", "MIDI Clock State");
+			Info.Description      = METASOUND_LOCTEXT("MidiClockStateNode_Description", "Provides the current tempo, current speed, and current position of the attached MIDI clock.");
 			Info.Author           = PluginAuthor;
 			Info.PromptIfMissing  = PluginNodeMissingPrompt;
 			Info.DefaultInterface = GetVertexInterface();
@@ -108,7 +97,7 @@ namespace HarmonixMetasound
 
 	namespace MidiClockStatePinNames
 	{
-		METASOUND_PARAM(TempoChnagedTriggerOutput, "Tempo-Speed Changed", "Triggers when either the tempo or speed changes.")
+		METASOUND_PARAM(TempoChangedTriggerOutput, "Tempo-Speed Changed", "Triggers when either the tempo or speed changes.")
 	}
 
 	const FVertexInterface& FMidiClockStateOperator::GetVertexInterface()
@@ -118,11 +107,11 @@ namespace HarmonixMetasound
 
 		static const FVertexInterface Interface(
 			FInputVertexInterface(
-				TInputDataVertex<FMidiClock>(METASOUND_GET_PARAM_NAME_AND_METADATA(Inputs::MidiClock)),
-				TInputDataVertex<bool>(METASOUND_GET_PARAM_NAME_AND_METADATA(Inputs::Enable), true)
+				TInputDataVertex<bool>(METASOUND_GET_PARAM_NAME_AND_METADATA(Inputs::Enable), true),
+				TInputDataVertex<FMidiClock>(METASOUND_GET_PARAM_NAME_AND_METADATA(Inputs::MidiClock))
 				),
 			FOutputVertexInterface(
-				TOutputDataVertex<FTrigger>(METASOUND_GET_PARAM_NAME_AND_METADATA(TempoChnagedTriggerOutput)),
+				TOutputDataVertex<FTrigger>(METASOUND_GET_PARAM_NAME_AND_METADATA(TempoChangedTriggerOutput)),
 				TOutputDataVertex<float>(METASOUND_GET_PARAM_NAME_AND_METADATA(Outputs::Tempo)),
 				TOutputDataVertex<float>(METASOUND_GET_PARAM_NAME_AND_METADATA(Outputs::Speed)),
 				TOutputDataVertex<int32>(METASOUND_GET_PARAM_NAME_AND_METADATA(Outputs::TimeSigNumerator)),
@@ -188,7 +177,7 @@ namespace HarmonixMetasound
 		using namespace MidiClockStatePinNames;
 		using namespace CommonPinNames;
 
-		InVertexData.BindReadVertex(TempoChnagedTriggerOutputName, TempoOrSpeedChangedTriggerOutPin);
+		InVertexData.BindReadVertex(TempoChangedTriggerOutputName, TempoOrSpeedChangedTriggerOutPin);
 		InVertexData.BindReadVertex(Outputs::TempoName, CurrentTempoOutPin);
 		InVertexData.BindReadVertex(Outputs::SpeedName, CurrentSpeedOutPin);
 		InVertexData.BindReadVertex(Outputs::MusicTimespanBarName, CurrentBarOutPin);
