Add Reset method to MultiThreadedSynchronizerNode.

I don't _think_ it needs to be more complicated than this?

#tests PIE, EngineTests

diff --git a/Engine/Plugins/Runtime/Harmonix/Source/HarmonixMetasound/Private/HarmonixMetasound/Nodes/MultiThreadedFusionSynchronizerNode.cpp b/Engine/Plugins/Runtime/Harmonix/Source/HarmonixMetasound/Private/HarmonixMetasound/Nodes/MultiThreadedFusionSynchronizerNode.cpp
--- a/Engine/Plugins/Runtime/Harmonix/Source/HarmonixMetasound/Private/HarmonixMetasound/Nodes/MultiThreadedFusionSynchronizerNode.cpp
+++ b/Engine/Plugins/Runtime/Harmonix/Source/HarmonixMetasound/Private/HarmonixMetasound/Nodes/MultiThreadedFusionSynchronizerNode.cpp
@@ -6,6 +6,7 @@
 #include "MetasoundFacade.h"
 #include "MetasoundNodeInterface.h"
 #include "MetasoundNodeRegistrationMacro.h"
+#include "MetasoundOperatorInterface.h"
 #include "MetasoundParamHelper.h"
 #include "MetasoundStandardNodesCategories.h"
 
@@ -267,6 +268,11 @@ namespace HarmonixMetasound
 			}
 		}
 
+		void Reset(const IOperator::FResetParams& Params)
+		{
+			RenderTasks.Empty(NUM_SAMPLERS);
+		}
+		
 		void Execute()
 		{
 			RenderTasks.Empty(NUM_SAMPLERS);
