Fix failing FusionSamplerNode tests.

Make EnableMTRendering a const-constructor pin so it can't change between multiple executes. It doesn't really make sense that it should change after executing, and allowing it to change can confuse the internals of the FusionSamplerNode. (The alternative would have been to just always reset the SyncLinkOutput pin regardless of EnableMTRenderingInPin is true or not.). These issues were triggered when running the automation test since it randomly changes the input values between each execute.

Also, added additional safe guard within the "Execute" of the operator. Make sure that the SyncLink task is complete before running execute, in case the PostExecute function doesn't get called.

#tests PIE,EngineTests

diff --git a/Engine/Plugins/Runtime/Harmonix/Source/HarmonixMetasound/Private/HarmonixMetasound/DataTypes/FusionSyncLink.h b/Engine/Plugins/Runtime/Harmonix/Source/HarmonixMetasound/Private/HarmonixMetasound/DataTypes/FusionSyncLink.h
--- a/Engine/Plugins/Runtime/Harmonix/Source/HarmonixMetasound/Private/HarmonixMetasound/DataTypes/FusionSyncLink.h
+++ b/Engine/Plugins/Runtime/Harmonix/Source/HarmonixMetasound/Private/HarmonixMetasound/DataTypes/FusionSyncLink.h
@@ -54,12 +54,24 @@ namespace HarmonixMetasound
 		{
 			if (!Task.IsCompleted())
 			{
-				UE_LOG(LogFusionAsync, Warning, TEXT("Fusion Async Render: Task created, but was never chedked for completion. This porbably means you are missing a Fusion Synchronizer Node!"));
+				UE_LOG(LogFusionAsync, Warning, TEXT("Fusion Async Render: Task created, but was never checked for completion. This probably means you are missing a Fusion Synchronizer Node!"));
 				Task.Wait();
 			}
 			Task = {};
 		}
 
+		/** This function is called by the FusionSamplerNode "Reset" 
+         * We need to block on reset and make sure the task completes and clear it out
+         */
+        void Reset()
+        {
+			if (!Task.IsCompleted())
+			{
+				Task.Wait();
+			}
+			Task = {};
+        }
+
 	private:
 		UE::Tasks::FTask Task;
 	};
diff --git a/Engine/Plugins/Runtime/Harmonix/Source/HarmonixMetasound/Private/HarmonixMetasound/Nodes/FusionSamplerNode.cpp b/Engine/Plugins/Runtime/Harmonix/Source/HarmonixMetasound/Private/HarmonixMetasound/Nodes/FusionSamplerNode.cpp
--- a/Engine/Plugins/Runtime/Harmonix/Source/HarmonixMetasound/Private/HarmonixMetasound/Nodes/FusionSamplerNode.cpp
+++ b/Engine/Plugins/Runtime/Harmonix/Source/HarmonixMetasound/Private/HarmonixMetasound/Nodes/FusionSamplerNode.cpp
@@ -67,7 +67,7 @@ namespace HarmonixMetasound
 			FFloatReadRef            InFineTuneCents;
 			FFusionPatchAssetReadRef InPatch;
 			FBoolReadRef             InClockSpeedToPitch;
-			FBoolReadRef             InMTRenderingEnable;
+			bool		             InMTRenderingEnable;
 			EAudioBufferChannelLayout InOutputChannelLayout;
 
 			FConstructionArgs(const FBuildOperatorParams& InParams, EAudioBufferChannelLayout OutputChannelLayout)
@@ -84,7 +84,7 @@ namespace HarmonixMetasound
 				, InFineTuneCents(InParams.InputData.GetOrCreateDefaultDataReadReference<float>(METASOUND_GET_PARAM_NAME(FusionSamplerNodePinNames::FineTuneCents), InParams.OperatorSettings))
 				, InPatch(InParams.InputData.GetOrConstructDataReadReference<FFusionPatchAsset>(METASOUND_GET_PARAM_NAME(CommonPinNames::Inputs::SynthPatch)))
 				, InClockSpeedToPitch(InParams.InputData.GetOrCreateDefaultDataReadReference<bool>(METASOUND_GET_PARAM_NAME(CommonPinNames::Inputs::ClockSpeedToPitch), InParams.OperatorSettings))
-				, InMTRenderingEnable(InParams.InputData.GetOrCreateDefaultDataReadReference<bool>(METASOUND_GET_PARAM_NAME(FusionSamplerNodePinNames::EnableMTFusion), InParams.OperatorSettings))
+				, InMTRenderingEnable(InParams.InputData.GetOrCreateDefaultValue<bool>(METASOUND_GET_PARAM_NAME(FusionSamplerNodePinNames::EnableMTFusion), InParams.OperatorSettings))
 				, InOutputChannelLayout(OutputChannelLayout)
 			{
 			}
@@ -106,9 +106,7 @@ namespace HarmonixMetasound
 				TInputDataVertex<float>(METASOUND_GET_PARAM_NAME_AND_METADATA(FusionSamplerNodePinNames::Lfo1DepthOverride), -1.0f),
 				TInputDataVertex<float>(METASOUND_GET_PARAM_NAME_AND_METADATA(FusionSamplerNodePinNames::FineTuneCents), -1200.0f),
 				TInputDataVertex<bool>(METASOUND_GET_PARAM_NAME_AND_METADATA(Inputs::ClockSpeedToPitch), true),
-				TInputDataVertex<bool>(FusionSamplerNodePinNames::EnableMTFusionName, 
-					{FusionSamplerNodePinNames::EnableMTFusionTooltip, FusionSamplerNodePinNames::EnableMTFusionDisplayName, true},
-					false)
+				TInputConstructorVertex<bool>(METASOUND_GET_PARAM_NAME_AND_METADATA_ADVANCED(FusionSamplerNodePinNames::EnableMTFusion), false)
 			);
 
 			return InputInterface;
@@ -134,6 +132,11 @@ namespace HarmonixMetasound
 			Init();
 		}
 
+		virtual ~FFusionSamplerOperatorBase()
+		{
+			SyncLinkOutPin->Reset();
+		}
+
 		void Init()
 		{
 			SetTicksPerQuarterNote(MidiStreamInPin->GetTicksPerQuarterNote());
@@ -155,13 +158,15 @@ namespace HarmonixMetasound
 			InVertexData.BindReadVertex(METASOUND_GET_PARAM_NAME(FusionSamplerNodePinNames::FineTuneCents), FineTuneCentsInPin);
 			InVertexData.BindReadVertex(METASOUND_GET_PARAM_NAME(Inputs::SynthPatch), PatchInPin);
 			InVertexData.BindReadVertex(METASOUND_GET_PARAM_NAME(Inputs::ClockSpeedToPitch), ClockSpeedAffectsPitchInPin);
-			InVertexData.BindReadVertex(METASOUND_GET_PARAM_NAME(FusionSamplerNodePinNames::EnableMTFusion), EnableMTRenderingInPin);
+			InVertexData.SetValue(METASOUND_GET_PARAM_NAME(FusionSamplerNodePinNames::EnableMTFusion), EnableMTRenderingInPin);
 
 			Init();
 		}
 
 		void Reset(const FResetParams& Params)
 		{
+			SyncLinkOutPin->Reset();
+			
 			Prepare(Params.OperatorSettings.GetSampleRate(), OutputChannelLayout, AudioRendering::kFramesPerRenderBuffer, true);
 			BusBuffer.SetNumValidFrames(0);
 			FramesPerBlock = Params.OperatorSettings.GetNumFramesPerBlock();
@@ -199,7 +204,7 @@ namespace HarmonixMetasound
 		FFloatReadRef FineTuneCentsInPin;
 		FFusionPatchAssetReadRef PatchInPin;
 		FBoolReadRef ClockSpeedAffectsPitchInPin;
-		FBoolReadRef EnableMTRenderingInPin;
+		bool EnableMTRenderingInPin;
 
 		//** DATA
 		int32 FramesPerBlock = 0;
@@ -530,6 +535,15 @@ namespace HarmonixMetasound
 	*********************************************************************************************************/
 	void FFusionSamplerOperatorBase::Execute()
 	{
+		if (EnableMTRenderingInPin)
+		{
+			if (!SyncLinkOutPin->GetTask().IsCompleted())
+			{
+				UE_LOG(LogFusionSamplerPlayer, Warning, TEXT("FusionSamplerOperator is running multi-threaded, and the Execute method was called again before the last render was complete!"));
+			}
+			SyncLinkOutPin->Reset();
+		}
+		
 		if (!*EnableInPin)
 		{
 			if (MadeAudioLastFrame)
@@ -553,13 +567,13 @@ namespace HarmonixMetasound
 
 		SetRawTransposition(*TranspositionInPin);
 
-		if (*EnableMTRenderingInPin && !bDisableMultithreadedRender)
+		if (EnableMTRenderingInPin && !bDisableMultithreadedRender)
 		{
 			SyncLinkOutPin->KickAsyncRender([this]()
-				{
-					TRACE_CPUPROFILER_EVENT_SCOPE_STR("FusionAsyncRender");
-					DoRender();
-				});
+			{
+				TRACE_CPUPROFILER_EVENT_SCOPE_STR("FusionAsyncRender");
+				DoRender();
+			});
 		}
 		else
 		{
@@ -686,7 +700,7 @@ namespace HarmonixMetasound
 	}
 	void FFusionSamplerOperatorBase::PostExecute()
 	{
-		if (*EnableMTRenderingInPin)
+		if (EnableMTRenderingInPin)
 		{
 			// This will ensure that our deferred renderer HAS, in fact, completed,
 			// and will clear out the task handle. 
