Fix path issue caused by moving the Harmonix plugin out of restricted folder.

diff --git a/Engine/Plugins/Runtime/Harmonix/Source/HarmonixMetasoundEditor/Private/HarmonixMetasoundSlateStyle.cpp b/Engine/Plugins/Runtime/Harmonix/Source/HarmonixMetasoundEditor/Private/HarmonixMetasoundSlateStyle.cpp
--- a/Engine/Plugins/Runtime/Harmonix/Source/HarmonixMetasoundEditor/Private/HarmonixMetasoundSlateStyle.cpp
+++ b/Engine/Plugins/Runtime/Harmonix/Source/HarmonixMetasoundEditor/Private/HarmonixMetasoundSlateStyle.cpp
@@ -11,7 +11,7 @@ namespace HarmonixMetasoundEditor
 		const FVector2D Icon22x22(22.0f, 22.0f);
 		const FVector2D Icon18x10(18.0f, 10.0f);
 		const FVector2D Icon18x18(18.0f, 18.0f);
-		SetContentRoot(FPaths::EngineDir() / TEXT("Restricted/NotForLicensees/Plugins/Harmonix/Content/Editor"));
+		SetContentRoot(FPaths::EnginePluginsDir() / TEXT("Runtime/Harmonix/Content/Editor"));
 
 #define IMAGE_BRUSH( RelativePath, ... ) FSlateImageBrush( RootToContentDir( RelativePath, TEXT(".png") ), __VA_ARGS__ )
 		Set("HarmonixMetasound.MidiConnectedIcon", new IMAGE_BRUSH(TEXT("Icons/MidiConnectedPin"), Icon22x22));
