EnumArray macro nonsense

Don't expose a property as a reference to an array if you can't properly move that reference when the array changes.

Change access to a "Getter" method so that we're always referencing the array directly, and never risking access to stale references.

(I should just kill this macro nonsense)

#tests PIE

diff --git a/Engine/Plugins/Runtime/Harmonix/Source/HarmonixDsp/Private/HarmonixDsp/FusionSampler/FusionPatch.cpp b/Engine/Plugins/Runtime/Harmonix/Source/HarmonixDsp/Private/HarmonixDsp/FusionSampler/FusionPatch.cpp
--- a/Engine/Plugins/Runtime/Harmonix/Source/HarmonixDsp/Private/HarmonixDsp/FusionSampler/FusionPatch.cpp
+++ b/Engine/Plugins/Runtime/Harmonix/Source/HarmonixDsp/Private/HarmonixDsp/FusionSampler/FusionPatch.cpp
@@ -32,7 +32,7 @@ TSharedPtr<Audio::IProxyData> UFusionPatch::CreateProxyData(const Audio::FProxyD
 UFusionPatch::UFusionPatch() : Super()
 {
 	FFusionPatchSettings& DefaultSettings = FusionPatchData.Settings;
-	FAdsrSettings& VolumeAdsr = DefaultSettings.Adsrs.Volume;
+	FAdsrSettings& VolumeAdsr = DefaultSettings.Adsrs.Volume();
 	VolumeAdsr.Target = EAdsrTarget::Volume;
 	VolumeAdsr.Depth = 1.0f;
 	VolumeAdsr.IsEnabled = true;
diff --git a/Engine/Plugins/Runtime/Harmonix/Source/HarmonixDsp/Private/HarmonixDsp/FusionSampler/FusionSampler.cpp b/Engine/Plugins/Runtime/Harmonix/Source/HarmonixDsp/Private/HarmonixDsp/FusionSampler/FusionSampler.cpp
--- a/Engine/Plugins/Runtime/Harmonix/Source/HarmonixDsp/Private/HarmonixDsp/FusionSampler/FusionSampler.cpp
+++ b/Engine/Plugins/Runtime/Harmonix/Source/HarmonixDsp/Private/HarmonixDsp/FusionSampler/FusionSampler.cpp
@@ -913,17 +913,17 @@ void FFusionSampler::ResetPatchRelatedState()
 		AdsrSettings[Idx].ResetToDefaults();
 	}
 
-	AdsrSettings.Volume.IsEnabled = true;
-	AdsrSettings.Volume.Target = EAdsrTarget::Volume;
-	AdsrSettings.Assignable.IsEnabled = false;
-	AdsrSettings.Assignable.Target = EAdsrTarget::None;
+	AdsrSettings.Volume().IsEnabled = true;
+	AdsrSettings.Volume().Target = EAdsrTarget::Volume;
+	AdsrSettings.Assignable().IsEnabled = false;
+	AdsrSettings.Assignable().Target = EAdsrTarget::None;
 
 	// set up the default Lfo settings
-	LfoSettings.Pan.ResetToDefaults();
-	LfoSettings.Pitch.ResetToDefaults();
+	LfoSettings.Pan().ResetToDefaults();
+	LfoSettings.Pitch().ResetToDefaults();
 
-	LfoSettings.Pan.Target = ELfoTarget::Pan;
-	LfoSettings.Pitch.Target = ELfoTarget::Pitch;
+	LfoSettings.Pan().Target = ELfoTarget::Pan;
+	LfoSettings.Pitch().Target = ELfoTarget::Pitch;
 
 	// set up the default randomizer settings
 	for (uint8 ModulatorIdx = 0; ModulatorIdx < (uint8)EModulatorTarget::Num; ++ModulatorIdx)
@@ -1248,8 +1248,8 @@ void FFusionSampler::PrepareToProcess(uint32 InNumSamples)
 	PortamentoPitchRamper.Ramp();
 	UpdatePitchBendFactor();
 
-	Lfos.Pan.Advance(static_cast<uint32_t>(InNumSamples * Speed));
-	Lfos.Pitch.Advance(static_cast<uint32_t>(InNumSamples * Speed));
+	Lfos.Pan().Advance(static_cast<uint32_t>(InNumSamples * Speed));
+	Lfos.Pitch().Advance(static_cast<uint32_t>(InNumSamples * Speed));
 }
 
 void FFusionSampler::Process(uint32 InSliceIdx, uint32 InSubSliceIdx, TAudioBuffer<float>& Output)
@@ -1437,7 +1437,7 @@ void FFusionSampler::GetController(MidiConstants::EControllerID InController, in
 		Min = 0.005f;
 		Max = 2.000f;
 		Range = Max - Min;
-		ValueP = AdsrSettings.Volume.ReleaseTime;
+		ValueP = AdsrSettings.Volume().ReleaseTime;
 		ValueF = ((ValueP - Min) / Range);
 		Msb = (int8)(ValueF * 127.0f);
 		break;
@@ -1448,7 +1448,7 @@ void FFusionSampler::GetController(MidiConstants::EControllerID InController, in
 		Min = 0.005f;
 		Max = 2.000f;
 		Range = Max - Min;
-		ValueP = AdsrSettings.Volume.AttackTime;
+		ValueP = AdsrSettings.Volume().AttackTime;
 		ValueF = ((ValueP - Min) / Range);
 		Msb = (int8)(ValueF * 127.0f);
 		break;
@@ -1640,14 +1640,14 @@ void FFusionSampler::SetController(MidiConstants::EControllerID InController, fl
 	}
 	case MidiConstants::EControllerID::Release: 
 	{ 
-		AdsrSettings.Volume.ReleaseTime = InValue; 
-		AdsrSettings.Volume.BuildReleaseTable(); 
+		AdsrSettings.Volume().ReleaseTime = InValue; 
+		AdsrSettings.Volume().BuildReleaseTable(); 
 		break; 
 	}
 	case MidiConstants::EControllerID::Attack: 
 	{ 
-		AdsrSettings.Volume.AttackTime = InValue; 
-		AdsrSettings.Volume.BuildAttackTable(); 
+		AdsrSettings.Volume().AttackTime = InValue; 
+		AdsrSettings.Volume().BuildAttackTable(); 
 		break; 
 	}
 	case MidiConstants::EControllerID::PortamentoSwitch: 
@@ -2117,8 +2117,8 @@ void FFusionSampler::UpdateVoiceLfos()
 		// skip voices not assigned to this channel
 		FFusionVoice* Voice = Node->GetValue();
 		check(Voice);
-		Voice->SetupLfo(0, LfoSettings.Pan);
-		Voice->SetupLfo(1, LfoSettings.Pitch);
+		Voice->SetupLfo(0, LfoSettings.Pan());
+		Voice->SetupLfo(1, LfoSettings.Pitch());
 	}
 }
 
diff --git a/Engine/Plugins/Runtime/Harmonix/Source/HarmonixDsp/Private/HarmonixDsp/FusionSampler/FusionVoice.cpp b/Engine/Plugins/Runtime/Harmonix/Source/HarmonixDsp/Private/HarmonixDsp/FusionSampler/FusionVoice.cpp
--- a/Engine/Plugins/Runtime/Harmonix/Source/HarmonixDsp/Private/HarmonixDsp/FusionSampler/FusionVoice.cpp
+++ b/Engine/Plugins/Runtime/Harmonix/Source/HarmonixDsp/Private/HarmonixDsp/FusionSampler/FusionVoice.cpp
@@ -65,8 +65,8 @@ void FFusionVoice::SetSampleRate(double InSampleRate)
 
 	float fSamplesPerSecond = (float)SamplesPerSecond;
 
-	Adsrs.Assignable.Prepare(fSamplesPerSecond);
-	Adsrs.Volume.Prepare(fSamplesPerSecond);
+	Adsrs.Assignable().Prepare(fSamplesPerSecond);
+	Adsrs.Volume().Prepare(fSamplesPerSecond);
 	Lfo1.Prepare(fSamplesPerSecond);
 	Lfo2.Prepare(fSamplesPerSecond);
 
@@ -110,8 +110,8 @@ bool FFusionVoice::AssignIDs(FFusionSampler* InSampler, const FKeyzoneSettings*
 	MySampler = InSampler;
 	PitchShifter = Shifter;
 
-	Adsrs.Volume.UseSettings(&(MySampler->AdsrSettings.Volume));
-	Adsrs.Assignable.UseSettings(&(MySampler->AdsrSettings.Assignable));
+	Adsrs.Volume().UseSettings(&(MySampler->AdsrSettings.Volume()));
+	Adsrs.Assignable().UseSettings(&(MySampler->AdsrSettings.Assignable()));
 	VelocityGain = 0.0f;
 
 	Lfo1.UseSettings(&MySampler->LfoSettings[0]);
@@ -170,22 +170,22 @@ void FFusionVoice::Attack()
 	checkSlow(KeyZone);
 	checkSlow(KeyZone->AudioSample);
 
-	Adsrs.Volume.Attack();
-	Adsrs.Assignable.Attack();
+	Adsrs.Volume().Attack();
+	Adsrs.Assignable().Attack();
 	bWaitingForAttack = false;
 	bHasRenderedAnySamples = false;
 }
 
 void FFusionVoice::Release()
 {
-	Adsrs.Volume.Release();
-	Adsrs.Assignable.Release();
+	Adsrs.Volume().Release();
+	Adsrs.Assignable().Release();
 }
 
 void FFusionVoice::FastRelease()
 {
-	Adsrs.Volume.FastRelease();
-	Adsrs.Assignable.FastRelease();
+	Adsrs.Volume().FastRelease();
+	Adsrs.Assignable().FastRelease();
 }
 
 void FFusionVoice::Kill()
@@ -212,8 +212,8 @@ void FFusionVoice::Kill()
 	}
 	PitchShifter = nullptr;
 
-	Adsrs.Volume.Kill();
-	Adsrs.Assignable.Kill();
+	Adsrs.Volume().Kill();
+	Adsrs.Assignable().Kill();
 	SamplePos = 0.0f;
 	bWaitingForAttack = false;
 	PitchOffsetCents = 0.0f;
@@ -276,8 +276,8 @@ void FFusionVoice::SetSampler(FFusionSampler* InSampler)
 	}
 
 	MySampler = InSampler;
-	Adsrs.Volume.UseSettings(&MySampler->AdsrSettings.Volume, false);
-	Adsrs.Assignable.UseSettings(&MySampler->AdsrSettings.Assignable, false);
+	Adsrs.Volume().UseSettings(&MySampler->AdsrSettings.Volume(), false);
+	Adsrs.Assignable().UseSettings(&MySampler->AdsrSettings.Assignable(), false);
 	Lfo1.UseSettings(&MySampler->LfoSettings[0]);
 	Lfo2.UseSettings(&MySampler->LfoSettings[1]);
 
@@ -550,7 +550,7 @@ bool FFusionVoice::IsWaitingForAttack() const
 
 bool FFusionVoice::IsInUse() const
 {
-	bool IsAdsrIdle = Adsrs.Volume.GetStage() == Harmonix::Dsp::Modulators::EAdsrStage::Idle;
+	bool IsAdsrIdle = Adsrs.Volume().GetStage() == Harmonix::Dsp::Modulators::EAdsrStage::Idle;
 	bool IsAudioActive = MaxAudioLevel > 0.0f;
 	return !IsAdsrIdle || IsAudioActive;
 }
@@ -571,9 +571,9 @@ float FFusionVoice::ComputeOctaveShift()
 		OutOctaveShift += kMaxOctaveSwing * Lfo2.GetValue();
 	}
 
-	if (Adsrs.Assignable.GetSettings()->Target == EAdsrTarget::FilterFreq)
+	if (Adsrs.Assignable().GetSettings()->Target == EAdsrTarget::FilterFreq)
 	{
-		OutOctaveShift += 2 * kMaxOctaveSwing * Adsrs.Assignable.GetValue();
+		OutOctaveShift += 2 * kMaxOctaveSwing * Adsrs.Assignable().GetValue();
 	}
 
 	return OutOctaveShift;
@@ -723,7 +723,7 @@ uint32 FFusionVoice::Process(uint32 InSliceIndex, uint32 InSubsliceIndex, float*
 	// we should loop (if possible) as long as we are not currently releasing the voice
 	// (This seem hinky to me. Often looping samples DO NOT have a nice exit 
 	// portion after the loop end. In these cases the sample will just stop short. - Buzz)
-	bool IsInReleaseStage = (Adsrs.Volume.GetStage() == EAdsrStage::Release);
+	bool IsInReleaseStage = (Adsrs.Volume().GetStage() == EAdsrStage::Release);
 
 	// Forcing this to true so that Adsr release actually has samples to work on!
 	bool ShouldHonorLoopPoints = true; // !IsInReleaseStage; 
@@ -869,12 +869,12 @@ uint32 FFusionVoice::Process(uint32 InSliceIndex, uint32 InSubsliceIndex, float*
 			// write interpolated sample data into the output buffer
 			{
 				//TIME_BLOCK(fusion_voice_process_filter);
-				Adsrs.Volume.Advance(static_cast<uint32> (NumToProcessThisPass * InSpeed));
-				Adsrs.Assignable.Advance(static_cast<uint32> (NumToProcessThisPass * InSpeed));
+				Adsrs.Volume().Advance(static_cast<uint32> (NumToProcessThisPass * InSpeed));
+				Adsrs.Assignable().Advance(static_cast<uint32> (NumToProcessThisPass * InSpeed));
 			}
 
 			WorkingGain.Lerp(PrevGain, Panner.GetCurrentGainMatrix(), (float)NumProcessed / (float)InMaxNumSamples);
-			WorkingGain *= Adsrs.Volume.GetValue();
+			WorkingGain *= Adsrs.Volume().GetValue();
 
 			{
 				SamplePos = ActiveRenderer->Render(
@@ -947,11 +947,11 @@ uint32 FFusionVoice::Process(uint32 InSliceIndex, uint32 InSubsliceIndex, float*
 	bool IsPortamentoEnabled = MySampler->GetIsPortamentoEnabled(); 
 	// could be processing an oscillating filter
 	bool IsAudioLevelLow = MaxAudioLevel < 0.0001; 
-	bool IsOutOfSampleData = Adsrs.Volume.GetStage() == EAdsrStage::Idle || SamplePos >= EndOfSampleData;
+	bool IsOutOfSampleData = Adsrs.Volume().GetStage() == EAdsrStage::Idle || SamplePos >= EndOfSampleData;
 
 	if (IsAudioLevelLow && IsOutOfSampleData)
 	{
-		if (!IsPortamentoEnabled || Adsrs.Volume.GetStage() == EAdsrStage::Idle)
+		if (!IsPortamentoEnabled || Adsrs.Volume().GetStage() == EAdsrStage::Idle)
 		{
 			Kill();
 		}
diff --git a/Engine/Plugins/Runtime/Harmonix/Source/HarmonixDsp/Private/HarmonixDsp/FusionSampler/Settings/FusionPatchSettings.cpp b/Engine/Plugins/Runtime/Harmonix/Source/HarmonixDsp/Private/HarmonixDsp/FusionSampler/Settings/FusionPatchSettings.cpp
--- a/Engine/Plugins/Runtime/Harmonix/Source/HarmonixDsp/Private/HarmonixDsp/FusionSampler/Settings/FusionPatchSettings.cpp
+++ b/Engine/Plugins/Runtime/Harmonix/Source/HarmonixDsp/Private/HarmonixDsp/FusionSampler/Settings/FusionPatchSettings.cpp
@@ -10,6 +10,6 @@ FFusionPatchSettings::FFusionPatchSettings()
 	, MaxVoices(32)
 	, KeyzoneSelectMode(EKeyzoneSelectMode::Layers)
 {
-	Lfos.Pan.Target = ELfoTarget::Pan;
-	Lfos.Pitch.Target = ELfoTarget::Pitch;
+	Lfos.Pan().Target = ELfoTarget::Pan;
+	Lfos.Pitch().Target = ELfoTarget::Pitch;
 }
\ No newline at end of file
diff --git a/Engine/Plugins/Runtime/Harmonix/Source/HarmonixDsp/Public/HarmonixDsp/Containers/EnumArray.h b/Engine/Plugins/Runtime/Harmonix/Source/HarmonixDsp/Public/HarmonixDsp/Containers/EnumArray.h
--- a/Engine/Plugins/Runtime/Harmonix/Source/HarmonixDsp/Public/HarmonixDsp/Containers/EnumArray.h
+++ b/Engine/Plugins/Runtime/Harmonix/Source/HarmonixDsp/Public/HarmonixDsp/Containers/EnumArray.h
@@ -196,7 +196,9 @@ private:																								\
 																										\
 public:					
 
-// Makes a property with the name "Name", for a TEnumArray or a Struct that uses IMPLEMENT_ENUM_ARRAY
-// The property will then be a reference to the array element indexed by the enum value of the same name
+// Makes a method with the name "Name", for a TEnumArray or a Struct that uses IMPLEMENT_ENUM_ARRAY
+// The method will return a reference to the array element indexed by the enum value of the same name
 // assumes "this" implements the ref by index operator (operator[])
-#define ENUM_PROPERTY(Name) FValueType& Name = this->operator[](EEnumType::Name)
\ No newline at end of file
+#define ENUM_PROPERTY(Name)																				\
+	const FValueType& Name() const { return this->operator[](EEnumType::Name); }						\
+		  FValueType& Name()       { return this->operator[](EEnumType::Name); }
\ No newline at end of file
diff --git a/Engine/Plugins/Runtime/Harmonix/Source/HarmonixDsp/Public/HarmonixDsp/FusionSampler/FusionVoice.h b/Engine/Plugins/Runtime/Harmonix/Source/HarmonixDsp/Public/HarmonixDsp/FusionSampler/FusionVoice.h
--- a/Engine/Plugins/Runtime/Harmonix/Source/HarmonixDsp/Public/HarmonixDsp/FusionSampler/FusionVoice.h
+++ b/Engine/Plugins/Runtime/Harmonix/Source/HarmonixDsp/Public/HarmonixDsp/FusionSampler/FusionVoice.h
@@ -54,9 +54,9 @@ public:
 
 	virtual TSharedPtr<IStretcherAndPitchShifter, ESPMode::ThreadSafe> GetPitchShifter() const { return PitchShifter;  }
 
-	virtual Harmonix::Dsp::Modulators::EAdsrStage GetAdsrStage() const { return Adsrs.Volume.GetStage(); }
+	virtual Harmonix::Dsp::Modulators::EAdsrStage GetAdsrStage() const { return Adsrs.Volume().GetStage(); }
 
-	uint32 GetAge() const { return Adsrs.Volume.GetAge(); }
+	uint32 GetAge() const { return Adsrs.Volume().GetAge(); }
 	FMidiVoiceId GetVoiceID() const { return VoiceID; }
 	uint8 GetTriggeredMidiNote() const { return TriggeredMidiNote; }
 
diff --git a/Engine/Plugins/Runtime/Harmonix/Source/HarmonixDspEditor/Private/SettingsJsonImporter.cpp b/Engine/Plugins/Runtime/Harmonix/Source/HarmonixDspEditor/Private/SettingsJsonImporter.cpp
--- a/Engine/Plugins/Runtime/Harmonix/Source/HarmonixDspEditor/Private/SettingsJsonImporter.cpp
+++ b/Engine/Plugins/Runtime/Harmonix/Source/HarmonixDspEditor/Private/SettingsJsonImporter.cpp
@@ -161,7 +161,7 @@ bool FSettingsJsonImporter::TryParseJson(TSharedPtr<FJsonObject> JsonObj, FFusio
 		if (!TryGetObjectField(AdsrObj, "adsr", AdsrValue))
 			continue;
 
-		FAdsrSettings& AdsrSettings = AdsrIdx == 0 ? PatchSettings.Adsrs.Volume : PatchSettings.Adsrs.Assignable;
+		FAdsrSettings& AdsrSettings = AdsrIdx == 0 ? PatchSettings.Adsrs.Volume() : PatchSettings.Adsrs.Assignable();
 		if (!TryParseJson(AdsrValue, AdsrSettings))
 			continue;
 
