Fix automated test of the MidiClock/MusicClockComponent for new ticking options.

diff --git a/Engine/Plugins/Runtime/Harmonix/Source/HarmonixMetasound/Private/HarmonixMetasound/Subsystems/MidiClockUpdateSubsystem.cpp b/Engine/Plugins/Runtime/Harmonix/Source/HarmonixMetasound/Private/HarmonixMetasound/Subsystems/MidiClockUpdateSubsystem.cpp
--- a/Engine/Plugins/Runtime/Harmonix/Source/HarmonixMetasound/Private/HarmonixMetasound/Subsystems/MidiClockUpdateSubsystem.cpp
+++ b/Engine/Plugins/Runtime/Harmonix/Source/HarmonixMetasound/Private/HarmonixMetasound/Subsystems/MidiClockUpdateSubsystem.cpp
@@ -196,3 +196,26 @@ void UMidiClockUpdateSubsystem::UpdateUMusicClockComponents()
 		}
 	}
 }
+
+// Implement a "tick" method that can be used during automated testing so that
+// the test code doesn't need knowledge of how the low-res clocks are being ticked...
+void UMidiClockUpdateSubsystem::TickForTesting()
+{
+	using namespace MidiClockUpdateSubsystem;
+	switch (UpdateMethod)
+	{
+	case MidiClockUpdateSubsystem::EUpdateMethod::EngineTickableObjectAndTickComponent:
+		UpdateFMidiClocks();
+		break;
+	case MidiClockUpdateSubsystem::EUpdateMethod::EngineSubsystemCoreDelegatesOnBeginFrame:
+		UpdateFMidiClocks();
+		UpdateUMusicClockComponents();
+		break;
+	case MidiClockUpdateSubsystem::EUpdateMethod::EngineTickableObject:
+		UpdateFMidiClocks();
+		UpdateUMusicClockComponents();
+		break;
+	default:
+		checkNoEntry();
+	}
+}
diff --git a/Engine/Plugins/Runtime/Harmonix/Source/HarmonixMetasound/Public/HarmonixMetasound/Subsystems/MidiClockUpdateSubsystem.h b/Engine/Plugins/Runtime/Harmonix/Source/HarmonixMetasound/Public/HarmonixMetasound/Subsystems/MidiClockUpdateSubsystem.h
--- a/Engine/Plugins/Runtime/Harmonix/Source/HarmonixMetasound/Public/HarmonixMetasound/Subsystems/MidiClockUpdateSubsystem.h
+++ b/Engine/Plugins/Runtime/Harmonix/Source/HarmonixMetasound/Public/HarmonixMetasound/Subsystems/MidiClockUpdateSubsystem.h
@@ -76,4 +76,9 @@ private:
 
 	FDelegateHandle EngineBeginFrameDelegate;
 	void CoreDelegatesBeginFrame();
+
+public:
+	// Declare a "tick" method that can be used during automated testing so that
+	// the test code doesn't need knowledge of how the low-res clocks are being ticked...
+	void TickForTesting();
 };
diff --git a/Engine/Plugins/Runtime/Harmonix/Source/HarmonixMetasoundTests/Private/Subsystems/MidiClockUpdateSubsystemTests.cpp b/Engine/Plugins/Runtime/Harmonix/Source/HarmonixMetasoundTests/Private/Subsystems/MidiClockUpdateSubsystemTests.cpp
--- a/Engine/Plugins/Runtime/Harmonix/Source/HarmonixMetasoundTests/Private/Subsystems/MidiClockUpdateSubsystemTests.cpp
+++ b/Engine/Plugins/Runtime/Harmonix/Source/HarmonixMetasoundTests/Private/Subsystems/MidiClockUpdateSubsystemTests.cpp
@@ -81,7 +81,7 @@ namespace HarmonixMetasoundTests::MidiClockUpdateSubsystem
 			}
 
 			// Tick the subsystem (low-resolution clocks)
-			Subsystem->Tick(0);
+			Subsystem->TickForTesting();
 
 			// Check that the high- and low-resolution clocks are at the same place
 			UTEST_EQUAL(
