Deprecate unused FusionPatch settings

#tests PIE

diff --git a/Engine/Plugins/Runtime/Harmonix/Source/HarmonixDsp/Private/HarmonixDsp/FusionSampler/FusionPatch.cpp b/Engine/Plugins/Runtime/Harmonix/Source/HarmonixDsp/Private/HarmonixDsp/FusionSampler/FusionPatch.cpp
--- a/Engine/Plugins/Runtime/Harmonix/Source/HarmonixDsp/Private/HarmonixDsp/FusionSampler/FusionPatch.cpp
+++ b/Engine/Plugins/Runtime/Harmonix/Source/HarmonixDsp/Private/HarmonixDsp/FusionSampler/FusionPatch.cpp
@@ -164,19 +164,59 @@ void UFusionPatch::Serialize(FArchive& Ar)
 	Ar.UsingCustomVersion(FFusionPatchCustomVersion::GUID);
 	const int32 Version = Ar.CustomVer(FFusionPatchCustomVersion::GUID);
 
-	if (Ar.IsLoading() && Version < FFusionPatchCustomVersion::PitchShifterNameRedirects)
+	if (Ar.IsLoading())
 	{
-		UE_LOG(LogFusionPatch, Warning, TEXT("Fusion patch was loaded with an outdated version. Asset requires reimport: %s"), *GetPathName());
-	}
+		if (Version < FFusionPatchCustomVersion::PitchShifterNameRedirects)
+		{
+			UE_LOG(LogFusionPatch, Warning, TEXT("Fusion patch was loaded with an outdated version. Asset requires reimport: %s"), *GetPathName());
+		}
 
-	if (Ar.IsLoading() && Version < FFusionPatchCustomVersion::DeprecatedPresets)
-	{
-		if (FusionPatchData.Presets_DEPRECATED.IsValidIndex(0))
+		if (Version < FFusionPatchCustomVersion::DeprecatedPresets)
 		{
-			FusionPatchData.Settings = FusionPatchData.Presets_DEPRECATED[0];
+			if (FusionPatchData.Presets_DEPRECATED.IsValidIndex(0))
+			{
+				FusionPatchData.Settings = FusionPatchData.Presets_DEPRECATED[0];
+			}
+
+			if (FusionPatchData.Presets_DEPRECATED.Num() > 1)
+			{
+				UE_LOG(LogFusionPatch, Warning, TEXT("Fusion patch (%s) has more than one preset, but presets have been deprecated. Only the first (Default) preset will be loaded"), *GetPathName());
+			}
+		}
+
+		if (Version < FFusionPatchCustomVersion::DeprecatedUnusedEffectsSettings)
+		{
+			if (FusionPatchData.Settings.Delay_DEPRECATED.IsEnabled)
+			{
+				UE_LOG(LogFusionPatch, Warning, TEXT("Fusion Patch (%s) has the \"Delay\" effect enabled, but that effect is no longer supported by the Fusion Sampler! "
+					"If you intended to use the Delay effect, rework your metasound to apply the Delay to the output of the FusionSampler node using this patch"),
+					*GetPathName());
+			}
+
+			if (FusionPatchData.Settings.Distortion_DEPRECATED.IsEnabled)
+			{
+				UE_LOG(LogFusionPatch, Warning, TEXT("Fusion Patch (%s) has the \"Distortion\" effect enabled, but that effect is no longer supported by the Fusion Sampler! "
+					"If you intended to use the Distortion effect, rework your metasound to apply the Distortion to the output of the FusionSampler node using this patch"),
+					*GetPathName());
+			}
+
+			if (FusionPatchData.Settings.BitCrusher_DEPRECATED.IsEnabled)
+			{
+				UE_LOG(LogFusionPatch, Warning, TEXT("Fusion Patch (%s) has the \"BitCrusher\" effect enabled, but that effect is no longer supported by the Fusion Sampler! "
+					"If you intended to use the BitCrusher effect, rework your metasound to apply the Distortion to the output of the FusionSampler node using this patch"),
+					*GetPathName());
+			}
+
+			if (FusionPatchData.Settings.Vocoder_DEPRECATED.IsEnabled)
+			{
+				UE_LOG(LogFusionPatch, Warning, TEXT("Fusion Patch (%s) has the \"Vocoder\" effect enabled, but that effect is no longer supported by the Fusion Sampler! "
+					"If you intended to use the Vocoder effect, rework your metasound to apply the Vocoder to the output of the FusionSampler node using this patch"),
+					*GetPathName());	
+			}
 		}
 	}
 
+
 	const UStretcherAndPitchShifterFactoryConfig* FactoryConfig = GetDefault<UStretcherAndPitchShifterFactoryConfig>();
 	for (FKeyzoneSettings& Keyzone : FusionPatchData.Keyzones)
 	{
diff --git a/Engine/Plugins/Runtime/Harmonix/Source/HarmonixDsp/Private/HarmonixDsp/FusionSampler/FusionPatchCustomVersion.h b/Engine/Plugins/Runtime/Harmonix/Source/HarmonixDsp/Private/HarmonixDsp/FusionSampler/FusionPatchCustomVersion.h
--- a/Engine/Plugins/Runtime/Harmonix/Source/HarmonixDsp/Private/HarmonixDsp/FusionSampler/FusionPatchCustomVersion.h
+++ b/Engine/Plugins/Runtime/Harmonix/Source/HarmonixDsp/Private/HarmonixDsp/FusionSampler/FusionPatchCustomVersion.h
@@ -17,6 +17,8 @@ struct HARMONIXDSP_API FFusionPatchCustomVersion
 
 		DeprecatedPresets = 4,
 
+		DeprecatedUnusedEffectsSettings = 5,
+
 		// -----<new versions can be added above this line>-------------------------------------------------
 		VersionPlusOne,
 		LatestVersion = VersionPlusOne - 1
diff --git a/Engine/Plugins/Runtime/Harmonix/Source/HarmonixDsp/Private/HarmonixDsp/FusionSampler/FusionSampler.cpp b/Engine/Plugins/Runtime/Harmonix/Source/HarmonixDsp/Private/HarmonixDsp/FusionSampler/FusionSampler.cpp
--- a/Engine/Plugins/Runtime/Harmonix/Source/HarmonixDsp/Private/HarmonixDsp/FusionSampler/FusionSampler.cpp
+++ b/Engine/Plugins/Runtime/Harmonix/Source/HarmonixDsp/Private/HarmonixDsp/FusionSampler/FusionSampler.cpp
@@ -47,11 +47,6 @@ FFusionSampler::FFusionSampler()
 	, MinPitchBendCents(Defaults::kMinPitchBend)
 	, MaxPitchBendCents(Defaults::kMaxPitchBend)
 	, VoicePool(nullptr)
-	, DelayEnabled(false)
-	, BitCrusherEnabled(false)
-	, DistortionEnabled(false)
-	, VocoderEnabled(false)
-	, VocoderModulatorIndex(0)
 	, CurrentTempoBPM(120.0f)
 	, RawPitchMultiplier(1.0f)
 	, Transposition(0)
@@ -243,9 +238,6 @@ void FFusionSampler::SetSpeed(float InSpeed, bool InMaintainPitch)
 	// this line applies Speed to portamento time
 	SetPortamentoTime(GetPortamentoTime());   
 
-	Delay.SetSpeed(Speed);
-	Delay.SetDelaySeconds(Delay.GetDelaySeconds());
-
 	MaintainPitchWhenSpeedChanges = InMaintainPitch;
 }
 
@@ -265,7 +257,6 @@ void FFusionSampler::SetTempo(float InBPM)
 	}
 	
 	UpdateVoiceLfos();
-	Delay.SetTempo(InBPM * Speed);
 }
 
 void FFusionSampler::SetBeat(float Beat)
@@ -732,80 +723,7 @@ void FFusionSampler::ApplyPatchSettings()
 		SetPortamentoMode(Portamento.Mode);
 		SetPortamentoTime(Portamento.Seconds);
 	}
-	//-------------------------------------
-	// read in delay settings
-	//-------------------------------------
-	{
-		const FDelaySettings& DelaySettings = Settings.Delay;
-		SetDelayEnabled(DelaySettings.IsEnabled);
-		Delay.SetTimeSyncOption(DelaySettings.TimeSyncOption);
-		Delay.SetDelaySeconds(DelaySettings.TimeSeconds);
-		Delay.SetDryGain(DelaySettings.DryGain);
-		Delay.SetWetGain(DelaySettings.WetGain);
-		Delay.SetFeedbackGain(DelaySettings.FeedbackGain);
-		Delay.SetFilterFreq(DelaySettings.EQFreq);
-		Delay.SetFilterType(DelaySettings.EQType);
-		Delay.SetWetFilterEnabled(DelaySettings.EQEnabled);
-		Delay.SetFilterQ(DelaySettings.EQQ);
-		Delay.SetLfoEnabled(DelaySettings.LfoEnabled);
-		Delay.SetLfoTimeSyncOption(DelaySettings.LfoTimeSyncOption);
-		Delay.SetLfoFreq(DelaySettings.LfoRate);
-		Delay.SetLfoDepth(DelaySettings.LfoDepth);
-		Delay.SetStereoSpreadLeft(DelaySettings.PanLeft);
-		Delay.SetStereoSpreadRight(DelaySettings.PanRight);
-		Delay.SetStereoType(DelaySettings.StereoType);
-		// snap settings immediately
-		Delay.SetParamsToTargets();
-	}
 
-	//-------------------------------------
-	// read in bitcrusher settings
-	//-------------------------------------
-	{
-		const FBitCrusherSettings& CrushSettings = Settings.BitCrusher;
-		SetBitCrusherEnabled(CrushSettings.IsEnabled);
-		BitCrusher.Reset();
-		BitCrusher.SetWetGain(CrushSettings.WetGain);
-		BitCrusher.SetCrush(CrushSettings.Crush);
-		BitCrusher.SetSampleHoldFactor(CrushSettings.SampleHoldFactor);
-	}
-
-	//-------------------------------------
-	// read in vocoder settings
-	//-------------------------------------
-	{
-		const FVocoderSettings& VocoderSettings = Settings.Vocoder;
-		VocoderEnabled = VocoderSettings.IsEnabled;
-		Vocoder.SetBandConfig(VocoderSettings.BandConfig);
-		Vocoder.SetSoloing(VocoderSettings.Soloing);
-		Vocoder.SetWet(VocoderSettings.Wet);
-		Vocoder.SetCarrierGain(VocoderSettings.CarrierGain);
-		Vocoder.SetModulatorGain(VocoderSettings.ModulatorGain);
-		Vocoder.SetCarrierThin(VocoderSettings.CarrierThin);
-		Vocoder.SetModulatorThin(VocoderSettings.ModulatorThin);
-		Vocoder.SetAttack(VocoderSettings.Attack);
-		Vocoder.SetRelease(VocoderSettings.Release);
-		Vocoder.SetHighEmphasis(VocoderSettings.HighEmphasis);
-		Vocoder.SetOutputGain(VocoderSettings.OutputGain);
-		Vocoder.SetPerBandSettings(VocoderSettings.Bands);
-	}
-
-	//-------------------------------------
-	// read in distortion settings
-	//-------------------------------------
-	{
-		const FDistortionSettingsV1& DistortionSettings = Settings.Distortion;
-		DistortionEnabled = DistortionSettings.IsEnabled;
-		Distortion.SetSampleRate((int32)GetSamplesPerSecond());
-		Distortion.SetInputGainDb(DistortionSettings.InputGainDb);
-		Distortion.SetOutputGainDb(DistortionSettings.OutputGainDb);
-		Distortion.SetType(DistortionSettings.Type);
-		Distortion.SetOversample(DistortionSettings.Oversample, AudioRendering::kFramesPerRenderBuffer);
-		for (int32 Idx = 0; Idx < FDistortionSettingsV1::kNumFilters; ++Idx)
-		{
-			Distortion.SetupFilter(Idx, DistortionSettings.Filters[Idx]);
-		}
-	}
 }
 
 void FFusionSampler::Prepare(float InSampleRateHz, EAudioBufferChannelLayout InChannelLayout, uint32 InMaxSamples, bool bInAllocateBuffer /*= true*/)
@@ -851,15 +769,6 @@ void FFusionSampler::SetSampleRate(float InSampleRateHz)
 	}
 
 	SetPortamentoTime(GetPortamentoTime());
-
-	if (DelayEnabled)
-	{
-		Delay.Prepare(InSampleRateHz, kScratchBufferChannels, 2500.0f);
-	}
-
-	BitCrusher.Setup(kScratchBufferChannels, InSampleRateHz);
-	Vocoder.Setup(kScratchBufferChannels, InSampleRateHz, AudioRendering::kFramesPerRenderBuffer);
-
 }
 
 void FFusionSampler::ResetInstrumentStateImpl()
@@ -948,10 +857,6 @@ void FFusionSampler::ResetPatchRelatedState()
 	SetPortamentoMode(EPortamentoMode::Legato);
 	SetPortamentoTime(Defaults::kPortamentoTimeMs / 1000.0f);
 
-	SetDelayEnabled(false);
-	SetBitCrusherEnabled(false);
-	DistortionEnabled = false;
-
 	FMemory::Memset(LastStartLayerSelect, -1, sizeof(int8) * MidiConstants::kMaxNumNotes);
 	FMemory::Memset(LastStopLayerSelect, -1, sizeof(int8) * MidiConstants::kMaxNumNotes);
 	FMemory::Memset(LastVelocity, 0, sizeof(int8) * MidiConstants::kMaxNumNotes);
@@ -1273,20 +1178,6 @@ void FFusionSampler::Process(uint32 InSliceIdx, uint32 InSubSliceIdx, TAudioBuff
 	{
 		VoiceUsage.Merge(0.0);
 		Output.FillData(0);
-
-		// don't cut off delay or buffers, even if no voices are running
-		if (DelayEnabled && DelaySilenceCountdown > 0.0f)
-		{
-			Output.SetIsSilent(false);
-			Delay.Process(Output);
-			DelaySilenceCountdown -= GetSecondsPerSample() * (double)Output.GetNumValidFrames();
-		}
-
-		if (BitCrusherEnabled)
-		{
-			BitCrusher.Reset();
-		}
-
 		return;
 	}
 
@@ -1344,38 +1235,10 @@ void FFusionSampler::Process(uint32 InSliceIdx, uint32 InSubSliceIdx, TAudioBuff
 	// mix in any pre-fx children now
 	for (int32 ChildIdx = 0; ChildIdx < FusionPreChildren.Num(); ++ChildIdx)
 	{
-		if (VocoderEnabled && ChildIdx == VocoderModulatorIndex)
-		{
-			continue;
-		}
 		FusionPreChildren[ChildIdx]->Process(InSliceIdx, InSubSliceIdx, ChildMixBuffer);
 		Output.Accumulate(ChildMixBuffer);
 	}
 
-	if (!IsChildRenderer() && FusionPreChildren.Num() > 0 && VocoderEnabled)
-	{
-		FusionPreChildren[VocoderModulatorIndex]->Process(InSliceIdx, InSubSliceIdx, ChildMixBuffer);
-		Vocoder.Process(Output, ChildMixBuffer);
-	}
-
-	// apply effects
-	if (DistortionEnabled)
-	{
-		Distortion.Process(Output);
-	}
-
-	if (BitCrusherEnabled)
-	{
-		BitCrusher.Process(Output);
-	}
-
-	if (DelayEnabled)
-	{
-		Delay.Process(Output);
-		DelaySilenceCountdown = Delay.CalculateSecsToIdle();
-		Output.SetIsSilent(false);
-	}
-
 	// de-activate portamento if we're in legato mode
 	// and there are no voices playing
 
@@ -1511,98 +1374,80 @@ void FFusionSampler::GetController(MidiConstants::EControllerID InController, in
 
 	case MidiConstants::EControllerID::BitCrushWetMix:
 	{
-		Msb = (int8)(BitCrusher.GetWetGain() * 127.0f);
 		break;
 	}
 
 	case MidiConstants::EControllerID::BitCrushLevel:
 	{
-		Msb = (int8)(((float)BitCrusher.GetCrush() / 15.0f) * 127.0f);
 		break;
 	}
 
 	case MidiConstants::EControllerID::BitCrushSampleHold:
 	{
-		Msb = (int8)((((float)BitCrusher.GetSampleHoldFactor() - 1.0f) / 15.0f) * 127.0f);
 		break;
 	}
 
 	case MidiConstants::EControllerID::DelayTime:
 	{
-		Msb = (int8)(gDelayTimeInterp.InverseClamped(Delay.GetDelaySeconds()) * 127.0f);
 		break;
 	}
 
 	case MidiConstants::EControllerID::DelayDryGain:
 	{
-		Msb = (int8)(Delay.GetDryGain() * 127.0f);
 		break;
 	}
 
 	case MidiConstants::EControllerID::DelayWetGain:
 	{
-		Msb = (int8)(Delay.GetWetGain() * 127.0f);
 		break;
 	}
 
 	case MidiConstants::EControllerID::DelayFeedback:
 	{
-		Msb = (int8)(Delay.GetFeedbackGain() * 127.0f);
 		break;
 	}
 	case MidiConstants::EControllerID::DelayEQEnabled:
 	{
-		Msb = Delay.GetWetFilterEnabled() ? 127 : 0;
 		break;
 	}
 	case MidiConstants::EControllerID::DelayEQType:
 	{
-		Msb = (int8)Delay.GetFilterType();
 		break;
 	}
 	case MidiConstants::EControllerID::DelayEQFreq:
 	{
-		Msb = (int8)(Delay.GetFilterFreq() * 127.0f);
 		break;
 	}
 	case MidiConstants::EControllerID::DelayEQQ:
 	{
-		Msb = (int8)(Delay.GetFilterQ() * 127.0f);
 		break;
 	}
 	case MidiConstants::EControllerID::DelayLFOEnabled:
 	{
-		Msb = Delay.GetLfoEnabled() ? 127 : 0;
 		break;
 	}
 	case MidiConstants::EControllerID::DelayLFOBeatSync:
 	{
-		Msb = Delay.GetLfoTimeSyncOption() == ETimeSyncOption::TempoSync ? 127 : 0;
 		break;
 	}
 	case MidiConstants::EControllerID::DelayLFORate:
 	{
-		Msb = (int8)(Delay.GetLfoFreq() * 127.0f);
 		break;
 	}
 	case MidiConstants::EControllerID::DelayLFODepth:
 	{
-		Msb = (int8)(Delay.GetLfoDepth() * 127.0f);
 		break;
 	}
 	case MidiConstants::EControllerID::DelayStereoType:
 	{
-		Msb = (int8)Delay.GetStereoType();
 		break;
 	}
 	case MidiConstants::EControllerID::DelayPanLeft:
 	{
-		Msb = (int8)(Delay.GetStereoSpreadLeft() * 127.0f);
 		break;
 	}
 	case MidiConstants::EControllerID::DelayPanRight:
 	{
-		Msb = (int8)(Delay.GetStereoSpreadRight() * 127.0f);
 	}
 	default:
 		break;
@@ -1706,93 +1551,75 @@ void FFusionSampler::SetController(MidiConstants::EControllerID InController, fl
 		break; 
 	}
 	case MidiConstants::EControllerID::BitCrushWetMix: 
-	{ 
-		BitCrusher.SetWetGain(InValue);
+	{
 		break; 
 	}
 	case MidiConstants::EControllerID::BitCrushLevel: 
-	{ 
-		BitCrusher.SetCrush(InValue); 
+	{
 		break; 
 	}
 	case MidiConstants::EControllerID::BitCrushSampleHold: 
-	{ 
-		BitCrusher.SetSampleHoldFactor((uint16)InValue); 
+	{
 		break; 
 	}
 	case MidiConstants::EControllerID::DelayTime: 
-	{ 
-		Delay.SetDelaySeconds(gDelayTimeInterp.EvalClamped(InValue)); 
+	{
 		break; 
 	}
 	case MidiConstants::EControllerID::DelayDryGain: 
-	{ 
-		Delay.SetDryGain(InValue); 
+	{
 		break; 
 	}
 	case MidiConstants::EControllerID::DelayWetGain: 
-	{ 
-		Delay.SetWetGain(InValue); 
+	{
 		break; 
 	}
 	case MidiConstants::EControllerID::DelayFeedback: 
-	{ 
-		Delay.SetFeedbackGain(InValue); 
+	{
 		break; 
 	}
 	case MidiConstants::EControllerID::DelayEQEnabled: 
-	{ 
-		Delay.SetWetFilterEnabled(InValue > 0); 
+	{
 		break; 
 	}
 	case MidiConstants::EControllerID::DelayEQType: 
-	{ 
-		Delay.SetFilterType((EBiquadFilterType)(InValue * 2.0f)); 
+	{
 		break; 
 	}
 	case MidiConstants::EControllerID::DelayEQFreq: 
-	{ 
-		Delay.SetFilterFreq(20.0f + 17980.0f * InValue); 
+	{
 		break; 
 	}
 	case MidiConstants::EControllerID::DelayEQQ: 
-	{ 
-		Delay.SetFilterQ(0.1f + 29.9f * InValue); 
+	{
 		break; 
 	}
 	case MidiConstants::EControllerID::DelayLFOEnabled: 
-	{ 
-		Delay.SetLfoEnabled(InValue > 0); 
+	{
 		break; 
 	}
 	case MidiConstants::EControllerID::DelayLFOBeatSync: 
-	{ 
-		Delay.SetLfoTimeSyncOption(InValue > 0 ? ETimeSyncOption::TempoSync : ETimeSyncOption::None); 
+	{
 		break; 
 	}
 	case MidiConstants::EControllerID::DelayLFORate: 
-	{ 
-		Delay.SetLfoFreq(InValue * 5.0f); 
+	{
 		break; 
 	}
 	case MidiConstants::EControllerID::DelayLFODepth: 
-	{ 
-		Delay.SetLfoDepth(InValue * 40.0f);
+	{
 		break; 
 	}
 	case MidiConstants::EControllerID::DelayStereoType: 
-	{ 
-		Delay.SetStereoType(EDelayStereoType(InValue * 4.0f));
+	{
 		break; 
 	}
 	case MidiConstants::EControllerID::DelayPanLeft: 
-	{ 
-		Delay.SetStereoSpreadLeft(InValue * 100.0f); 
+	{
 		break; 
 	}
 	case MidiConstants::EControllerID::DelayPanRight: 
 	{
-		Delay.SetStereoSpreadRight(InValue * 100.0f); 
 		break; 
 	}
 	case MidiConstants::EControllerID::TimeStretchEnvelopeOrder:
@@ -2139,30 +1966,6 @@ void FFusionSampler::UpdateVoicesForEnvelopeOrderChange()
 	}
 }
 
-void FFusionSampler::SetDelayEnabled(bool InEnabled)
-{
-	FScopeLock Lock(&GetBusLock());
-	DelayEnabled = InEnabled;
-
-	if (InEnabled)
-	{
-		Delay.Prepare((float)GetSamplesPerSecond(), kScratchBufferChannels, 2500.0f);
-	}
-	else
-	{
-		Delay.Unprepare();
-	}
-}
-
-void FFusionSampler::SetBitCrusherEnabled(bool InEnabled)
-{
-	if (BitCrusherEnabled != InEnabled && InEnabled)
-	{
-		BitCrusher.Reset();
-	}
-	BitCrusherEnabled = InEnabled;
-}
-
 // TODO: Needs to be reimplemented without FHarmonixGeneratorHandle
 //bool FFusionSampler::AddChild(const FHarmonixGeneratorHandle& child, EInstrumentRenderMode InRenderMode)
 //{
diff --git a/Engine/Plugins/Runtime/Harmonix/Source/HarmonixDsp/Public/HarmonixDsp/FusionSampler/FusionSampler.h b/Engine/Plugins/Runtime/Harmonix/Source/HarmonixDsp/Public/HarmonixDsp/FusionSampler/FusionSampler.h
--- a/Engine/Plugins/Runtime/Harmonix/Source/HarmonixDsp/Public/HarmonixDsp/FusionSampler/FusionSampler.h
+++ b/Engine/Plugins/Runtime/Harmonix/Source/HarmonixDsp/Public/HarmonixDsp/FusionSampler/FusionSampler.h
@@ -4,12 +4,7 @@
 #include "HarmonixDsp/Instruments/VirtualInstrument.h"
 
 #include "Containers/List.h"
-
-#include "HarmonixDsp/Effects/BitCrusher.h"
 #include "HarmonixDsp/Effects/BiquadFilter.h"
-#include "HarmonixDsp/Effects/Delay.h"
-#include "HarmonixDsp/Effects/DistortionV1.h"
-#include "HarmonixDsp/Effects/Vocoder.h"
 
 #include "HarmonixDsp/Modulators/Settings/AdsrSettings.h"
 #include "HarmonixDsp/Modulators/Settings/ModulatorSettings.h"
@@ -221,9 +216,6 @@ private:
 	void  SetTrimVolume(float Value);
 	float GetTrimGain() const;
 
-	void SetDelayEnabled(bool enabled);
-	void SetBitCrusherEnabled(bool enabled);
-
 	float GetRampedExpression() const { return ExpressionGainRamper.GetCurrent(); }
 	float GetRampedPitchBend() const { return PitchBendRamper.GetCurrent(); }
 
@@ -336,20 +328,6 @@ private:
 
 	FSharedFusionVoicePoolPtr VoicePool;
 
-	Harmonix::Dsp::Effects::FDelay Delay;
-	bool DelayEnabled = false;
-	double DelaySilenceCountdown = 0.0;
-
-	Harmonix::Dsp::Effects::FBitCrusher BitCrusher;
-	bool BitCrusherEnabled = false;
-
-	Harmonix::Dsp::Effects::FDistortionV1 Distortion;
-	bool DistortionEnabled = false;
-
-	Harmonix::Dsp::Effects::FVocoder Vocoder;
-	bool VocoderEnabled = false;
-	int32 VocoderModulatorIndex = 0;
-
 	// if we get note-ons or note-offs on the main
 	// thread, or if we get multiple requests for
 	// the same note number, then this data helps
diff --git a/Engine/Plugins/Runtime/Harmonix/Source/HarmonixDsp/Public/HarmonixDsp/FusionSampler/Settings/FusionPatchSettings.h b/Engine/Plugins/Runtime/Harmonix/Source/HarmonixDsp/Public/HarmonixDsp/FusionSampler/Settings/FusionPatchSettings.h
--- a/Engine/Plugins/Runtime/Harmonix/Source/HarmonixDsp/Public/HarmonixDsp/FusionSampler/Settings/FusionPatchSettings.h
+++ b/Engine/Plugins/Runtime/Harmonix/Source/HarmonixDsp/Public/HarmonixDsp/FusionSampler/Settings/FusionPatchSettings.h
@@ -68,18 +68,6 @@ public:
 	UPROPERTY(EditDefaultsOnly, Category = "Settings")
 	FBiquadFilterSettings Filter;
 	
-	UPROPERTY(EditDefaultsOnly, Category = "Settings")
-	FDelaySettings Delay;
-
-	UPROPERTY(EditDefaultsOnly, Category = "Settings")
-	FDistortionSettingsV1 Distortion;
-
-	UPROPERTY(EditDefaultsOnly, Category = "Settings")
-	FBitCrusherSettings BitCrusher;
-
-	UPROPERTY(EditDefaultsOnly, Category = "Settings")
-	FVocoderSettings Vocoder;
-
 	UPROPERTY(EditDefaultsOnly, Category = "Settings")
 	FAdsrSettingsArray Adsrs;
 
@@ -91,4 +79,16 @@ public:
 
 	UPROPERTY(EditDefaultsOnly, Category = "Settings")
 	FModulatorSettingsArray VelocityModulators;
+	
+	UPROPERTY()
+	FDelaySettings Delay_DEPRECATED;
+
+	UPROPERTY()
+	FDistortionSettingsV1 Distortion_DEPRECATED;
+
+	UPROPERTY()
+	FBitCrusherSettings BitCrusher_DEPRECATED;
+
+	UPROPERTY()
+	FVocoderSettings Vocoder_DEPRECATED;
 };
\ No newline at end of file
diff --git a/Engine/Plugins/Runtime/Harmonix/Source/HarmonixDspEditor/Private/SettingsJsonImporter.cpp b/Engine/Plugins/Runtime/Harmonix/Source/HarmonixDspEditor/Private/SettingsJsonImporter.cpp
--- a/Engine/Plugins/Runtime/Harmonix/Source/HarmonixDspEditor/Private/SettingsJsonImporter.cpp
+++ b/Engine/Plugins/Runtime/Harmonix/Source/HarmonixDspEditor/Private/SettingsJsonImporter.cpp
@@ -258,54 +258,6 @@ bool FSettingsJsonImporter::TryParseJson(TSharedPtr<FJsonObject> JsonObj, FFusio
 	}
 	}
 
-	{
-	TSharedPtr<FJsonObject> DelayObj;
-	if (TryGetObjectField(JsonObj, "delay_effect", DelayObj))
-	{
-		TryParseJson(DelayObj, PatchSettings.Delay);
-	}
-	else
-	{
-		PatchSettings.Delay.IsEnabled = false;
-	}
-	}
-
-	{
-	TSharedPtr<FJsonObject> BitcrushObj;
-	if (TryGetObjectField(JsonObj, "bitcrush_effect", BitcrushObj))
-	{
-		TryParseJson(BitcrushObj, PatchSettings.BitCrusher);
-	}
-	else
-	{
-		PatchSettings.BitCrusher.IsEnabled = false;
-	}
-	}
-
-	{
-	TSharedPtr<FJsonObject> VocoderObj;
-	if (TryGetObjectField(JsonObj, "vocoder_effect", VocoderObj))
-	{
-		TryParseJson(VocoderObj, PatchSettings.Vocoder);
-	}
-	else
-	{
-		PatchSettings.Vocoder.IsEnabled = false;
-	}
-	}
-
-	{
-	TSharedPtr<FJsonObject> DistortionObj;
-	if (TryGetObjectField(JsonObj, "distortion_effect", DistortionObj))
-	{
-		TryParseJson(DistortionObj, PatchSettings.Distortion);
-	}
-	else
-	{
-		PatchSettings.Distortion.IsEnabled = false;
-	}
-	}
-
 	{
 	TSharedPtr<FJsonObject> PortamentoObj;
 	if (TryGetObjectField(JsonObj, "portamento", PortamentoObj))
@@ -377,138 +329,6 @@ bool FSettingsJsonImporter::TryParseJson(TSharedPtr<FJsonObject> JsonObj, FModul
 	return true;
 }
 
-bool FSettingsJsonImporter::TryParseJson(TSharedPtr<FJsonObject> JsonObj, FDelaySettings& DelaySettings)
-{
-	TryGetBoolField(JsonObj, "enabled", DelaySettings.IsEnabled, false);
-
-	{
-		bool OldBeatSyncSetting;
-		TryGetBoolField(JsonObj, "beat_sync", OldBeatSyncSetting, false);
-		TryGetEnumField(JsonObj, "time_sync", DelaySettings.TimeSyncOption, OldBeatSyncSetting ? ETimeSyncOption::TempoSync : ETimeSyncOption::None);
-	}
-
-	TryGetNumberField(JsonObj, "time", DelaySettings.TimeSeconds, 0.5f);
-	TryGetNumberField(JsonObj, "dry_gain", DelaySettings.DryGain, 1.0f);
-	TryGetNumberField(JsonObj, "wet_gain", DelaySettings.WetGain, 0.5f);
-	TryGetNumberField(JsonObj, "feedback", DelaySettings.FeedbackGain, 0.25f);
-	TryGetBoolField(JsonObj, "delay_eq_enabled", DelaySettings.EQEnabled, false);
-	TryGetEnumField(JsonObj, "delay_eq_type", DelaySettings.EQType, EDelayFilterType::LowPass);
-	TryGetNumberField(JsonObj, "delay_eq_frequency", DelaySettings.EQFreq, 600.0f);
-	TryGetNumberField(JsonObj, "delay_eq_q", DelaySettings.EQQ, 1.0f);
-	TryGetBoolField(JsonObj, "delay_lfo_enabled", DelaySettings.LfoEnabled, false);
-
-	{
-		bool OldBeatSyncSetting;
-		TryGetBoolField(JsonObj, "delay_lfo_beatsync", OldBeatSyncSetting, false);
-		TryGetEnumField(JsonObj, "delay_lfo_time_sync", DelaySettings.LfoTimeSyncOption, OldBeatSyncSetting ? ETimeSyncOption::TempoSync : ETimeSyncOption::None);
-	}
-
-	TryGetNumberField(JsonObj, "delay_lfo_rate", DelaySettings.LfoRate, 0.0f);
-	TryGetNumberField(JsonObj, "delay_lfo_depth", DelaySettings.LfoDepth, 0.0f);
-	TryGetEnumField(JsonObj, "delay_stereo_type", DelaySettings.StereoType, EDelayStereoType::Default);
-	TryGetNumberField(JsonObj, "delay_stereo_left", DelaySettings.PanLeft, 0.0f);
-	TryGetNumberField(JsonObj, "delay_stereo_right", DelaySettings.PanRight, 0.0f);
-	return true;
-}
-
-bool FSettingsJsonImporter::TryParseJson(TSharedPtr<FJsonObject> JsonObj, FBitCrusherSettings& BitCrusherSettings)
-{
-	static TFunction<float(float)> PctToGain = [](float Pct) { return Pct / 100.0f; };
-
-	TryGetBoolField(JsonObj, "enabled", BitCrusherSettings.IsEnabled, false);
-	TryGetNumberField(JsonObj, "wet_percent", BitCrusherSettings.WetGain, PctToGain, 1.0f);
-	TryGetNumberField(JsonObj, "crush_level", BitCrusherSettings.Crush, (uint16)4);
-	TryGetNumberField(JsonObj, "sample_hold", BitCrusherSettings.SampleHoldFactor, (uint16)5);
-	return true;
-}
-
-bool FSettingsJsonImporter::TryParseJson(TSharedPtr<FJsonObject> JsonObj, FVocoderSettings& VocoderSettings)
-{
-	TryGetBoolField(JsonObj, "enabled", VocoderSettings.IsEnabled, false);
-	TryGetNumberField(JsonObj, "modulator_index", VocoderSettings.ModulatorIndex, -1);
-	if (VocoderSettings.ModulatorIndex == -1)
-	{
-		// index name intentionally masked
-		FString DeprecatedIndex = "\x73\x6c\x61\x76\x65\x5f\x69\x6e\x64\x65\x78";
-		TryGetNumberField(JsonObj, DeprecatedIndex, VocoderSettings.ModulatorIndex, 0);
-	}
-
-	TryGetEnumField(JsonObj, "band_config_index", VocoderSettings.BandConfig, EVocoderBandConfig::k64);
-	TryGetNumberField(JsonObj, "carrier_gain", VocoderSettings.CarrierGain, 1.0f);
-	TryGetNumberField(JsonObj, "carrier_thin", VocoderSettings.CarrierThin, 1.0f);
-	TryGetNumberField(JsonObj, "modulator_gain", VocoderSettings.ModulatorGain, 1.0f);
-	TryGetNumberField(JsonObj, "modultar_thin", VocoderSettings.ModulatorThin, 1.0f);
-	TryGetNumberField(JsonObj, "attack", VocoderSettings.Attack, 0.5f);
-	TryGetNumberField(JsonObj, "release", VocoderSettings.Release, 0.5f);
-	TryGetNumberField(JsonObj, "high_emphasis", VocoderSettings.HighEmphasis, 0.0f);
-	TryGetNumberField(JsonObj, "output_gain_db", VocoderSettings.OutputGain, 1.0f);
-	TryGetBoolField(JsonObj, "solo", VocoderSettings.Soloing, false);
-
-	{
-	const TArray<TSharedPtr<FJsonValue>>* BandArray;
-	if (JsonObj->TryGetArrayField(TEXT("band"), BandArray))
-	{
-		VocoderSettings.Bands.SetNum(BandArray->Num());
-
-		for (TSharedPtr<FJsonValue> BandValue : *BandArray)
-		{
-			TSharedPtr<FJsonObject> BandObj;
-			if (!TryGetObject(BandValue, BandObj))
-				continue;
-
-			VocoderSettings.Bands.SetNum(VocoderSettings.Bands.Num() + 1);
-
-			TryGetNumberField(BandObj, "gain", VocoderSettings.Bands.Last().Gain, 1.0f);
-			TryGetBoolField(BandObj, "solo", VocoderSettings.Bands.Last().Solo, false);
-		}
-	}
-	}
-	return true;
-}
-
-bool FSettingsJsonImporter::TryParseJson(TSharedPtr<FJsonObject> JsonObj, FDistortionSettingsV1& DistortionSettings)
-{
-	TryGetBoolField(JsonObj, "enabled", DistortionSettings.IsEnabled, false);
-	TryGetNumberField(JsonObj, "input_gain", DistortionSettings.InputGainDb, 0.0f);
-	TryGetNumberField(JsonObj, "output_gain", DistortionSettings.OutputGainDb, 0.0f);
-	TryGetNumberField(JsonObj, "dc_adjust", DistortionSettings.DCAdjust, 0.0f);
-	TryGetEnumField(JsonObj, "type", DistortionSettings.Type, EDistortionTypeV1::Clean);
-	TryGetBoolField(JsonObj, "oversample", DistortionSettings.Oversample, false);
-
-	{
-	int FilterIdx = -1;
-	for (TSharedPtr<FJsonValue> ArrayValue : IterField(JsonObj, "filter_set"))
-	{
-		++FilterIdx;
-
-		if (!ensureAlways(FilterIdx < FDistortionSettingsV1::kNumFilters))
-			break;
-
-		TSharedPtr<FJsonObject> FilterObj;
-		if (!TryGetObject(ArrayValue, FilterObj))
-			continue;
-
-		TSharedPtr<FJsonObject> FilterValue;
-		if (!TryGetObjectField(FilterObj, "filter", FilterValue))
-			continue;
-
-		TryParseJson(FilterValue, DistortionSettings.Filters[FilterIdx]);
-	}
-	}
-	return true;
-}
-
-
-bool FSettingsJsonImporter::TryParseJson(TSharedPtr<FJsonObject> JsonObj, FDistortionFilterSettings& DistortionFilterSettings)
-{
-	TryParseJson(JsonObj, DistortionFilterSettings.Filter);
-	TryGetBoolField(JsonObj, "filter_pre", DistortionFilterSettings.FilterPreClip);
-	TryGetNumberField(JsonObj, "num_passes", DistortionFilterSettings.NumPasses);
-
-	return true;
-}
-
-
 bool FSettingsJsonImporter::TryParseJson(TSharedPtr<FJsonObject> JsonObj, FPortamentoSettings& PortamentoSettings)
 {
 	TryGetBoolField(JsonObj, "enabled", PortamentoSettings.IsEnabled, false);
diff --git a/Engine/Plugins/Runtime/Harmonix/Source/HarmonixDspEditor/Private/SettingsJsonImporter.h b/Engine/Plugins/Runtime/Harmonix/Source/HarmonixDspEditor/Private/SettingsJsonImporter.h
--- a/Engine/Plugins/Runtime/Harmonix/Source/HarmonixDspEditor/Private/SettingsJsonImporter.h
+++ b/Engine/Plugins/Runtime/Harmonix/Source/HarmonixDspEditor/Private/SettingsJsonImporter.h
@@ -26,15 +26,5 @@ public:
 
 	static bool TryParseJson(TSharedPtr<FJsonObject> JsonObj, struct FModulatorSettings& ModulatorSettings);
 
-	static bool TryParseJson(TSharedPtr<FJsonObject> JsonObj, struct FDelaySettings& DelaySettings);
-
-	static bool TryParseJson(TSharedPtr<FJsonObject> JsonObj, struct FBitCrusherSettings& BitCrusherSettings);
-
-	static bool TryParseJson(TSharedPtr<FJsonObject> JsonObj, struct FVocoderSettings& VocoderSettings);
-
-	static bool TryParseJson(TSharedPtr<FJsonObject> JsonObj, struct FDistortionSettingsV1& DistortionSettings);
-
-	static bool TryParseJson(TSharedPtr<FJsonObject> JsonObj, struct FDistortionFilterSettings& DistortionFilterSettings);
-
 	static bool TryParseJson(TSharedPtr<FJsonObject> JsonObj, struct FPortamentoSettings& PortamentoSettings);
 };
\ No newline at end of file
