Fix MidiStreamFilterNode log spam due to not clearing events
#tests PIE

diff --git a/Engine/Plugins/Runtime/Harmonix/Source/HarmonixMetasound/Private/HarmonixMetasound/Nodes/MidiStreamTrackFilterNode.cpp b/Engine/Plugins/Runtime/Harmonix/Source/HarmonixMetasound/Private/HarmonixMetasound/Nodes/MidiStreamTrackFilterNode.cpp
--- a/Engine/Plugins/Runtime/Harmonix/Source/HarmonixMetasound/Private/HarmonixMetasound/Nodes/MidiStreamTrackFilterNode.cpp
+++ b/Engine/Plugins/Runtime/Harmonix/Source/HarmonixMetasound/Private/HarmonixMetasound/Nodes/MidiStreamTrackFilterNode.cpp
@@ -199,6 +199,8 @@ namespace HarmonixMetasound::Nodes::MidiStreamTrackFilter
 
 	void FOp::Execute()
 	{
+		MidiStreamOutPin->PrepareBlock();
+
 		if (CurrentTrackIndexFilter != *TrackSelectInPin)
 		{
 			CurrentTrackIndexFilter = *TrackSelectInPin;
