[Harmonix] Add new version of MIDI text trigger node which does not include the track filter. We have a separate track filter node already.

#tests automated tests

diff --git a/Engine/Plugins/Runtime/Harmonix/Source/HarmonixMetasound/Private/HarmonixMetasound/Nodes/MidiTextTriggerNode.cpp b/Engine/Plugins/Runtime/Harmonix/Source/HarmonixMetasound/Private/HarmonixMetasound/Nodes/MidiTextTriggerNode.cpp
--- a/Engine/Plugins/Runtime/Harmonix/Source/HarmonixMetasound/Private/HarmonixMetasound/Nodes/MidiTextTriggerNode.cpp
+++ b/Engine/Plugins/Runtime/Harmonix/Source/HarmonixMetasound/Private/HarmonixMetasound/Nodes/MidiTextTriggerNode.cpp
@@ -13,178 +13,127 @@
 #include "HarmonixMetasound/DataTypes/MidiStream.h"
 #include "HarmonixMetasound/DataTypes/MusicTransport.h"
 
-DEFINE_LOG_CATEGORY_STATIC(LogMidiTextTrigger, Log, All);
-
 #define LOCTEXT_NAMESPACE "HarmonixMetaSound_MidiTextTriggerNode"
 
 namespace HarmonixMetasound
 {
 	using namespace Metasound;
 
-	class FMidiTextTriggerOperator : public TExecutableOperator<FMidiTextTriggerOperator>
-	{
-	public:
-		static const FNodeClassMetadata& GetNodeInfo();
-		static const FVertexInterface& GetVertexInterface();
-		static TUniquePtr<IOperator> CreateOperator(const FBuildOperatorParams& InParams, FBuildResults& OutResults);
-
-		FMidiTextTriggerOperator(const FBuildOperatorParams& InParams,
-								 const FBoolReadRef&       InEnabled,
-								 const FMidiStreamReadRef& InMidiStream,
-								 const FInt32ReadRef&      InTrackNumber,
-								 const FStringReadRef&     InText);
-
-		virtual void BindInputs(FInputVertexInterfaceData& InVertexData) override;
-		virtual void BindOutputs(FOutputVertexInterfaceData& InVertexData) override;
-
-		void Reset(const FResetParams& ResetParams);
-		
-		void Execute();
-
-	private:
-		//** INPUTS
-		FBoolReadRef       EnableInPin;
-		FMidiStreamReadRef MidiStreamInPin;
-		FInt32ReadRef      TrackNumberInPin;
-		FStringReadRef     TextInPin;
-
-		//** OUTPUTS
-		FTriggerWriteRef TriggerOutPin;
-	};
-
-	class FMidiTextTriggerNode : public FNodeFacade
-	{
-	public:
-		FMidiTextTriggerNode(const FNodeInitData& InInitData)
-			: FNodeFacade(InInitData.InstanceName, InInitData.InstanceID, TFacadeOperatorClass<FMidiTextTriggerOperator>())
-		{}
-		virtual ~FMidiTextTriggerNode() = default;
-	};
-
-	METASOUND_REGISTER_NODE(FMidiTextTriggerNode)
-
-	const FNodeClassMetadata& FMidiTextTriggerOperator::GetNodeInfo()
-	{
-		auto InitNodeInfo = []() -> FNodeClassMetadata
-		{
-			FNodeClassMetadata Info;
-			Info.ClassName        = { HarmonixNodeNamespace, TEXT("MidiTextTrigger"), TEXT("")};
-			Info.MajorVersion     = 0;
-			Info.MinorVersion     = 1;
-			Info.DisplayName      = METASOUND_LOCTEXT("MidiTextTriggerNode_DisplayName", "Midi Text Trigger");
-			Info.Description      = METASOUND_LOCTEXT("MidiTextTriggerNode_Description", "Receives a midi stream, filters for the desired messages, and outputs triggers.");
-			Info.Author           = PluginAuthor;
-			Info.PromptIfMissing  = PluginNodeMissingPrompt;
-			Info.DefaultInterface = GetVertexInterface();
-			Info.CategoryHierarchy.Emplace(NodeCategories::Music);
-			return Info;
-		};
-
-		static const FNodeClassMetadata Info = InitNodeInfo();
-
-		return Info;
-	}
-
 	namespace MidiTextTriggerPinNames
 	{
 		METASOUND_PARAM(TextInput, "Text", "String of characters to look for.")
 		METASOUND_PARAM(TriggerOutput, "Trigger Out", "A trigger when the text is encountered.")
 	}
 
-	const FVertexInterface& FMidiTextTriggerOperator::GetVertexInterface()
+	class FMidiTextTriggerOperator_V1 : public TExecutableOperator<FMidiTextTriggerOperator_V1>
 	{
-		using namespace MidiTextTriggerPinNames;
-		using namespace CommonPinNames;
-
-		static const FVertexInterface Interface(
-			FInputVertexInterface(
-				TInputDataVertex<bool>(METASOUND_GET_PARAM_NAME_AND_METADATA(Inputs::Enable), true),
-				TInputDataVertex<FMidiStream>(METASOUND_GET_PARAM_NAME_AND_METADATA(Inputs::MidiStream)),
-				TInputDataVertex<int32>(METASOUND_GET_PARAM_NAME_AND_METADATA(Inputs::MidiTrackNumber),1),
-				TInputDataVertex<FString>(METASOUND_GET_PARAM_NAME_AND_METADATA(TextInput))
-			),
-			FOutputVertexInterface(
-				TOutputDataVertex<FTrigger>(METASOUND_GET_PARAM_NAME_AND_METADATA(TriggerOutput))
-				)
-			);
-
-		return Interface;
-	}
-
-	TUniquePtr<IOperator> FMidiTextTriggerOperator::CreateOperator(const FBuildOperatorParams& InParams, FBuildResults& OutResults)
-	{
-		using namespace MidiTextTriggerPinNames;
-		using namespace CommonPinNames;
+	public:
+		static const FNodeClassMetadata& GetNodeInfo()
+		{
+			auto InitNodeInfo = []() -> FNodeClassMetadata
+			{
+				FNodeClassMetadata Info;
+				Info.ClassName        = { HarmonixNodeNamespace, TEXT("MidiTextTrigger"), TEXT("")};
+				Info.MajorVersion     = 1;
+				Info.MinorVersion     = 0;
+				Info.DisplayName      = METASOUND_LOCTEXT("MidiTextTriggerNode_DisplayName", "MIDI Text Trigger");
+				Info.Description      = METASOUND_LOCTEXT("MidiTextTriggerNode_Description", "Receives a MIDI stream, filters for the desired messages, and outputs triggers.");
+				Info.Author           = PluginAuthor;
+				Info.PromptIfMissing  = PluginNodeMissingPrompt;
+				Info.DefaultInterface = GetVertexInterface();
+				Info.CategoryHierarchy.Emplace(NodeCategories::Music);
+				return Info;
+			};
+
+			static const FNodeClassMetadata Info = InitNodeInfo();
 
-		const FMidiTextTriggerNode& LoggerNode = static_cast<const FMidiTextTriggerNode&>(InParams.Node);
+			return Info;
+		}
+		
+		static const FVertexInterface& GetVertexInterface()
+		{
+			using namespace MidiTextTriggerPinNames;
+			using namespace CommonPinNames;
+
+			static const FVertexInterface Interface(
+				FInputVertexInterface(
+					TInputDataVertex<bool>(METASOUND_GET_PARAM_NAME_AND_METADATA(Inputs::Enable), true),
+					TInputDataVertex<FMidiStream>(METASOUND_GET_PARAM_NAME_AND_METADATA(Inputs::MidiStream)),
+					TInputDataVertex<FString>(METASOUND_GET_PARAM_NAME_AND_METADATA(TextInput))
+				),
+				FOutputVertexInterface(
+					TOutputDataVertex<FTrigger>(METASOUND_GET_PARAM_NAME_AND_METADATA(TriggerOutput))
+					)
+				);
+
+			return Interface;
+		}
+		
+		static TUniquePtr<IOperator> CreateOperator(const FBuildOperatorParams& InParams, FBuildResults& OutResults)
+		{
+			using namespace MidiTextTriggerPinNames;
+			using namespace CommonPinNames;
 
-		const FInputVertexInterfaceData& InputData = InParams.InputData;
-		FBoolReadRef InEnabled          = InputData.GetOrCreateDefaultDataReadReference<bool>(METASOUND_GET_PARAM_NAME(Inputs::Enable), InParams.OperatorSettings);
-		FMidiStreamReadRef InMidiStream = InputData.GetOrConstructDataReadReference<FMidiStream>(METASOUND_GET_PARAM_NAME(Inputs::MidiStream), InParams.OperatorSettings);
-		FInt32ReadRef InTrackNumber     = InputData.GetOrCreateDefaultDataReadReference<int32>(METASOUND_GET_PARAM_NAME(Inputs::MidiTrackNumber), InParams.OperatorSettings);
-		FStringReadRef InText           = InputData.GetOrCreateDefaultDataReadReference<FString>(METASOUND_GET_PARAM_NAME(TextInput), InParams.OperatorSettings);
+			const FInputVertexInterfaceData& InputData = InParams.InputData;
+			FBoolReadRef InEnabled          = InputData.GetOrCreateDefaultDataReadReference<bool>(METASOUND_GET_PARAM_NAME(Inputs::Enable), InParams.OperatorSettings);
+			FMidiStreamReadRef InMidiStream = InputData.GetOrConstructDataReadReference<FMidiStream>(METASOUND_GET_PARAM_NAME(Inputs::MidiStream), InParams.OperatorSettings);
+			FStringReadRef InText           = InputData.GetOrCreateDefaultDataReadReference<FString>(METASOUND_GET_PARAM_NAME(TextInput), InParams.OperatorSettings);
 
-		return MakeUnique<FMidiTextTriggerOperator>(InParams, InEnabled, InMidiStream, InTrackNumber, InText);
-	}
+			return MakeUnique<FMidiTextTriggerOperator_V1>(InParams, InEnabled, InMidiStream, InText);
+		}
 
-	FMidiTextTriggerOperator::FMidiTextTriggerOperator(const FBuildOperatorParams& InParams,
-													   const FBoolReadRef&       InEnabled,
-													   const FMidiStreamReadRef& InMidiStream,
-													   const FInt32ReadRef&      InTrackNumber,
-													   const FStringReadRef&     InText)
+		FMidiTextTriggerOperator_V1(
+			const FBuildOperatorParams& InParams,
+			const FBoolReadRef& InEnabled,
+			const FMidiStreamReadRef& InMidiStream,
+			const FStringReadRef& InText)
 		: EnableInPin(InEnabled)
 		, MidiStreamInPin(InMidiStream)
-		, TrackNumberInPin(InTrackNumber)
 		, TextInPin(InText)
 		, TriggerOutPin(FTriggerWriteRef::CreateNew(InParams.OperatorSettings))
-	{
-		Reset(InParams);
-	}
-
-	void FMidiTextTriggerOperator::BindInputs(FInputVertexInterfaceData& InVertexData)
-	{
-		using namespace MidiTextTriggerPinNames;
-		using namespace CommonPinNames;
-		InVertexData.BindReadVertex(METASOUND_GET_PARAM_NAME(Inputs::Enable), EnableInPin);
-		InVertexData.BindReadVertex(METASOUND_GET_PARAM_NAME(Inputs::MidiStream), MidiStreamInPin);
-		InVertexData.BindReadVertex(METASOUND_GET_PARAM_NAME(Inputs::MidiTrackNumber), TrackNumberInPin);
-		InVertexData.BindReadVertex(METASOUND_GET_PARAM_NAME(TextInput), TextInPin);
-	}
-
-	void FMidiTextTriggerOperator::BindOutputs(FOutputVertexInterfaceData& InVertexData)
-	{
-		using namespace MidiTextTriggerPinNames;
-		using namespace CommonPinNames;
-		InVertexData.BindReadVertex(METASOUND_GET_PARAM_NAME(TriggerOutput), TriggerOutPin);
-	}
-
-	void FMidiTextTriggerOperator::Reset(const FResetParams& ResetParams)
-	{
-		TriggerOutPin->Reset();
-	}
-
-	void FMidiTextTriggerOperator::Execute()
-	{
-		TriggerOutPin->AdvanceBlock();
+		{
+			Reset(InParams);
+		}
 
-		if (!*EnableInPin)
+		virtual void BindInputs(FInputVertexInterfaceData& InVertexData) override
 		{
-			return;
+			using namespace MidiTextTriggerPinNames;
+			using namespace CommonPinNames;
+			InVertexData.BindReadVertex(METASOUND_GET_PARAM_NAME(Inputs::Enable), EnableInPin);
+			InVertexData.BindReadVertex(METASOUND_GET_PARAM_NAME(Inputs::MidiStream), MidiStreamInPin);
+			InVertexData.BindReadVertex(METASOUND_GET_PARAM_NAME(TextInput), TextInPin);
+		}
+		
+		virtual void BindOutputs(FOutputVertexInterfaceData& InVertexData) override
+		{
+			using namespace MidiTextTriggerPinNames;
+			using namespace CommonPinNames;
+			InVertexData.BindReadVertex(METASOUND_GET_PARAM_NAME(TriggerOutput), TriggerOutPin);
 		}
 
-		int32 TrackIndex = *TrackNumberInPin;
-		if (TrackIndex >= 0)
+		void Reset(const FResetParams&)
 		{
+			TriggerOutPin->Reset();
+		}
+		
+		virtual void Execute()
+		{
+			TriggerOutPin->AdvanceBlock();
+
+			if (!*EnableInPin)
+			{
+				return;
+			}
+			
 			const TArray<FMidiStreamEvent>& MidiEvents = MidiStreamInPin->GetEventsInBlock();
 			for (auto& Event : MidiEvents)
 			{
-				bool IsText = Event.MidiMessage.IsText();
-
-				if (Event.TrackIndex != TrackIndex || !IsText)
+				if (!Event.MidiMessage.IsText())
 				{
 					continue;
 				}
 
+				const int32 TrackIndex = Event.TrackIndex;
 				const FString* EventString = MidiStreamInPin->GetMidiTrackText(TrackIndex, Event.MidiMessage.GetTextIndex());
 				if (EventString && EventString->Compare(*TextInPin) == 0)
 				{
@@ -192,8 +141,146 @@ namespace HarmonixMetasound
 				}
 			}
 		}
-	
-	}
+
+	protected:
+		FBoolReadRef EnableInPin;
+		FMidiStreamReadRef MidiStreamInPin;
+		FStringReadRef TextInPin;
+
+		FTriggerWriteRef TriggerOutPin;
+	};
+
+	// Here for backward compat
+	class FMidiTextTriggerOperator_V0 : public FMidiTextTriggerOperator_V1
+	{
+	public:
+		static const FNodeClassMetadata& GetNodeInfo()
+		{
+			auto InitNodeInfo = []() -> FNodeClassMetadata
+			{
+				FNodeClassMetadata Info;
+				Info.ClassName        = { HarmonixNodeNamespace, TEXT("MidiTextTrigger"), TEXT("")};
+				Info.MajorVersion     = 0;
+				Info.MinorVersion     = 1;
+				Info.DisplayName      = METASOUND_LOCTEXT("MidiTextTriggerNode_DisplayName", "MIDI Text Trigger");
+				Info.Description      = METASOUND_LOCTEXT("MidiTextTriggerNode_Description", "Receives a MIDI stream, filters for the desired messages, and outputs triggers.");
+				Info.Author           = PluginAuthor;
+				Info.PromptIfMissing  = PluginNodeMissingPrompt;
+				Info.DefaultInterface = GetVertexInterface();
+				Info.CategoryHierarchy.Emplace(NodeCategories::Music);
+				return Info;
+			};
+
+			static const FNodeClassMetadata Info = InitNodeInfo();
+
+			return Info;
+		}
+		
+		static const FVertexInterface& GetVertexInterface()
+		{
+			using namespace MidiTextTriggerPinNames;
+			using namespace CommonPinNames;
+
+			static const FVertexInterface Interface(
+				FInputVertexInterface(
+					TInputDataVertex<bool>(METASOUND_GET_PARAM_NAME_AND_METADATA(Inputs::Enable), true),
+					TInputDataVertex<FMidiStream>(METASOUND_GET_PARAM_NAME_AND_METADATA(Inputs::MidiStream)),
+					TInputDataVertex<int32>(METASOUND_GET_PARAM_NAME_AND_METADATA(Inputs::MidiTrackNumber),1),
+					TInputDataVertex<FString>(METASOUND_GET_PARAM_NAME_AND_METADATA(TextInput))
+				),
+				FOutputVertexInterface(
+					TOutputDataVertex<FTrigger>(METASOUND_GET_PARAM_NAME_AND_METADATA(TriggerOutput))
+					)
+				);
+
+			return Interface;
+		}
+		
+		static TUniquePtr<IOperator> CreateOperator(const FBuildOperatorParams& InParams, FBuildResults& OutResults)
+		{
+			using namespace MidiTextTriggerPinNames;
+			using namespace CommonPinNames;
+
+			const FInputVertexInterfaceData& InputData = InParams.InputData;
+			FBoolReadRef InEnabled          = InputData.GetOrCreateDefaultDataReadReference<bool>(METASOUND_GET_PARAM_NAME(Inputs::Enable), InParams.OperatorSettings);
+			FMidiStreamReadRef InMidiStream = InputData.GetOrConstructDataReadReference<FMidiStream>(METASOUND_GET_PARAM_NAME(Inputs::MidiStream), InParams.OperatorSettings);
+			FInt32ReadRef InTrackNumber     = InputData.GetOrCreateDefaultDataReadReference<int32>(METASOUND_GET_PARAM_NAME(Inputs::MidiTrackNumber), InParams.OperatorSettings);
+			FStringReadRef InText           = InputData.GetOrCreateDefaultDataReadReference<FString>(METASOUND_GET_PARAM_NAME(TextInput), InParams.OperatorSettings);
+
+			return MakeUnique<FMidiTextTriggerOperator_V0>(InParams, InEnabled, InMidiStream, InTrackNumber, InText);
+		}
+
+		FMidiTextTriggerOperator_V0(
+			const FBuildOperatorParams& InParams,
+			const FBoolReadRef& InEnabled,
+			const FMidiStreamReadRef& InMidiStream,
+			const FInt32ReadRef& InTrackNumber,
+			const FStringReadRef& InText)
+		: FMidiTextTriggerOperator_V1(InParams, InEnabled, InMidiStream, InText)
+		, TrackNumberInPin(InTrackNumber)
+		{
+			Reset(InParams);
+		}
+
+		virtual void BindInputs(FInputVertexInterfaceData& InVertexData) override
+		{
+			FMidiTextTriggerOperator_V1::BindInputs(InVertexData);
+			InVertexData.BindReadVertex(CommonPinNames::Inputs::MidiTrackNumberName, TrackNumberInPin);
+		}
+
+		virtual void Execute() override
+		{
+			TriggerOutPin->AdvanceBlock();
+
+			if (!*EnableInPin)
+			{
+				return;
+			}
+
+			if (const int32 TrackIndex = *TrackNumberInPin; TrackIndex >= 0)
+			{
+				const TArray<FMidiStreamEvent>& MidiEvents = MidiStreamInPin->GetEventsInBlock();
+				for (auto& Event : MidiEvents)
+				{
+					if (Event.TrackIndex != TrackIndex || !Event.MidiMessage.IsText())
+					{
+						continue;
+					}
+
+					const FString* EventString = MidiStreamInPin->GetMidiTrackText(TrackIndex, Event.MidiMessage.GetTextIndex());
+					if (EventString && EventString->Compare(*TextInPin) == 0)
+					{
+						TriggerOutPin->TriggerFrame(Event.BlockSampleFrameIndex);
+					}
+				}
+			}
+		}
+
+	private:
+		FInt32ReadRef TrackNumberInPin;
+	};
+
+	class FMidiTextTriggerNode_V1 : public FNodeFacade
+	{
+	public:
+		explicit FMidiTextTriggerNode_V1(const FNodeInitData& InInitData)
+			: FNodeFacade(InInitData.InstanceName, InInitData.InstanceID, TFacadeOperatorClass<FMidiTextTriggerOperator_V1>())
+		{}
+		virtual ~FMidiTextTriggerNode_V1() override = default;
+	};
+
+	METASOUND_REGISTER_NODE(FMidiTextTriggerNode_V1)
+
+	class FMidiTextTriggerNode_V0 : public FNodeFacade
+	{
+	public:
+		explicit FMidiTextTriggerNode_V0(const FNodeInitData& InInitData)
+			: FNodeFacade(InInitData.InstanceName, InInitData.InstanceID, TFacadeOperatorClass<FMidiTextTriggerOperator_V0>())
+		{}
+		virtual ~FMidiTextTriggerNode_V0() override = default;
+	};
+
+	METASOUND_REGISTER_NODE(FMidiTextTriggerNode_V0)
 }
 
 #undef LOCTEXT_NAMESPACE // "HarmonixMetaSound"
diff --git a/Engine/Plugins/Runtime/Harmonix/Source/HarmonixMetasoundTests/Private/Nodes/MidiTextTriggerNodeTests.cpp b/Engine/Plugins/Runtime/Harmonix/Source/HarmonixMetasoundTests/Private/Nodes/MidiTextTriggerNodeTests.cpp
--- a/Engine/Plugins/Runtime/Harmonix/Source/HarmonixMetasoundTests/Private/Nodes/MidiTextTriggerNodeTests.cpp
+++ b/Engine/Plugins/Runtime/Harmonix/Source/HarmonixMetasoundTests/Private/Nodes/MidiTextTriggerNodeTests.cpp
@@ -68,7 +68,7 @@ namespace HarmonixMetasoundTests::MidiTextTriggerNode
 		UTEST_TRUE("Transport Out connected to MidiPlayer node", ConnectionSuccess);
 
 		// Midi text trigger node...
-		FNodeHandle TextTriggerNode = Builder.AddNode({ HarmonixMetasound::HarmonixNodeNamespace, "MidiTextTrigger", "" }, 0);
+		FNodeHandle TextTriggerNode = Builder.AddNode({ HarmonixMetasound::HarmonixNodeNamespace, "MidiTextTrigger", "" }, 1);
 		UTEST_TRUE("MidiTextTrigger Node Created", TextTriggerNode.Get().IsValid());
 		// hoist up text to match input...
 		Builder.AddAndConnectDataReferenceInput(TextTriggerNode, "Text", GetMetasoundDataTypeName<FString>());
