[Harmonix] Remove using declarations from headers and fix up files which were taking advantage of them

#tests automated tests

diff --git a/Engine/Plugins/Runtime/Harmonix/Source/HarmonixMetasound/Private/HarmonixMetasound/Analysis/MidiClockVertexAnalyzer.cpp b/Engine/Plugins/Runtime/Harmonix/Source/HarmonixMetasound/Private/HarmonixMetasound/Analysis/MidiClockVertexAnalyzer.cpp
--- a/Engine/Plugins/Runtime/Harmonix/Source/HarmonixMetasound/Private/HarmonixMetasound/Analysis/MidiClockVertexAnalyzer.cpp
+++ b/Engine/Plugins/Runtime/Harmonix/Source/HarmonixMetasound/Private/HarmonixMetasound/Analysis/MidiClockVertexAnalyzer.cpp
@@ -8,10 +8,10 @@
 
 namespace HarmonixMetasound::Analysis
 {
-	const Frontend::FAnalyzerOutput FMidiClockVertexAnalyzer::FOutputs::Timestamp = { "Timestamp", GetMetasoundDataTypeName<FMusicTimestamp>() };
-	const Frontend::FAnalyzerOutput FMidiClockVertexAnalyzer::FOutputs::Tempo = { "Tempo", GetMetasoundDataTypeName<float>() };
-	const Frontend::FAnalyzerOutput FMidiClockVertexAnalyzer::FOutputs::TimeSignature = { "TimeSignature", GetMetasoundDataTypeName<FTimeSignature>() };
-	const Frontend::FAnalyzerOutput FMidiClockVertexAnalyzer::FOutputs::Speed = { "Speed", GetMetasoundDataTypeName<float>() };
+	const Metasound::Frontend::FAnalyzerOutput FMidiClockVertexAnalyzer::FOutputs::Timestamp = { "Timestamp", Metasound::GetMetasoundDataTypeName<FMusicTimestamp>() };
+	const Metasound::Frontend::FAnalyzerOutput FMidiClockVertexAnalyzer::FOutputs::Tempo = { "Tempo", Metasound::GetMetasoundDataTypeName<float>() };
+	const Metasound::Frontend::FAnalyzerOutput FMidiClockVertexAnalyzer::FOutputs::TimeSignature = { "TimeSignature", Metasound::GetMetasoundDataTypeName<FTimeSignature>() };
+	const Metasound::Frontend::FAnalyzerOutput FMidiClockVertexAnalyzer::FOutputs::Speed = { "Speed", Metasound::GetMetasoundDataTypeName<float>() };
 	
 	const Metasound::Frontend::FAnalyzerOutput& FMidiClockVertexAnalyzer::FOutputs::GetValue()
 	{
@@ -20,7 +20,7 @@ namespace HarmonixMetasound::Analysis
 	
 	const TArray<Metasound::Frontend::FAnalyzerOutput>& FMidiClockVertexAnalyzer::FFactory::GetAnalyzerOutputs() const
 	{
-		static const TArray<Frontend::FAnalyzerOutput> Outputs
+		static const TArray<Metasound::Frontend::FAnalyzerOutput> Outputs
 		{
 			FOutputs::Timestamp,
 			FOutputs::Tempo,
@@ -44,9 +44,9 @@ namespace HarmonixMetasound::Analysis
 	FMidiClockVertexAnalyzer::FMidiClockVertexAnalyzer(const Metasound::Frontend::FCreateAnalyzerParams& InParams)
 		: FVertexAnalyzerBase(InParams.AnalyzerAddress, InParams.VertexDataReference)
 		, Timestamp(FMusicTimestampWriteRef::CreateNew())
-		, Tempo(FFloatWriteRef::CreateNew())
+		, Tempo(Metasound::FFloatWriteRef::CreateNew())
 		, TimeSignature(FTimeSignatureWriteRef::CreateNew())
-		, Speed(FFloatWriteRef::CreateNew())
+		, Speed(Metasound::FFloatWriteRef::CreateNew())
 	{
 		BindOutputData<FMusicTimestamp>(FOutputs::Timestamp.Name, InParams.OperatorSettings, Timestamp);
 		BindOutputData<float>(FOutputs::Tempo.Name, InParams.OperatorSettings, Tempo);
diff --git a/Engine/Plugins/Runtime/Harmonix/Source/HarmonixMetasound/Private/HarmonixMetasound/Analysis/MidiStreamVertexAnalyzer.cpp b/Engine/Plugins/Runtime/Harmonix/Source/HarmonixMetasound/Private/HarmonixMetasound/Analysis/MidiStreamVertexAnalyzer.cpp
--- a/Engine/Plugins/Runtime/Harmonix/Source/HarmonixMetasound/Private/HarmonixMetasound/Analysis/MidiStreamVertexAnalyzer.cpp
+++ b/Engine/Plugins/Runtime/Harmonix/Source/HarmonixMetasound/Private/HarmonixMetasound/Analysis/MidiStreamVertexAnalyzer.cpp
@@ -19,13 +19,13 @@ namespace HarmonixMetasound::Analysis
 
 	const Metasound::Frontend::FAnalyzerOutput& FMidiStreamVertexAnalyzer::FOutputs::GetValue()
 	{
-		static Metasound::Frontend::FAnalyzerOutput Value = { "LastMidiEvent", GetMetasoundDataTypeName<FMidiEventInfo>() };
+		static Metasound::Frontend::FAnalyzerOutput Value = { "LastMidiEvent", Metasound::GetMetasoundDataTypeName<FMidiEventInfo>() };
 		return Value;
 	}
 
 	const TArray<Metasound::Frontend::FAnalyzerOutput>& FMidiStreamVertexAnalyzer::FFactory::GetAnalyzerOutputs() const
 	{
-		static const TArray<Frontend::FAnalyzerOutput> Outputs { FOutputs::GetValue() };
+		static const TArray<Metasound::Frontend::FAnalyzerOutput> Outputs { FOutputs::GetValue() };
 		return Outputs;
 	}
 
@@ -33,7 +33,7 @@ namespace HarmonixMetasound::Analysis
 		: FVertexAnalyzerBase(InParams.AnalyzerAddress, InParams.VertexDataReference)
 		, LastMidiEvent(FMidiEventInfoWriteRef::CreateNew())
 	{
-		BindOutputData(FOutputs::GetValue().Name, InParams.OperatorSettings, TDataReadReference{ LastMidiEvent });
+		BindOutputData(FOutputs::GetValue().Name, InParams.OperatorSettings, Metasound::TDataReadReference{ LastMidiEvent });
 	}
 
 	void FMidiStreamVertexAnalyzer::Execute()
diff --git a/Engine/Plugins/Runtime/Harmonix/Source/HarmonixMetasound/Private/HarmonixMetasound/Nodes/MidiStreamTrackFilterNode.cpp b/Engine/Plugins/Runtime/Harmonix/Source/HarmonixMetasound/Private/HarmonixMetasound/Nodes/MidiStreamTrackFilterNode.cpp
--- a/Engine/Plugins/Runtime/Harmonix/Source/HarmonixMetasound/Private/HarmonixMetasound/Nodes/MidiStreamTrackFilterNode.cpp
+++ b/Engine/Plugins/Runtime/Harmonix/Source/HarmonixMetasound/Private/HarmonixMetasound/Nodes/MidiStreamTrackFilterNode.cpp
@@ -19,9 +19,11 @@ DEFINE_LOG_CATEGORY_STATIC(LogMidiStreamTrackFilter, Log, All);
 
 namespace HarmonixMetasound::Nodes::MidiStreamTrackFilter
 {
-	const Metasound::FNodeClassName& GetClassName()
+	using namespace Metasound;
+	
+	const FNodeClassName& GetClassName()
 	{
-		static Metasound::FNodeClassName ClassName
+		static FNodeClassName ClassName
 		{
 			HarmonixNodeNamespace,
 			"MidiStreamTrackFilter",
diff --git a/Engine/Plugins/Runtime/Harmonix/Source/HarmonixMetasound/Public/HarmonixMetasound/DataTypes/MidiClock.h b/Engine/Plugins/Runtime/Harmonix/Source/HarmonixMetasound/Public/HarmonixMetasound/DataTypes/MidiClock.h
--- a/Engine/Plugins/Runtime/Harmonix/Source/HarmonixMetasound/Public/HarmonixMetasound/DataTypes/MidiClock.h
+++ b/Engine/Plugins/Runtime/Harmonix/Source/HarmonixMetasound/Public/HarmonixMetasound/DataTypes/MidiClock.h
@@ -58,8 +58,6 @@ int32 SubdivisionToMidiTicks(EMidiClockSubdivisionQuantization Division, int32 C
 
 namespace HarmonixMetasound
 {
-	using namespace Metasound;
-
 	struct HARMONIXMETASOUND_API FMidiTimestampTransportState
 	{
 		int32 BlockSampleFrameIndex = 0;
@@ -86,7 +84,7 @@ namespace HarmonixMetasound
 	class HARMONIXMETASOUND_API FMidiClock
 	{
 	public:
-		explicit FMidiClock(const FOperatorSettings& InSettings);
+		explicit FMidiClock(const Metasound::FOperatorSettings& InSettings);
 		FMidiClock(const FMidiClock& Other);
 		FMidiClock(FMidiClock&& Other);
 
@@ -243,7 +241,7 @@ namespace HarmonixMetasound
 		int32 BlockSize;
 		int32 CurrentBlockFrameIndex;
 		float SampleRate;
-		FSampleCount SampleCount;
+		Metasound::FSampleCount SampleCount;
 		int32 FramesUntilNextProcess = 0;
 		static const int32 kMidiGranularity = 128;
 		FMidiTimestampTransportState CurrentTransportState;
diff --git a/Engine/Plugins/Runtime/Harmonix/Source/HarmonixMetasound/Public/HarmonixMetasound/DataTypes/MidiStream.h b/Engine/Plugins/Runtime/Harmonix/Source/HarmonixMetasound/Public/HarmonixMetasound/DataTypes/MidiStream.h
--- a/Engine/Plugins/Runtime/Harmonix/Source/HarmonixMetasound/Public/HarmonixMetasound/DataTypes/MidiStream.h
+++ b/Engine/Plugins/Runtime/Harmonix/Source/HarmonixMetasound/Public/HarmonixMetasound/DataTypes/MidiStream.h
@@ -18,8 +18,6 @@
 
 namespace HarmonixMetasound
 {
-	using namespace Metasound;
-
 	DECLARE_LOG_CATEGORY_EXTERN(LogMidiStreamDataType, Log, All)
 
 	struct HARMONIXMETASOUND_API FMidiStreamEvent
@@ -87,7 +85,7 @@ namespace HarmonixMetasound
 	{
 	public:
 
-		FMidiStream(const FOperatorSettings& InSettings);
+		FMidiStream(const Metasound::FOperatorSettings& InSettings);
 
 		void SetClockSource(const FMidiClockReadRef& ClockSource)
 		{
@@ -130,7 +128,7 @@ namespace HarmonixMetasound
 		}
 
 		// Forward declare the read ref for copying other streams
-		using FMidiStreamReadRef = TDataReadReference<FMidiStream>;
+		using FMidiStreamReadRef = Metasound::TDataReadReference<FMidiStream>;
 
 		// Copies all midi and transport events from the in stream to this stream. Obliterates any
 		// transport/midi events already in this stream!
@@ -391,7 +389,7 @@ namespace HarmonixMetasound
 	private:
 		FMidiFileProxyPtr MidiFileSourceOfEvents;
 		int32 NumFramesPerBlock = 0;
-		FSampleRate SampleRate  = 0;
+		Metasound::FSampleRate SampleRate  = 0;
 		int32 TicksPerQuarterNote = MidiConstants::kTicksPerQuarterNoteInt;
 
 		FMidiTimestampTransportState CurrentTransportState;
diff --git a/Engine/Plugins/Runtime/Harmonix/Source/HarmonixMetasound/Public/HarmonixMetasound/DataTypes/MusicTransport.h b/Engine/Plugins/Runtime/Harmonix/Source/HarmonixMetasound/Public/HarmonixMetasound/DataTypes/MusicTransport.h
--- a/Engine/Plugins/Runtime/Harmonix/Source/HarmonixMetasound/Public/HarmonixMetasound/DataTypes/MusicTransport.h
+++ b/Engine/Plugins/Runtime/Harmonix/Source/HarmonixMetasound/Public/HarmonixMetasound/DataTypes/MusicTransport.h
@@ -11,8 +11,6 @@
 
 namespace HarmonixMetasound
 {
-	using namespace Metasound;
-
 	enum class EMusicPlayerTransportRequest : uint8
 	{
 		None,
@@ -29,7 +27,7 @@ namespace HarmonixMetasound
 	class HARMONIXMETASOUND_API FMusicTransportEventStream
 	{
 	public:
-		FMusicTransportEventStream(const FOperatorSettings& InSettings);
+		FMusicTransportEventStream(const Metasound::FOperatorSettings& InSettings);
 
 		void AddTransportRequest(EMusicPlayerTransportRequest InRequest, int32 AtSampleIndex);
 		void AddSeekRequest(int32 AtSampleIndex, const FMusicSeekTarget& Target);
diff --git a/Engine/Plugins/Runtime/Harmonix/Source/HarmonixMetasoundTests/Private/Nodes/DelayNodeTests.cpp b/Engine/Plugins/Runtime/Harmonix/Source/HarmonixMetasoundTests/Private/Nodes/DelayNodeTests.cpp
--- a/Engine/Plugins/Runtime/Harmonix/Source/HarmonixMetasoundTests/Private/Nodes/DelayNodeTests.cpp
+++ b/Engine/Plugins/Runtime/Harmonix/Source/HarmonixMetasoundTests/Private/Nodes/DelayNodeTests.cpp
@@ -15,6 +15,7 @@ namespace HarmonixMetasoundTests::DelayNode
 {
 	using GraphBuilder = Metasound::Test::FNodeTestGraphBuilder;
 	using namespace HarmonixMetasound;
+	using namespace Metasound;
 	
 	class FTestFixture
 	{
