[Harmonix] Fix CIS issues with remaining node files/using Metasound namespace

#tests non-unity validator

diff --git a/Engine/Plugins/Runtime/Harmonix/Source/HarmonixMetasound/Private/HarmonixMetasound/Nodes/MidiCCTriggerNode.cpp b/Engine/Plugins/Runtime/Harmonix/Source/HarmonixMetasound/Private/HarmonixMetasound/Nodes/MidiCCTriggerNode.cpp
--- a/Engine/Plugins/Runtime/Harmonix/Source/HarmonixMetasound/Private/HarmonixMetasound/Nodes/MidiCCTriggerNode.cpp
+++ b/Engine/Plugins/Runtime/Harmonix/Source/HarmonixMetasound/Private/HarmonixMetasound/Nodes/MidiCCTriggerNode.cpp
@@ -11,7 +11,6 @@
 #include "HarmonixMetasound/Nodes/MidiCCTriggerNode.h"
 #include "HarmonixMetasound/Common.h"
 #include "HarmonixMetasound/DataTypes/MidiStream.h"
-#include "HarmonixMetasound/DataTypes/MusicTransport.h"
 #include "HarmonixMidi/MidiMsg.h"
 #include "MetasoundEnumRegistrationMacro.h"
 #include "HarmonixMetasound/DataTypes/MidiControllerID.h"
@@ -22,9 +21,11 @@ DEFINE_LOG_CATEGORY_STATIC(LogMidiCCTrigger, Log, All);
 
 namespace HarmonixMetasound::Nodes::MidiCCTriggerNode
 {
-	const Metasound::FNodeClassName& GetClassName()
+	using namespace Metasound;
+	
+	const FNodeClassName& GetClassName()
 	{
-		static const Metasound::FNodeClassName ClassName
+		static const FNodeClassName ClassName
 		{
 			HarmonixNodeNamespace,
 				TEXT("MidiCCTrigger"),
@@ -49,28 +50,28 @@ namespace HarmonixMetasound::Nodes::MidiCCTriggerNode
 		DEFINE_OUTPUT_METASOUND_PARAM(OutputTrigger, "Trigger Out", "A trigger when a Midi Control Change message is found");
 	}
 
-	class FOp final : public Metasound::TExecutableOperator<FOp>
+	class FOp final : public TExecutableOperator<FOp>
 	{
 	public:
 		struct FInputs
 		{
-			Metasound::FBoolReadRef       Enable;
-			Metasound::FInt32ReadRef      TrackNumber;
-			Metasound::FInt32ReadRef      ChannelNumber;
-			Metasound::FEnumStdMidiControllerIDReadRef  ControllerID;
-			HarmonixMetasound::FMidiStreamReadRef MidiStream;
+			FBoolReadRef Enable;
+			FInt32ReadRef TrackNumber;
+			FInt32ReadRef ChannelNumber;
+			FEnumStdMidiControllerIDReadRef ControllerID;
+			FMidiStreamReadRef MidiStream;
 		};
 
 		struct FOutputs
 		{
-			Metasound::FInt32WriteRef ControlChangeValueInt32;
-			Metasound::FFloatWriteRef ControlChangeValueFloat;
-			Metasound::FTriggerWriteRef TriggerOut;
+			FInt32WriteRef ControlChangeValueInt32;
+			FFloatWriteRef ControlChangeValueFloat;
+			FTriggerWriteRef TriggerOut;
 		};
 
-		static const Metasound::FVertexInterface& GetVertexInterface()
+		static const FVertexInterface& GetVertexInterface()
 		{
-			const auto MakeInterface = []() -> Metasound::FVertexInterface
+			const auto MakeInterface = []() -> FVertexInterface
 			{
 				using namespace Metasound;
 
@@ -93,12 +94,12 @@ namespace HarmonixMetasound::Nodes::MidiCCTriggerNode
 				};
 			};
 
-			static const Metasound::FVertexInterface Interface = MakeInterface();
+			static const FVertexInterface Interface = MakeInterface();
 
 			return Interface;
 		}
 
-		static const Metasound::FNodeClassMetadata& GetNodeInfo()
+		static const FNodeClassMetadata& GetNodeInfo()
 		{
 			auto InitNodeInfo = []() -> FNodeClassMetadata
 			{
@@ -120,9 +121,9 @@ namespace HarmonixMetasound::Nodes::MidiCCTriggerNode
 			return Info;
 		}
 
-		static TUniquePtr<IOperator> CreateOperator(const Metasound::FBuildOperatorParams& InParams, Metasound::FBuildResults& OutResults)
+		static TUniquePtr<IOperator> CreateOperator(const FBuildOperatorParams& InParams, FBuildResults& OutResults)
 		{
-			const Metasound::FOperatorSettings& OperatorSettings = InParams.OperatorSettings;
+			const FOperatorSettings& OperatorSettings = InParams.OperatorSettings;
 			const FInputVertexInterfaceData& InputData = InParams.InputData;
 
 			FInputs Inputs
@@ -131,20 +132,18 @@ namespace HarmonixMetasound::Nodes::MidiCCTriggerNode
 				InputData.GetOrCreateDefaultDataReadReference<int32>(Inputs::MidiTrackNumberName, OperatorSettings),
 				InputData.GetOrCreateDefaultDataReadReference<int32>(Inputs::MidiChannelNumberName, OperatorSettings),
 				InputData.GetOrCreateDefaultDataReadReference<FEnumStdMidiControllerID>(Inputs::InputMidiControllerIDName, OperatorSettings),
-				InputData.GetOrCreateDefaultDataReadReference<HarmonixMetasound::FMidiStream>(Inputs::MidiStreamName, OperatorSettings)
+				InputData.GetOrCreateDefaultDataReadReference<FMidiStream>(Inputs::MidiStreamName, OperatorSettings)
 			};
 
 			FOutputs Outputs
 			{
-				Metasound::FInt32WriteRef::CreateNew(0),
-				Metasound::FFloatWriteRef::CreateNew(0.0f),
-				Metasound::FTriggerWriteRef::CreateNew(InParams.OperatorSettings)
+				FInt32WriteRef::CreateNew(0), FFloatWriteRef::CreateNew(0.0f), FTriggerWriteRef::CreateNew(InParams.OperatorSettings)
 			};
 
 			return MakeUnique<FOp>(InParams, MoveTemp(Inputs), MoveTemp(Outputs));
 		}
 
-		FOp(const Metasound::FBuildOperatorParams& Params, FInputs&& InInputs, FOutputs&& InOutputs)
+		FOp(const FBuildOperatorParams& Params, FInputs&& InInputs, FOutputs&& InOutputs)
 			: Inputs(MoveTemp(InInputs))
 			, Outputs(MoveTemp(InOutputs))
 		{
@@ -221,10 +220,10 @@ namespace HarmonixMetasound::Nodes::MidiCCTriggerNode
 			FOutputs Outputs;
 	};
 
-	class FMidiCCTriggerNode final : public Metasound::FNodeFacade
+	class FMidiCCTriggerNode final : public FNodeFacade
 	{
 	public:
-		explicit FMidiCCTriggerNode(const Metasound::FNodeInitData& InInitData)
+		explicit FMidiCCTriggerNode(const FNodeInitData& InInitData)
 			: FNodeFacade(InInitData.InstanceName, InInitData.InstanceID, Metasound::TFacadeOperatorClass<FOp>())
 		{}
 	};
diff --git a/Engine/Plugins/Runtime/Harmonix/Source/HarmonixMetasound/Private/HarmonixMetasound/Nodes/MorphingLfoNode.cpp b/Engine/Plugins/Runtime/Harmonix/Source/HarmonixMetasound/Private/HarmonixMetasound/Nodes/MorphingLfoNode.cpp
--- a/Engine/Plugins/Runtime/Harmonix/Source/HarmonixMetasound/Private/HarmonixMetasound/Nodes/MorphingLfoNode.cpp
+++ b/Engine/Plugins/Runtime/Harmonix/Source/HarmonixMetasound/Private/HarmonixMetasound/Nodes/MorphingLfoNode.cpp
@@ -16,6 +16,8 @@
 
 namespace HarmonixMetasound::Nodes::MorphingLFO
 {
+	using namespace Metasound;
+	
 	namespace Inputs
 	{
 		DEFINE_METASOUND_PARAM_ALIAS(MidiClock, CommonPinNames::Inputs::MidiClock);
@@ -31,7 +33,7 @@ namespace HarmonixMetasound::Nodes::MorphingLFO
 	}
 
 	template<typename OutputDataType>
-	class TOp final : public Metasound::TExecutableOperator<TOp<OutputDataType>>
+	class TOp final : public TExecutableOperator<TOp<OutputDataType>>
 	{
 	public:
 		static const FVertexInterface& GetVertexInterface()
@@ -296,10 +298,10 @@ namespace HarmonixMetasound::Nodes::MorphingLFO
 	}
 	
 	template<typename OutputDataType>
-	class TNode final : public Metasound::FNodeFacade
+	class TNode final : public FNodeFacade
 	{
 	public:
-		explicit TNode(const Metasound::FNodeInitData& InitData)
+		explicit TNode(const FNodeInitData& InitData)
 		: FNodeFacade(InitData.InstanceName, InitData.InstanceID, Metasound::TFacadeOperatorClass<TOp<OutputDataType>>())
 		{}
 	};
