Don't warn about deprecated properties we're no longer supporting. These warnings will get thrown in Main.

I already have a DataValidator warning to tell the user to update their asset, so let's just leave it at that.

Also, we don't need to keep the deprecated properties around anymore either.

#tests PIE

diff --git a/Engine/Plugins/Runtime/Harmonix/Source/HarmonixDsp/Private/HarmonixDsp/FusionSampler/FusionPatch.cpp b/Engine/Plugins/Runtime/Harmonix/Source/HarmonixDsp/Private/HarmonixDsp/FusionSampler/FusionPatch.cpp
--- a/Engine/Plugins/Runtime/Harmonix/Source/HarmonixDsp/Private/HarmonixDsp/FusionSampler/FusionPatch.cpp
+++ b/Engine/Plugins/Runtime/Harmonix/Source/HarmonixDsp/Private/HarmonixDsp/FusionSampler/FusionPatch.cpp
@@ -183,40 +183,8 @@ void UFusionPatch::Serialize(FArchive& Ar)
 				UE_LOG(LogFusionPatch, Warning, TEXT("Fusion patch (%s) has more than one preset, but presets have been deprecated. Only the first (Default) preset will be loaded"), *GetPathName());
 			}
 		}
-
-		if (Version < FFusionPatchCustomVersion::DeprecatedUnusedEffectsSettings)
-		{
-			if (FusionPatchData.Settings.Delay_DEPRECATED.IsEnabled)
-			{
-				UE_LOG(LogFusionPatch, Warning, TEXT("Fusion Patch (%s) has the \"Delay\" effect enabled, but that effect is no longer supported by the Fusion Sampler! "
-					"If you intended to use the Delay effect, rework your metasound to apply the Delay to the output of the FusionSampler node using this patch"),
-					*GetPathName());
-			}
-
-			if (FusionPatchData.Settings.Distortion_DEPRECATED.IsEnabled)
-			{
-				UE_LOG(LogFusionPatch, Warning, TEXT("Fusion Patch (%s) has the \"Distortion\" effect enabled, but that effect is no longer supported by the Fusion Sampler! "
-					"If you intended to use the Distortion effect, rework your metasound to apply the Distortion to the output of the FusionSampler node using this patch"),
-					*GetPathName());
-			}
-
-			if (FusionPatchData.Settings.BitCrusher_DEPRECATED.IsEnabled)
-			{
-				UE_LOG(LogFusionPatch, Warning, TEXT("Fusion Patch (%s) has the \"BitCrusher\" effect enabled, but that effect is no longer supported by the Fusion Sampler! "
-					"If you intended to use the BitCrusher effect, rework your metasound to apply the Distortion to the output of the FusionSampler node using this patch"),
-					*GetPathName());
-			}
-
-			if (FusionPatchData.Settings.Vocoder_DEPRECATED.IsEnabled)
-			{
-				UE_LOG(LogFusionPatch, Warning, TEXT("Fusion Patch (%s) has the \"Vocoder\" effect enabled, but that effect is no longer supported by the Fusion Sampler! "
-					"If you intended to use the Vocoder effect, rework your metasound to apply the Vocoder to the output of the FusionSampler node using this patch"),
-					*GetPathName());	
-			}
-		}
 	}
 
-
 	const UStretcherAndPitchShifterFactoryConfig* FactoryConfig = GetDefault<UStretcherAndPitchShifterFactoryConfig>();
 	for (FKeyzoneSettings& Keyzone : FusionPatchData.Keyzones)
 	{
diff --git a/Engine/Plugins/Runtime/Harmonix/Source/HarmonixDsp/Public/HarmonixDsp/FusionSampler/Settings/FusionPatchSettings.h b/Engine/Plugins/Runtime/Harmonix/Source/HarmonixDsp/Public/HarmonixDsp/FusionSampler/Settings/FusionPatchSettings.h
--- a/Engine/Plugins/Runtime/Harmonix/Source/HarmonixDsp/Public/HarmonixDsp/FusionSampler/Settings/FusionPatchSettings.h
+++ b/Engine/Plugins/Runtime/Harmonix/Source/HarmonixDsp/Public/HarmonixDsp/FusionSampler/Settings/FusionPatchSettings.h
@@ -79,16 +79,5 @@ public:
 
 	UPROPERTY(EditDefaultsOnly, Category = "Settings")
 	FModulatorSettingsArray VelocityModulators;
-	
-	UPROPERTY()
-	FDelaySettings Delay_DEPRECATED;
-
-	UPROPERTY()
-	FDistortionSettingsV1 Distortion_DEPRECATED;
 
-	UPROPERTY()
-	FBitCrusherSettings BitCrusher_DEPRECATED;
-
-	UPROPERTY()
-	FVocoderSettings Vocoder_DEPRECATED;
 };
\ No newline at end of file
