[Harmonix] MIDI Text Trigger - deprecate old version

diff --git a/Engine/Plugins/Runtime/Harmonix/Source/HarmonixMetasound/Private/HarmonixMetasound/Nodes/MidiTextTriggerNode.cpp b/Engine/Plugins/Runtime/Harmonix/Source/HarmonixMetasound/Private/HarmonixMetasound/Nodes/MidiTextTriggerNode.cpp
--- a/Engine/Plugins/Runtime/Harmonix/Source/HarmonixMetasound/Private/HarmonixMetasound/Nodes/MidiTextTriggerNode.cpp
+++ b/Engine/Plugins/Runtime/Harmonix/Source/HarmonixMetasound/Private/HarmonixMetasound/Nodes/MidiTextTriggerNode.cpp
@@ -168,6 +168,7 @@ namespace HarmonixMetasound
 				Info.PromptIfMissing  = PluginNodeMissingPrompt;
 				Info.DefaultInterface = GetVertexInterface();
 				Info.CategoryHierarchy.Emplace(NodeCategories::Music);
+				Info.bDeprecated = true;
 				return Info;
 			};
 
