[Harmonix] New version of MIDI CC Trigger node to conform to library design

#virtualized

diff --git a/Engine/Plugins/Runtime/Harmonix/Source/HarmonixMetasound/Private/HarmonixMetasound/Nodes/MidiCCTriggerNode.cpp b/Engine/Plugins/Runtime/Harmonix/Source/HarmonixMetasound/Private/HarmonixMetasound/Nodes/MidiCCTriggerNode.cpp
--- a/Engine/Plugins/Runtime/Harmonix/Source/HarmonixMetasound/Private/HarmonixMetasound/Nodes/MidiCCTriggerNode.cpp
+++ b/Engine/Plugins/Runtime/Harmonix/Source/HarmonixMetasound/Private/HarmonixMetasound/Nodes/MidiCCTriggerNode.cpp
@@ -1,5 +1,7 @@
 // Copyright Epic Games, Inc. All Rights Reserved.
 
+#include "HarmonixMetasound/Nodes/MidiCCTriggerNode.h"
+
 #include "MetasoundExecutableOperator.h"
 #include "MetasoundFacade.h"
 #include "MetasoundNodeInterface.h"
@@ -8,15 +10,12 @@
 #include "MetasoundStandardNodesCategories.h"
 #include "MetasoundVertex.h"
 #include "MetasoundTrigger.h"
-#include "HarmonixMetasound/Nodes/MidiCCTriggerNode.h"
 #include "HarmonixMetasound/Common.h"
 #include "HarmonixMetasound/DataTypes/MidiStream.h"
 #include "HarmonixMidi/MidiMsg.h"
 #include "MetasoundEnumRegistrationMacro.h"
 #include "HarmonixMetasound/DataTypes/MidiControllerID.h"
 
-DEFINE_LOG_CATEGORY_STATIC(LogMidiCCTrigger, Log, All);
-
 #define LOCTEXT_NAMESPACE "HarmonixMetaSound"
 
 namespace HarmonixMetasound::Nodes::MidiCCTriggerNode
@@ -34,23 +33,184 @@ namespace HarmonixMetasound::Nodes::MidiCCTriggerNode
 		return ClassName;
 	}
 
+	int32 GetCurrentMajorVersion()
+	{
+		return 1;
+	}
+	
 	namespace Inputs
 	{
 		DEFINE_METASOUND_PARAM_ALIAS(Enable, CommonPinNames::Inputs::Enable);
+		DEFINE_METASOUND_PARAM_ALIAS(MidiStream, CommonPinNames::Inputs::MidiStream);
+		DEFINE_INPUT_METASOUND_PARAM(InputMidiControllerID, "Control Number", "MIDI Control Number");
+
+		// Removed in V1
+		DECLARE_METASOUND_PARAM_EXTERN(MidiTrackNumber);
 		DEFINE_METASOUND_PARAM_ALIAS(MidiTrackNumber, CommonPinNames::Inputs::MidiTrackNumber);
+		DECLARE_METASOUND_PARAM_EXTERN(MidiChannelNumber);
 		DEFINE_METASOUND_PARAM_ALIAS(MidiChannelNumber, CommonPinNames::Inputs::MidiChannelNumber);
-		DEFINE_INPUT_METASOUND_PARAM(InputMidiControllerID, "Standard Midi Controller ID", "Standard Midi Controller ID (0-127)");
-		DEFINE_METASOUND_PARAM_ALIAS(MidiStream, CommonPinNames::Inputs::MidiStream);
 	}
 
 	namespace Outputs
 	{
-		DEFINE_OUTPUT_METASOUND_PARAM(OutputControlChangeValueInt32, "Midi Control Change Value (Int32)", "Midi Control Change value (0-127)");
-		DEFINE_OUTPUT_METASOUND_PARAM(OutputControlChangeValueFloat, "Midi Control Change Value (Float)", "normalized Midi Control Change value (0.0-1.0)");
-		DEFINE_OUTPUT_METASOUND_PARAM(OutputTrigger, "Trigger Out", "A trigger when a Midi Control Change message is found");
+		DEFINE_OUTPUT_METASOUND_PARAM(OutputControlChangeValueInt32, "Value", "Control Change Value (0-127)");
+
+		// Removed in V1
+		DECLARE_METASOUND_PARAM_EXTERN(OutputTrigger);
+		DEFINE_OUTPUT_METASOUND_PARAM(OutputTrigger, "Trigger Out", "A trigger when a MIDI Control Change message is found");
+		DECLARE_METASOUND_PARAM_EXTERN(OutputControlNumber);
+		DEFINE_OUTPUT_METASOUND_PARAM(OutputControlNumber, "Control Number", "Control Number (0-127)");
+		DECLARE_METASOUND_PARAM_EXTERN(OutputControlChangeValueFloat);
+		DEFINE_OUTPUT_METASOUND_PARAM(OutputControlChangeValueFloat, "Value (Normalized)", "Normalized Control Change value (0.0-1.0)");
 	}
 
-	class FOp final : public TExecutableOperator<FOp>
+	class FMidiCCTriggerOperator_V1 final : public TExecutableOperator<FMidiCCTriggerOperator_V1>
+	{
+	public:
+		struct FInputs
+		{
+			FBoolReadRef Enable;
+			FMidiStreamReadRef MidiStream;
+			FEnumStdMidiControllerIDReadRef ControllerID;
+		};
+
+		struct FOutputs
+		{
+			FInt32WriteRef ControlChangeValueInt32;
+		};
+
+		static const FVertexInterface& GetVertexInterface()
+		{
+			const auto MakeInterface = []() -> FVertexInterface
+			{
+				using namespace Metasound;
+
+				return
+				{
+					FInputVertexInterface
+					{
+						TInputDataVertex<bool>(METASOUND_GET_PARAM_NAME_AND_METADATA(Inputs::Enable), true),
+						TInputDataVertex<FMidiStream>(METASOUND_GET_PARAM_NAME_AND_METADATA(Inputs::MidiStream)),
+						TInputDataVertex<FEnumStdMidiControllerID>(METASOUND_GET_PARAM_NAME_AND_METADATA(Inputs::InputMidiControllerID))
+					},
+					FOutputVertexInterface
+					{
+						TOutputDataVertex<int32>(METASOUND_GET_PARAM_NAME_AND_METADATA(Outputs::OutputControlChangeValueInt32))
+					}
+				};
+			};
+
+			static const FVertexInterface Interface = MakeInterface();
+
+			return Interface;
+		}
+
+		static const FNodeClassMetadata& GetNodeInfo()
+		{
+			auto InitNodeInfo = []() -> FNodeClassMetadata
+			{
+				FNodeClassMetadata Info;
+				Info.ClassName = GetClassName();
+				Info.MajorVersion = 1;
+				Info.MinorVersion = 0;
+				Info.DisplayName = METASOUND_LOCTEXT("MidiCCTriggerNode_V1_DisplayName", "MIDI CC Trigger");
+				Info.Description = METASOUND_LOCTEXT("MidiCCTriggerNode_V1_Description", "Find a MIDI Control Change message in a MIDI stream and output the value.");
+				Info.Author = PluginAuthor;
+				Info.PromptIfMissing = PluginNodeMissingPrompt;
+				Info.DefaultInterface = GetVertexInterface();
+				Info.CategoryHierarchy.Emplace(NodeCategories::Music);
+				return Info;
+			};
+
+			static const FNodeClassMetadata Info = InitNodeInfo();
+
+			return Info;
+		}
+
+		static TUniquePtr<IOperator> CreateOperator(const FBuildOperatorParams& InParams, FBuildResults& OutResults)
+		{
+			const FOperatorSettings& OperatorSettings = InParams.OperatorSettings;
+			const FInputVertexInterfaceData& InputData = InParams.InputData;
+
+			FInputs Inputs
+			{
+				InputData.GetOrCreateDefaultDataReadReference<bool>(Inputs::EnableName,OperatorSettings),
+				InputData.GetOrCreateDefaultDataReadReference<FMidiStream>(Inputs::MidiStreamName, OperatorSettings),
+				InputData.GetOrCreateDefaultDataReadReference<FEnumStdMidiControllerID>(Inputs::InputMidiControllerIDName, OperatorSettings)
+			};
+
+			FOutputs Outputs
+			{
+				FInt32WriteRef::CreateNew(0)
+			};
+
+			return MakeUnique<FMidiCCTriggerOperator_V1>(InParams, MoveTemp(Inputs), MoveTemp(Outputs));
+		}
+
+		FMidiCCTriggerOperator_V1(const FBuildOperatorParams& Params, FInputs&& InInputs, FOutputs&& InOutputs)
+			: Inputs(MoveTemp(InInputs))
+			, Outputs(MoveTemp(InOutputs))
+		{
+			Reset(Params);
+		}
+
+		virtual void BindInputs(FInputVertexInterfaceData& InVertexData) override
+		{
+			InVertexData.BindReadVertex(Inputs::EnableName, Inputs.Enable);
+			InVertexData.BindReadVertex(Inputs::MidiStreamName, Inputs.MidiStream);
+			InVertexData.BindReadVertex(Inputs::InputMidiControllerIDName, Inputs.ControllerID);
+		}
+
+		virtual void BindOutputs(FOutputVertexInterfaceData& InVertexData) override
+		{
+			InVertexData.BindReadVertex(Outputs::OutputControlChangeValueInt32Name, Outputs.ControlChangeValueInt32);
+		}
+
+		void Reset(const FResetParams& ResetParams)
+		{
+			*Outputs.ControlChangeValueInt32 = 0;
+		}
+
+		void Execute()
+		{
+			if (!*Inputs.Enable)
+			{
+				return;
+			}
+
+			const TArray<FMidiStreamEvent>& MidiEvents = Inputs.MidiStream->GetEventsInBlock();
+			for (const FMidiStreamEvent& Event : MidiEvents)
+			{
+				if (!Event.MidiMessage.IsControlChange())
+				{
+					continue;
+				}
+				
+				if (Event.MidiMessage.GetStdData1() != static_cast<uint8>(Inputs.ControllerID->Get()))
+				{
+					continue;
+				}
+
+				*Outputs.ControlChangeValueInt32 = static_cast<int32>(Event.MidiMessage.GetStdData2());
+			}
+		}
+
+	private:
+		FInputs Inputs;
+		FOutputs Outputs;
+	};
+
+	class FMidiCCTriggerNode_V1 final : public FNodeFacade
+	{
+	public:
+		explicit FMidiCCTriggerNode_V1(const FNodeInitData& InInitData)
+			: FNodeFacade(InInitData.InstanceName, InInitData.InstanceID, Metasound::TFacadeOperatorClass<FMidiCCTriggerOperator_V1>())
+		{}
+	};
+
+	METASOUND_REGISTER_NODE(FMidiCCTriggerNode_V1);
+
+	class FMidiCCTriggerOperator_V0 final : public Metasound::TExecutableOperator<FMidiCCTriggerOperator_V0>
 	{
 	public:
 		struct FInputs
@@ -104,7 +264,7 @@ namespace HarmonixMetasound::Nodes::MidiCCTriggerNode
 			auto InitNodeInfo = []() -> FNodeClassMetadata
 			{
 				FNodeClassMetadata Info;
-				Info.ClassName = { HarmonixNodeNamespace, TEXT("MidiCCTrigger"), TEXT("") };
+				Info.ClassName = GetClassName();
 				Info.MajorVersion = 0;
 				Info.MinorVersion = 1;
 				Info.DisplayName = METASOUND_LOCTEXT("MidiCCTriggerNode_DisplayName", "Midi CC Trigger");
@@ -113,6 +273,7 @@ namespace HarmonixMetasound::Nodes::MidiCCTriggerNode
 				Info.PromptIfMissing = PluginNodeMissingPrompt;
 				Info.DefaultInterface = GetVertexInterface();
 				Info.CategoryHierarchy.Emplace(NodeCategories::Music);
+				Info.bDeprecated = true;
 				return Info;
 			};
 
@@ -140,10 +301,10 @@ namespace HarmonixMetasound::Nodes::MidiCCTriggerNode
 				FInt32WriteRef::CreateNew(0), FFloatWriteRef::CreateNew(0.0f), FTriggerWriteRef::CreateNew(InParams.OperatorSettings)
 			};
 
-			return MakeUnique<FOp>(InParams, MoveTemp(Inputs), MoveTemp(Outputs));
+			return MakeUnique<FMidiCCTriggerOperator_V0>(InParams, MoveTemp(Inputs), MoveTemp(Outputs));
 		}
 
-		FOp(const FBuildOperatorParams& Params, FInputs&& InInputs, FOutputs&& InOutputs)
+		FMidiCCTriggerOperator_V0(const Metasound::FBuildOperatorParams& Params, FInputs&& InInputs, FOutputs&& InOutputs)
 			: Inputs(MoveTemp(InInputs))
 			, Outputs(MoveTemp(InOutputs))
 		{
@@ -220,15 +381,15 @@ namespace HarmonixMetasound::Nodes::MidiCCTriggerNode
 			FOutputs Outputs;
 	};
 
-	class FMidiCCTriggerNode final : public FNodeFacade
+	class FMidiCCTriggerNode_V0 final : public FNodeFacade
 	{
 	public:
-		explicit FMidiCCTriggerNode(const FNodeInitData& InInitData)
-			: FNodeFacade(InInitData.InstanceName, InInitData.InstanceID, Metasound::TFacadeOperatorClass<FOp>())
+		explicit FMidiCCTriggerNode_V0(const FNodeInitData& InInitData)
+			: FNodeFacade(InInitData.InstanceName, InInitData.InstanceID, Metasound::TFacadeOperatorClass<FMidiCCTriggerOperator_V0>())
 		{}
 	};
 
-	METASOUND_REGISTER_NODE(FMidiCCTriggerNode);
+	METASOUND_REGISTER_NODE(FMidiCCTriggerNode_V0);
 }
 
 #undef LOCTEXT_NAMESPACE
\ No newline at end of file
diff --git a/Engine/Plugins/Runtime/Harmonix/Source/HarmonixMetasound/Public/HarmonixMetasound/Nodes/MidiCCTriggerNode.h b/Engine/Plugins/Runtime/Harmonix/Source/HarmonixMetasound/Public/HarmonixMetasound/Nodes/MidiCCTriggerNode.h
--- a/Engine/Plugins/Runtime/Harmonix/Source/HarmonixMetasound/Public/HarmonixMetasound/Nodes/MidiCCTriggerNode.h
+++ b/Engine/Plugins/Runtime/Harmonix/Source/HarmonixMetasound/Public/HarmonixMetasound/Nodes/MidiCCTriggerNode.h
@@ -2,27 +2,23 @@
 
 #pragma once
 
-#include "MetasoundDataReference.h"
 #include "MetasoundNodeInterface.h"
 #include "HarmonixMetasound/Common.h"
 
 namespace HarmonixMetasound::Nodes::MidiCCTriggerNode
 {
 	const HARMONIXMETASOUND_API Metasound::FNodeClassName& GetClassName();
-
+	HARMONIXMETASOUND_API int32 GetCurrentMajorVersion();
+	
 	namespace Inputs
 	{
 		DECLARE_METASOUND_PARAM_EXTERN(Enable);
-		DECLARE_METASOUND_PARAM_EXTERN(MidiTrackNumber);
-		DECLARE_METASOUND_PARAM_EXTERN(MidiChannelNumber);
-		DECLARE_METASOUND_PARAM_EXTERN(InputMidiControllerID);
 		DECLARE_METASOUND_PARAM_EXTERN(MidiStream);
+		DECLARE_METASOUND_PARAM_EXTERN(InputMidiControllerID);
 	}
 
 	namespace Outputs
 	{
 		DECLARE_METASOUND_PARAM_EXTERN(OutputControlChangeValueInt32);
-		DECLARE_METASOUND_PARAM_EXTERN(OutputControlChangeValueFloat);
-		DECLARE_METASOUND_PARAM_EXTERN(OutputTrigger);
 	}
 }
\ No newline at end of file
diff --git a/Engine/Plugins/Runtime/Harmonix/Source/HarmonixMetasoundTests/Private/Nodes/MidiCCTriggerNodeTests.cpp b/Engine/Plugins/Runtime/Harmonix/Source/HarmonixMetasoundTests/Private/Nodes/MidiCCTriggerNodeTests.cpp
--- a/Engine/Plugins/Runtime/Harmonix/Source/HarmonixMetasoundTests/Private/Nodes/MidiCCTriggerNodeTests.cpp
+++ b/Engine/Plugins/Runtime/Harmonix/Source/HarmonixMetasoundTests/Private/Nodes/MidiCCTriggerNodeTests.cpp
@@ -2,128 +2,20 @@
 
 #include "NodeTestGraphBuilder.h"
 #include "HarmonixDsp/AudioBuffer.h"
-#include "HarmonixMetasound/Common.h"
 #include "HarmonixMetasound/DataTypes/MidiStream.h"
-#include "HarmonixMetasound/DataTypes/MidiAsset.h"
-#include "HarmonixMetasound/DataTypes/MusicTransport.h"
 #include "Misc/AutomationTest.h"
-#include "HarmonixMidi/MidiMsg.h"
-#include "HarmonixMidi/MidiFile.h"
 #include "HarmonixMetasound/DataTypes/MidiControllerID.h"
 #include "HarmonixMetasound/Nodes/MidiCCTriggerNode.h"
 
 #if WITH_DEV_AUTOMATION_TESTS
-/**
- * Helper enum class for evaluating different 
- * blocks containing different Midi Events in a Midi file
- */
-UENUM()
-enum class EMidiEventTestingBlock : uint8
-{
-	ZeroBlock,
-	MidiChannelMismatchBlock,
-	MidiTrackMismatchBlock,
-	MidiEventTypeMismatchBlock,
-	MidiControllerIDMismatchBlock,
-	Count
-};
 
 namespace HarmonixMetasoundTests::MidiCCTriggerNode
 {
-	//expected values for tests
-	constexpr int32 ExpectedCCInt32ValueOnPass = 127;
-	constexpr int32 ExpectedCCInt32ValueNotPass = 0;
-	constexpr float ExpectedCCFloatValueOnPass = 1.0f;
-	constexpr float ExpectedCCFloatValueNotPass = 0.0f;
-	constexpr int32 ExpectedTriggerCountOnPass = 1;
-	constexpr int32 ExpectedTriggerCountNotPass = 0;
-	constexpr int32 ExpectedTriggerFrameIndexOnPass = 0;
-	constexpr int32 ExpectedTriggerFrameIndexNotPass = -1;
-
-	//testing values
-	constexpr int32 TestInputChannel = 1;
-	constexpr int32 TestDefaultTrackIndex = 1;
-	constexpr uint8 DefaultTestCCValue = 127;
-	constexpr uint8 DifferentTestCCValue = 100;
-	constexpr uint8 TestControllerIDShouldPass = static_cast<uint8>(EStdMidiControllerID::Hold);
-	constexpr uint8 TestControllerIDShouldNotPass = static_cast<uint8>(EStdMidiControllerID::Hold2);
-
-	constexpr int32 DefaultTrackNumber = 0;
-	constexpr int32 DifferentTrackNumber = 1;
-	
-	//time frames of midi events to be added to midi file for testing
-	//for block 2 - 4 (block 1 starts on tick 0, 0.0ms)
-	//each block is 480 samples / 10ms long
-	constexpr float TestingMsBlock2 = 15.0f;
-	constexpr float TestingMsBlock3 = 25.0f;
-	constexpr float TestingMsBlock4 = 35.0f;
-
-	UMidiFile* BuildMidiFile()
-	{
-		UMidiFile* TheMidiFile = NewObject<UMidiFile>();
-
-		//add tempo info bpm = 120
-		TheMidiFile->GetSongMaps()->GetTempoMap().AddTempoInfoPoint(MidiConstants::BPMToMidiTempo(120), 0);
-
-		//Add two midi tracks for testing
-		FMidiTrack* CCTrackDefault = TheMidiFile->AddTrack("CCEventTestTrackDefault");//track 1
-		FMidiTrack* CCTrackDifferent = TheMidiFile->AddTrack("CCEventTestTrackDifferent");//track 2
-
-		/**
-		 * Midi Events for Testing Block 1:
-		 * Add 2 CC events at tick 0 with different channel information and different cc value for testing
-		 */
-		FMidiEvent CCEventChannelMatch(0, FMidiMsg(MidiConstants::kControl, TestControllerIDShouldPass, DefaultTestCCValue));
-		//change channel information to 2 by modifying the lowest 4 bits (should not pass)
-		FMidiEvent CCEventChannelMismatch(0, FMidiMsg(MidiConstants::kControl + 0x01, TestControllerIDShouldPass, DifferentTestCCValue));
-		//add both events to track 1 (default)
-		CCTrackDefault->AddEvent(CCEventChannelMatch);
-		CCTrackDefault->AddEvent(CCEventChannelMismatch);
-
-		/**
-		 * Midi Events for Testing Block 2:
-		 * Add 2 CC events to 2 midi tracks at the same tick, default is track 1
-		 */
-		FMidiEvent CCEventTrackMatch(TheMidiFile->GetSongMaps()->MsToTick(TestingMsBlock2), 
-			FMidiMsg(MidiConstants::kControl, TestControllerIDShouldPass, DefaultTestCCValue));
-		CCTrackDefault->AddEvent(CCEventTrackMatch);//added to track 1
-		//add the other event to track to, which should NOT pass through to output in this test, 
-		//changing the control change value 
-		FMidiEvent CCEventTrackMismatch(TheMidiFile->GetSongMaps()->MsToTick(TestingMsBlock2),
-			FMidiMsg(MidiConstants::kControl, TestControllerIDShouldPass, DifferentTestCCValue));
-		CCTrackDifferent->AddEvent(CCEventTrackMismatch);//added to track 2
-
-		/**
-		* Midi Events for Testing Block 3:
-		* Add 1 event that has a different controller ID (should NOT pass)
-		*/
-		FMidiEvent CCEventControllerIDMismatch(TheMidiFile->GetSongMaps()->MsToTick(TestingMsBlock3),
-			FMidiMsg(MidiConstants::kControl, TestControllerIDShouldNotPass, DefaultTestCCValue));
-		CCTrackDefault->AddEvent(CCEventControllerIDMismatch);
-
-		/**
-		* Midi Events for Testing Block 4:
-		* Add 1 event that has a different type (should NOT pass)
-		*/
-		FMidiEvent CCEventTypeMismatch(TheMidiFile->GetSongMaps()->MsToTick(TestingMsBlock4),
-			FMidiMsg::CreateNoteOn(0, 64, DefaultTestCCValue));
-		CCTrackDefault->AddEvent(CCEventTypeMismatch);
-
-		//call TracksChanged() to update SongLengthData
-		TheMidiFile->TracksChanged();
-
-		//conform midi file length to an integer bar length
-		TheMidiFile->ConformMidiFileLength(EMidiFileLengthConformOption::RoundUp);
-
-		return TheMidiFile;
-	}
-
 	using GraphBuilder = Metasound::Test::FNodeTestGraphBuilder;
 	using namespace Metasound;
 	using namespace Metasound::Frontend;
 	using namespace HarmonixMetasound;
 	using namespace Metasound::Test;
-	using namespace HarmonixMetasound::Nodes::MidiCCTriggerNode;
 	
 	IMPLEMENT_SIMPLE_AUTOMATION_TEST(
 		FMidiCCTriggerNodeTestBasic,
@@ -131,129 +23,113 @@ namespace HarmonixMetasoundTests::MidiCCTriggerNode
 		EAutomationTestFlags::EditorContext | EAutomationTestFlags::EngineFilter)
 		bool FMidiCCTriggerNodeTestBasic::RunTest(const FString&)
 	{
-		//************************************************************************************************************************
-		// Build the graph.
-		//************************************************************************************************************************
-		FNodeTestGraphBuilder Builder;
-		Builder.AddOutput("AudioOut", GetMetasoundDataTypeName<FAudioBuffer>());
-		//transport
-		Frontend::FNodeHandle TransportNode = Builder.AddNode({ HarmonixMetasound::HarmonixNodeNamespace, "TriggerToTransport", "" }, 0);
-		UTEST_TRUE("Transport Node Created", TransportNode.Get().IsValid());
-		// hoist up play trigger input
-		Builder.AddAndConnectDataReferenceInput(TransportNode, CommonPinNames::Inputs::TransportPlayName, GetMetasoundDataTypeName<FTrigger>());
-		// midi player
-		Frontend::FNodeHandle MidiPlayerNode = Builder.AddNode({ HarmonixMetasound::HarmonixNodeNamespace, "MidiPlayer", "" }, 0);
-		UTEST_TRUE("MidiPlayer Node Created", MidiPlayerNode.Get().IsValid());
-		// hoist up midi asset input
-		Builder.AddAndConnectDataReferenceInput(MidiPlayerNode, CommonPinNames::Inputs::MidiFileAssetName, GetMetasoundDataTypeName<FMidiAsset>());
-		//connect transport to midi player
-		bool ConnectionSuccess = Builder.ConnectNodes(TransportNode, CommonPinNames::Outputs::TransportName, MidiPlayerNode, CommonPinNames::Inputs::TransportName);
-		UTEST_TRUE("Transport Out connected to MidiPlayer node", ConnectionSuccess);
-		// Midi CC trigger node to test
-		FNodeHandle CCTriggerNode = Builder.AddNode({ HarmonixMetasound::HarmonixNodeNamespace, "MidiCCTrigger", "" }, 0);
-		UTEST_TRUE("MidiCCTrigger Node Created", CCTriggerNode.Get().IsValid());
+		const TUniquePtr<FMetasoundGenerator> Generator = GraphBuilder::MakeSingleNodeGraph(
+			Nodes::MidiCCTriggerNode::GetClassName(),
+			Nodes::MidiCCTriggerNode::GetCurrentMajorVersion());
+		UTEST_TRUE("Graph successfully built", Generator.IsValid());
 
-		//Midi Channel Number 
-		Builder.AddAndConnectDataReferenceInput(CCTriggerNode, Inputs::MidiChannelNumberName, GetMetasoundDataTypeName<int32>());
-		//Midi Track Number
-		Builder.AddAndConnectDataReferenceInput(CCTriggerNode, Inputs::MidiTrackNumberName, GetMetasoundDataTypeName<int32>());
-		//Midi Controller ID input 
-		Builder.AddAndConnectDataReferenceInput(CCTriggerNode, Inputs::InputMidiControllerIDName, GetMetasoundDataTypeName<FEnumStdMidiControllerID>());
+		/*
+		 * Get refs to inputs we want to manipulate
+		 */
+		TOptional<FBoolWriteRef> InputEnable = Generator->GetInputWriteReference<bool>(Nodes::MidiCCTriggerNode::Inputs::EnableName);
+		UTEST_TRUE("Got enable input", InputEnable.IsSet());
 
-		//output Midi Control Change Value (int32)
-		Builder.AddAndConnectDataReferenceOutput(CCTriggerNode, Outputs::OutputControlChangeValueInt32Name, GetMetasoundDataTypeName<int32>());
+		TOptional<FMidiStreamWriteRef> InputMidiStream =
+			Generator->GetInputWriteReference<FMidiStream>(Nodes::MidiCCTriggerNode::Inputs::MidiStreamName);
+		UTEST_TRUE("Got MIDI stream input", InputMidiStream.IsSet());
 
-		//output normalized Midi Control Change Value (float)
-		Builder.AddAndConnectDataReferenceOutput(CCTriggerNode, Outputs::OutputControlChangeValueFloatName, GetMetasoundDataTypeName<float>(), Outputs::OutputControlChangeValueFloatName);
+		TOptional<FEnumStdMidiControllerIDWriteRef> InputControlNumber =
+			Generator->GetInputWriteReference<FEnumStdMidiControllerID>(Nodes::MidiCCTriggerNode::Inputs::InputMidiControllerIDName);
+		UTEST_TRUE("Got control number input", InputControlNumber.IsSet());
 
-		//output trigger
-		Builder.AddAndConnectDataReferenceOutput(CCTriggerNode, Outputs::OutputTriggerName, GetMetasoundDataTypeName<FTrigger>(), Outputs::OutputTriggerName);
-		
-		//connect midi player node to CC Trigger node
-		ConnectionSuccess = Builder.ConnectNodes(MidiPlayerNode, CommonPinNames::Outputs::MidiStreamName, CCTriggerNode, Inputs::MidiStreamName);
-		UTEST_TRUE("MidiPlayer Out successfully connected to Midi CC Trigger node", ConnectionSuccess);
-		
-		//************************************************************************************************************************
-		// Make the generator.
-		//************************************************************************************************************************
-		constexpr int32 NumSamplesPerBlock = 480; // 10ms per block for sample rate = 48000
-		const TUniquePtr<FMetasoundGenerator> Generator = Builder.BuildGenerator(48000, NumSamplesPerBlock);
-		UTEST_TRUE("Graph successfully built", Generator.IsValid());
-		
-		//************************************************************************************************************************
-		// Get reference to outputs we want to check.
-		//************************************************************************************************************************
-		TOptional<FInt32ReadRef> OutputControlChangeValueInt32 = Generator->GetOutputReadReference<int32>(Outputs::OutputControlChangeValueInt32Name);
+		/*
+		 * Get refs to outputs we want to check.
+		 */
+		TOptional<FInt32ReadRef> OutputControlChangeValueInt32 =
+			Generator->GetOutputReadReference<int32>(Nodes::MidiCCTriggerNode::Outputs::OutputControlChangeValueInt32Name);
 		UTEST_TRUE("Output exists", OutputControlChangeValueInt32.IsSet());
 		UTEST_EQUAL("Control Change value (int32) at start", **OutputControlChangeValueInt32, 0);
 
-		TOptional<FFloatReadRef> OutputControlChangeValueFloat = Generator->GetOutputReadReference<float>(Outputs::OutputControlChangeValueFloatName);
-		UTEST_TRUE("Output exists", OutputControlChangeValueFloat.IsSet());
-		UTEST_EQUAL("Control Change value (float) at start", **OutputControlChangeValueFloat, 0.0f);
+		constexpr uint8 Channel = 3;
+		constexpr uint8 ControlNumber = 24;
+		constexpr uint8 Value = 100;
 
-		TOptional<FTriggerReadRef> CCOutputTrigger = Generator->GetOutputReadReference<FTrigger>(Outputs::OutputTriggerName);
-		UTEST_TRUE("Output exists", CCOutputTrigger.IsSet());
-		UTEST_EQUAL("Number of CC triggers at start", (*CCOutputTrigger)->NumTriggeredInBlock(), 0);
-		UTEST_EQUAL("CC Output Trigger sample", (*CCOutputTrigger)->First(), -1);
-		
-		//************************************************************************************************************************
-		// Make initial inputs
-		//************************************************************************************************************************
-		// Make a Midi File to play by Midi Player
-		UMidiFile* MidiFile = BuildMidiFile();
-		TSharedPtr<Audio::IProxyData> MidiFileProxy = MidiFile->CreateProxyData({ "MidiCCTriggerTest" });
-		Generator->SetInputValue<FMidiAsset>(CommonPinNames::Inputs::MidiFileAssetName, MidiFileProxy);
+		**InputControlNumber = FEnumStdMidiControllerID{ ControlNumber };
 		
-		// Set the testing Midi Channel (default = 1)
-		Generator->SetInputValue<int32>(Inputs::MidiChannelNumberName, TestInputChannel);
-		// Set the testing Midi Track (default = 1)
-		Generator->SetInputValue<int32>(Inputs::MidiTrackNumberName, TestDefaultTrackIndex);
-		// Set the testing Midi Controller ID
-		Generator->SetInputValue<FEnumStdMidiControllerID>(Inputs::InputMidiControllerIDName, static_cast<FEnumStdMidiControllerID>(TestControllerIDShouldPass));
-		
-		// Trigger transport
-		Generator->ApplyToInputValue<FTrigger>(CommonPinNames::Inputs::TransportPlayName, [](FTrigger& TriggerIn) { TriggerIn.TriggerFrame(0); });
+		/*
+		 * Add a CC to the input stream and expect to get the value
+		 */
+		{
+			// Add a CC to the input stream
+			const FMidiStreamEvent CCEvent{ static_cast<uint32>(0), FMidiMsg::CreateControlChange(Channel, ControlNumber, Value) };
+			(*InputMidiStream)->AddMidiEvent(CCEvent);
+			
+			// Process
+			FAudioBuffer Buffer{ Generator->OperatorSettings };
+			Generator->OnGenerateAudio(Buffer.GetData(), Buffer.Num());
 
-		//************************************************************************************************************************
-		// Generate 4 10ms blocks for testing various Midi Events in file/stream
-		//************************************************************************************************************************
-		
-		constexpr uint8 FirstTestBlock = static_cast<uint8>(EMidiEventTestingBlock::MidiChannelMismatchBlock);
-		constexpr uint8 NumBlocksToTest = static_cast<uint8>(EMidiEventTestingBlock::Count) - 1;
+			// Check that we got the CC info out the other side
+			UTEST_EQUAL("Got correct value", **OutputControlChangeValueInt32, Value);
 
-		TAudioBuffer<float> Buffer{ Generator->GetNumChannels(), NumSamplesPerBlock, EAudioBufferCleanupMode::Delete};
-		for (uint8 BlockIndex = FirstTestBlock; BlockIndex <= NumBlocksToTest; ++BlockIndex)
+			// Advance the MIDI stream
+			(*InputMidiStream)->PrepareBlock();
+		}
+
+		/*
+		 * Process again and expect the same output values
+		 */
 		{
-			Generator->OnGenerateAudio(Buffer.GetRawChannelData(0), Buffer.GetNumTotalValidSamples());
+			// Process
+			FAudioBuffer Buffer{ Generator->OperatorSettings };
+			Generator->OnGenerateAudio(Buffer.GetData(), Buffer.Num());
+
+			// Check that we got nothing new out the other side
+			UTEST_EQUAL("Got correct value", **OutputControlChangeValueInt32, Value);
+
+			// Advance the MIDI stream
+			(*InputMidiStream)->PrepareBlock();
+		}
+
+		/*
+		 * Add a CC with a different number and value to the input stream, and expect to get the same output values
+		 */
+		{
+			// Add a CC to the input stream
+			const FMidiStreamEvent CCEvent{ static_cast<uint32>(0), FMidiMsg::CreateControlChange(Channel, ControlNumber + 1, Value + 1) };
+			(*InputMidiStream)->AddMidiEvent(CCEvent);
+			
+			// Process
+			FAudioBuffer Buffer{ Generator->OperatorSettings };
+			Generator->OnGenerateAudio(Buffer.GetData(), Buffer.Num());
+
+			// Check that we got the CC info out the other side
+			UTEST_EQUAL("Got correct value", **OutputControlChangeValueInt32, Value);
+
+			// Advance the MIDI stream
+			(*InputMidiStream)->PrepareBlock();
+		}
+
+
+		/*
+		 * Add a CC to the input stream, turn off the node, and expect the same output values
+		 */
+		{
+			// Add a CC to the input stream
+			const FMidiStreamEvent CCEvent{ static_cast<uint32>(0), FMidiMsg::CreateControlChange(Channel, ControlNumber, Value) };
+			(*InputMidiStream)->AddMidiEvent(CCEvent);
+
+			// Turn off the node
+			**InputEnable = false;
 			
-			//Block 1 & 2 each should  have 1 out of 2 events pass through to output with their correlating CC values
-			if (static_cast<EMidiEventTestingBlock>(BlockIndex) == EMidiEventTestingBlock::MidiChannelMismatchBlock ||
-				static_cast<EMidiEventTestingBlock>(BlockIndex) == EMidiEventTestingBlock::MidiTrackMismatchBlock)
-			{
-				//check output CC values
-				UTEST_EQUAL("Control Change value (int32) check", **OutputControlChangeValueInt32, ExpectedCCInt32ValueOnPass);
-				UTEST_EQUAL("Control Change value (float) check", **OutputControlChangeValueFloat, ExpectedCCFloatValueOnPass);
+			// Process
+			FAudioBuffer Buffer{ Generator->OperatorSettings };
+			Generator->OnGenerateAudio(Buffer.GetData(), Buffer.Num());
 
-				//check if we have a trigger
-				UTEST_EQUAL("Number of CC Output Triggers", (*CCOutputTrigger)->NumTriggeredInBlock(), ExpectedTriggerCountOnPass);
-				
-				//check trigger sample frame at the first block
-				if (BlockIndex == 1)
-				{
-					UTEST_EQUAL("CC Output Trigger sample", (*CCOutputTrigger)->First(), ExpectedTriggerFrameIndexOnPass);
-				}
-			}
-			else {
-				//Block 3 & 4 each should have 1 event ONLY,and NONE should pass through to output
-				
-				//check output CC value (should both be 0)
-				UTEST_EQUAL("Control Change value (int32) check", **OutputControlChangeValueInt32, ExpectedCCInt32ValueNotPass);
-				UTEST_EQUAL("Control Change value (float) check", **OutputControlChangeValueFloat, ExpectedCCFloatValueNotPass);
+			// Check that we got the CC info out the other side
+			UTEST_EQUAL("Got correct value", **OutputControlChangeValueInt32, Value);
 
-				//no trigger should pass through
-				UTEST_EQUAL("Number of CC Output Triggers", (*CCOutputTrigger)->NumTriggeredInBlock(), ExpectedTriggerCountNotPass);
-				UTEST_EQUAL("CC Output Trigger sample", (*CCOutputTrigger)->First(), ExpectedTriggerFrameIndexNotPass);
-			}
+			// Advance the MIDI stream
+			(*InputMidiStream)->PrepareBlock();
 		}
 
 		return true;
diff --git a/Engine/Plugins/Runtime/Harmonix/Source/HarmonixMidi/Private/HarmonixMidi/MidiMsg.cpp b/Engine/Plugins/Runtime/Harmonix/Source/HarmonixMidi/Private/HarmonixMidi/MidiMsg.cpp
--- a/Engine/Plugins/Runtime/Harmonix/Source/HarmonixMidi/Private/HarmonixMidi/MidiMsg.cpp
+++ b/Engine/Plugins/Runtime/Harmonix/Source/HarmonixMidi/Private/HarmonixMidi/MidiMsg.cpp
@@ -14,6 +14,12 @@ FMidiMsg FMidiMsg::CreateNoteOff(int32 Channel, int32 Note)
 	return FMidiMsg(0x80 | (uint8)(Channel & 0xF), (uint8)Note, 0);
 }
 
+FMidiMsg FMidiMsg::CreateControlChange(uint8 Channel, uint8 ControlNumber, uint8 Value)
+{
+	const uint8 StatusByte = MidiConstants::kControl | (Channel & 0xf);
+	return FMidiMsg{ StatusByte, ControlNumber, Value };
+}
+
 FMidiMsg FMidiMsg::CreateText(uint16 InTextIndex, uint8 InTextType)
 {
 	FMidiMsg Msg;
diff --git a/Engine/Plugins/Runtime/Harmonix/Source/HarmonixMidi/Public/HarmonixMidi/MidiMsg.h b/Engine/Plugins/Runtime/Harmonix/Source/HarmonixMidi/Public/HarmonixMidi/MidiMsg.h
--- a/Engine/Plugins/Runtime/Harmonix/Source/HarmonixMidi/Public/HarmonixMidi/MidiMsg.h
+++ b/Engine/Plugins/Runtime/Harmonix/Source/HarmonixMidi/Public/HarmonixMidi/MidiMsg.h
@@ -32,6 +32,7 @@ public:
 
 	static FMidiMsg CreateNoteOn(int32 Channel, int32 Note, int32 Velocity);
 	static FMidiMsg CreateNoteOff(int32 Channel, int32 Note);
+	static FMidiMsg CreateControlChange(uint8 Channel, uint8 ControlNumber, uint8 Value);
 	static FMidiMsg CreateText(uint16 InTextIndex, uint8 InTextType);
 	static float GetPitchBendFromData(int8 Data1, int8 Data2);
 
@@ -44,6 +45,7 @@ public:
 	bool  IsNoteOff() const { return Type == EType::Std && (MidiConstants::IsNoteOff(Status) || (MidiConstants::IsNoteOn(Status) && Data2 == 0)); }
 	bool  IsAllNotesOff() const { return Type == EType::Runtime && Status == MidiConstants::kRuntimeAllNotesOffStatus; }
 	bool  IsAllNotesKill() const { return Type == EType::Runtime && Status == MidiConstants::kRuntimeAllNotesKillStatus; }
+	bool  IsControlChange() const { return IsStd() && MidiConstants::IsControl(Status); }
 
 	/** Construct a standard (std) midi message */
 	FMidiMsg(uint8 InStatus, uint8 InData1, uint8 InData2);
