delete enum array nonsense

deprecate "struct of array of settings" in favor of a simple "array of settings" .

#tests PIE

diff --git a/Engine/Plugins/Runtime/Harmonix/Source/HarmonixDsp/Private/HarmonixDsp/FusionSampler/FusionPatch.cpp b/Engine/Plugins/Runtime/Harmonix/Source/HarmonixDsp/Private/HarmonixDsp/FusionSampler/FusionPatch.cpp
--- a/Engine/Plugins/Runtime/Harmonix/Source/HarmonixDsp/Private/HarmonixDsp/FusionSampler/FusionPatch.cpp
+++ b/Engine/Plugins/Runtime/Harmonix/Source/HarmonixDsp/Private/HarmonixDsp/FusionSampler/FusionPatch.cpp
@@ -32,7 +32,7 @@ TSharedPtr<Audio::IProxyData> UFusionPatch::CreateProxyData(const Audio::FProxyD
 UFusionPatch::UFusionPatch() : Super()
 {
 	FFusionPatchSettings& DefaultSettings = FusionPatchData.Settings;
-	FAdsrSettings& VolumeAdsr = DefaultSettings.Adsrs.Volume();
+	FAdsrSettings& VolumeAdsr = DefaultSettings.Adsr[0];
 	VolumeAdsr.Target = EAdsrTarget::Volume;
 	VolumeAdsr.Depth = 1.0f;
 	VolumeAdsr.IsEnabled = true;
@@ -183,6 +183,31 @@ void UFusionPatch::Serialize(FArchive& Ar)
 				UE_LOG(LogFusionPatch, Warning, TEXT("Fusion patch (%s) has more than one preset, but presets have been deprecated. Only the first (Default) preset will be loaded"), *GetPathName());
 			}
 		}
+
+		if (Version < FFusionPatchCustomVersion::DeprecateTypedSettingsArray)
+		{
+			for (int32 Idx = 0; Idx < FFusionPatchSettings::kNumAdsrs; ++Idx)
+			{
+				FusionPatchData.Settings.Adsr[Idx] = FusionPatchData.Settings.Adsrs_DEPRECATED.Array[Idx];
+			}
+
+			for (int32 Idx = 0; Idx < FFusionPatchSettings::kNumLfos; ++Idx)
+			{
+				FusionPatchData.Settings.Lfo[Idx] = FusionPatchData.Settings.Lfos_DEPRECATED.Array[Idx];
+			}
+
+			for (int32 Idx = 0; Idx < FFusionPatchSettings::kNumRandomizers; ++Idx)
+			{
+				FusionPatchData.Settings.Randomizer[Idx] = FusionPatchData.Settings.Randomizers_DEPRECATED.Array[Idx];
+			}
+
+			
+			for (int32 Idx = 0; Idx < FFusionPatchSettings::kNumModulators; ++Idx)
+			{
+				FusionPatchData.Settings.Randomizer[Idx] = FusionPatchData.Settings.Randomizers_DEPRECATED.Array[Idx];
+			}
+			
+		}
 	}
 
 	const UStretcherAndPitchShifterFactoryConfig* FactoryConfig = GetDefault<UStretcherAndPitchShifterFactoryConfig>();
diff --git a/Engine/Plugins/Runtime/Harmonix/Source/HarmonixDsp/Private/HarmonixDsp/FusionSampler/FusionPatchCustomVersion.h b/Engine/Plugins/Runtime/Harmonix/Source/HarmonixDsp/Private/HarmonixDsp/FusionSampler/FusionPatchCustomVersion.h
--- a/Engine/Plugins/Runtime/Harmonix/Source/HarmonixDsp/Private/HarmonixDsp/FusionSampler/FusionPatchCustomVersion.h
+++ b/Engine/Plugins/Runtime/Harmonix/Source/HarmonixDsp/Private/HarmonixDsp/FusionSampler/FusionPatchCustomVersion.h
@@ -19,6 +19,8 @@ struct HARMONIXDSP_API FFusionPatchCustomVersion
 
 		DeprecatedUnusedEffectsSettings = 5,
 
+		DeprecateTypedSettingsArray = 6,
+
 		// -----<new versions can be added above this line>-------------------------------------------------
 		VersionPlusOne,
 		LatestVersion = VersionPlusOne - 1
diff --git a/Engine/Plugins/Runtime/Harmonix/Source/HarmonixDsp/Private/HarmonixDsp/FusionSampler/FusionSampler.cpp b/Engine/Plugins/Runtime/Harmonix/Source/HarmonixDsp/Private/HarmonixDsp/FusionSampler/FusionSampler.cpp
--- a/Engine/Plugins/Runtime/Harmonix/Source/HarmonixDsp/Private/HarmonixDsp/FusionSampler/FusionSampler.cpp
+++ b/Engine/Plugins/Runtime/Harmonix/Source/HarmonixDsp/Private/HarmonixDsp/FusionSampler/FusionSampler.cpp
@@ -250,7 +250,7 @@ void FFusionSampler::SetTempo(float InBPM)
 
 	CurrentTempoBPM = InBPM;
 
-	for (ELfoIndex LfoIdx : TEnumRange<ELfoIndex>())
+	for (int32 LfoIdx = 0; LfoIdx < kNumLfos; ++LfoIdx)
 	{
 		LfoSettings[LfoIdx].TempoBPM = InBPM * Speed;
 		Lfos[LfoIdx].UseSettings(&LfoSettings[LfoIdx]);
@@ -533,16 +533,16 @@ bool FFusionSampler::TryKeyOnZone(FMidiVoiceId InVoiceId, uint8 InTriggeredNote,
 
 	// From here on out we use the transposed note to that the correct pitch is played...
 
-	for (uint8 ModulatorIdx = 0; ModulatorIdx < (uint8)EModulatorTarget::Num; ++ModulatorIdx)
+	for (int32 ModulatorIdx = 0; ModulatorIdx < kNumModulators; ++ModulatorIdx)
 	{
-		Randomizer[ModulatorIdx].Modulate(FMath::FRand());
+		Randomizers[ModulatorIdx].Modulate(FMath::FRand());
 	}
 
 
 	float VelocityNorm = (float)InVelocity * sVelocityNormalizer;
-	for (uint8 ModulatorIdx = 0; ModulatorIdx < (uint32)EModulatorTarget::Num; ++ModulatorIdx)
+	for (int32 ModulatorIdx = 0; ModulatorIdx < kNumModulators; ++ModulatorIdx)
 	{
-		VelocityModulator[ModulatorIdx].Modulate(VelocityNorm);
+		VelocityModulators[ModulatorIdx].Modulate(VelocityNorm);
 	}
 
 	float AttackGain = VelocityNorm; // should this really just be linear?
@@ -693,13 +693,13 @@ void FFusionSampler::ApplyPatchSettings()
 	SetFineTuneCents(Settings.FineTuneCents);
 	SetPan(Settings.PannerDetails);
 	KeyzoneSelectMode = Settings.KeyzoneSelectMode;
-
-	for (EAdsrIndex Idx : TEnumRange<EAdsrIndex>())
-	{
-		AdsrSettings[Idx].CopySettings(Settings.Adsrs[Idx]);
-		AdsrSettings[Idx].CopyCurveTables(Settings.Adsrs[Idx]);
-	}
-
+	
+	AdsrVolumeSettings.CopySettings(Settings.Adsr[0]);
+	AdsrVolumeSettings.CopyCurveTables(Settings.Adsr[0]);
+	
+	AdsrAssignableSettings.CopySettings(Settings.Adsr[1]);
+	AdsrAssignableSettings.CopyCurveTables(Settings.Adsr[1]);
+	
 	//-------------------------------------
 	// get the filter info
 	//------------------------------------
@@ -709,9 +709,9 @@ void FFusionSampler::ApplyPatchSettings()
 	// Read in Lfo settings
 	//------------------------------------
 
-	for (ELfoIndex Idx : TEnumRange<ELfoIndex>())
+	for (int32 Idx = 0; Idx < kNumLfos; ++Idx)
 	{
-		LfoSettings[Idx].CopySettings(Settings.Lfos[Idx]);
+		LfoSettings[Idx].CopySettings(Settings.Lfo[Idx]);
 	}
 
 	//------------------------------------
@@ -762,7 +762,7 @@ void FFusionSampler::SetSampleRate(float InSampleRateHz)
 	PitchBendRamper.SetTarget(0.0f);
 	PitchBendRamper.SnapToTarget();
 
-	for (ELfoIndex Idx : TEnumRange<ELfoIndex>())
+	for (int32 Idx = 0; Idx < kNumLfos; ++Idx)
 	{
 		Lfos[Idx].Prepare(InSampleRateHz);
 		Lfos[Idx].UseSettings(&LfoSettings[Idx]);
@@ -816,23 +816,21 @@ void FFusionSampler::ResetPatchRelatedState()
 	FilterSettings.ResetToDefaults();
 
 	// set up the default Adsr setting
+	
+	AdsrVolumeSettings.ResetToDefaults();
+	AdsrVolumeSettings.IsEnabled = true;
+	AdsrVolumeSettings.Target = EAdsrTarget::Volume;
 
-	for (EAdsrIndex Idx : TEnumRange<EAdsrIndex>())
-	{
-		AdsrSettings[Idx].ResetToDefaults();
-	}
-
-	AdsrSettings.Volume().IsEnabled = true;
-	AdsrSettings.Volume().Target = EAdsrTarget::Volume;
-	AdsrSettings.Assignable().IsEnabled = false;
-	AdsrSettings.Assignable().Target = EAdsrTarget::None;
+	AdsrAssignableSettings.ResetToDefaults();
+	AdsrAssignableSettings.IsEnabled = false;
+	AdsrAssignableSettings.Target = EAdsrTarget::FilterFreq;
 
 	// set up the default Lfo settings
-	LfoSettings.Pan().ResetToDefaults();
-	LfoSettings.Pitch().ResetToDefaults();
+	LfoSettings[0].ResetToDefaults();
+	LfoSettings[0].Target = ELfoTarget::Pan;
 
-	LfoSettings.Pan().Target = ELfoTarget::Pan;
-	LfoSettings.Pitch().Target = ELfoTarget::Pitch;
+	LfoSettings[1].ResetToDefaults();
+	LfoSettings[1].Target = ELfoTarget::Pitch;	
 
 	// set up the default randomizer settings
 	for (uint8 ModulatorIdx = 0; ModulatorIdx < (uint8)EModulatorTarget::Num; ++ModulatorIdx)
@@ -843,10 +841,10 @@ void FFusionSampler::ResetPatchRelatedState()
 
 
 	// set up the default randomizer settings
-	for (uint8 ModulatorIdx = 0; ModulatorIdx < (uint8)EModulatorTarget::Num; ++ModulatorIdx)
+	for (int32 ModulatorIdx = 0; ModulatorIdx < kNumModulators; ++ModulatorIdx)
 	{
-		VelocityModulator[ModulatorIdx].Reset();
-		VelocityModulator[ModulatorIdx].SetTarget(&Harmonix::Dsp::Modulators::FModulatorTarget::kDummyTarget);
+		VelocityModulators[ModulatorIdx].Reset();
+		VelocityModulators[ModulatorIdx].SetTarget(&Harmonix::Dsp::Modulators::FModulatorTarget::kDummyTarget);
 	}
 
 
@@ -1153,8 +1151,10 @@ void FFusionSampler::PrepareToProcess(uint32 InNumSamples)
 	PortamentoPitchRamper.Ramp();
 	UpdatePitchBendFactor();
 
-	Lfos.Pan().Advance(static_cast<uint32_t>(InNumSamples * Speed));
-	Lfos.Pitch().Advance(static_cast<uint32_t>(InNumSamples * Speed));
+	for (int32 LfoIdx = 0; LfoIdx < kNumLfos; ++LfoIdx)
+	{
+		Lfos[LfoIdx].Advance(static_cast<uint32_t>(InNumSamples * Speed));
+	}
 }
 
 void FFusionSampler::Process(uint32 InSliceIdx, uint32 InSubSliceIdx, TAudioBuffer<float>& Output)
@@ -1300,7 +1300,7 @@ void FFusionSampler::GetController(MidiConstants::EControllerID InController, in
 		Min = 0.005f;
 		Max = 2.000f;
 		Range = Max - Min;
-		ValueP = AdsrSettings.Volume().ReleaseTime;
+		ValueP = AdsrVolumeSettings.ReleaseTime;
 		ValueF = ((ValueP - Min) / Range);
 		Msb = (int8)(ValueF * 127.0f);
 		break;
@@ -1311,7 +1311,7 @@ void FFusionSampler::GetController(MidiConstants::EControllerID InController, in
 		Min = 0.005f;
 		Max = 2.000f;
 		Range = Max - Min;
-		ValueP = AdsrSettings.Volume().AttackTime;
+		ValueP = AdsrVolumeSettings.AttackTime;
 		ValueF = ((ValueP - Min) / Range);
 		Msb = (int8)(ValueF * 127.0f);
 		break;
@@ -1356,7 +1356,7 @@ void FFusionSampler::GetController(MidiConstants::EControllerID InController, in
 	case MidiConstants::EControllerID::LFO0Frequency:
 	case MidiConstants::EControllerID::LFO1Frequency:
 	{
-		ELfoIndex LfoIdx = (InController == MidiConstants::EControllerID::LFO0Frequency) ? ELfoIndex::Pan : ELfoIndex::Pitch;
+		int32 LfoIdx = (InController == MidiConstants::EControllerID::LFO1Frequency);
 
 		float Freq = LfoSettings[LfoIdx].Freq;
 		Msb = (int8)(gLfoFreqInterp.InverseClamped(Freq) * 127.0f);
@@ -1366,7 +1366,7 @@ void FFusionSampler::GetController(MidiConstants::EControllerID InController, in
 	case MidiConstants::EControllerID::LFO0Depth:
 	case MidiConstants::EControllerID::LFO1Depth:
 	{
-		ELfoIndex LfoIdx = (InController == MidiConstants::EControllerID::LFO0Depth) ? ELfoIndex::Pan : ELfoIndex::Pitch;
+		int32 LfoIdx = (InController == MidiConstants::EControllerID::LFO1Depth);
 
 		Msb = (int8)(LfoSettings[LfoIdx].Depth * 127.0f);
 		break;
@@ -1485,14 +1485,14 @@ void FFusionSampler::SetController(MidiConstants::EControllerID InController, fl
 	}
 	case MidiConstants::EControllerID::Release: 
 	{ 
-		AdsrSettings.Volume().ReleaseTime = InValue; 
-		AdsrSettings.Volume().BuildReleaseTable(); 
+		AdsrVolumeSettings.ReleaseTime = InValue; 
+		AdsrVolumeSettings.BuildReleaseTable(); 
 		break; 
 	}
 	case MidiConstants::EControllerID::Attack: 
 	{ 
-		AdsrSettings.Volume().AttackTime = InValue; 
-		AdsrSettings.Volume().BuildAttackTable(); 
+		AdsrVolumeSettings.AttackTime = InValue; 
+		AdsrVolumeSettings.BuildAttackTable(); 
 		break; 
 	}
 	case MidiConstants::EControllerID::PortamentoSwitch: 
@@ -1528,7 +1528,7 @@ void FFusionSampler::SetController(MidiConstants::EControllerID InController, fl
 	case MidiConstants::EControllerID::LFO0Frequency:
 	case MidiConstants::EControllerID::LFO1Frequency: 
 	{ 
-		ELfoIndex LfoIdx = (InController == MidiConstants::EControllerID::LFO0Frequency) ? ELfoIndex::Pan : ELfoIndex::Pitch;
+		int32 LfoIdx = (InController == MidiConstants::EControllerID::LFO1Frequency);
 
 		if (LfoSettings[LfoIdx].Freq != InValue)
 		{
@@ -1541,7 +1541,7 @@ void FFusionSampler::SetController(MidiConstants::EControllerID InController, fl
 	case MidiConstants::EControllerID::LFO0Depth:
 	case MidiConstants::EControllerID::LFO1Depth: 
 	{ 
-		ELfoIndex LfoIdx = (InController == MidiConstants::EControllerID::LFO0Depth) ? ELfoIndex::Pan : ELfoIndex::Pitch;
+		int32 LfoIdx = (InController == MidiConstants::EControllerID::LFO1Depth);
 		if (LfoSettings[LfoIdx].Depth != InValue)
 		{
 			LfoSettings[LfoIdx].Depth = InValue;
@@ -1944,8 +1944,10 @@ void FFusionSampler::UpdateVoiceLfos()
 		// skip voices not assigned to this channel
 		FFusionVoice* Voice = Node->GetValue();
 		check(Voice);
-		Voice->SetupLfo(0, LfoSettings.Pan());
-		Voice->SetupLfo(1, LfoSettings.Pitch());
+		for (int32 LfoIdx = 0; LfoIdx < kNumLfos; ++LfoIdx)
+		{
+			Voice->SetupLfo(LfoIdx, LfoSettings[LfoIdx]);
+		}
 	}
 }
 
diff --git a/Engine/Plugins/Runtime/Harmonix/Source/HarmonixDsp/Private/HarmonixDsp/FusionSampler/FusionVoice.cpp b/Engine/Plugins/Runtime/Harmonix/Source/HarmonixDsp/Private/HarmonixDsp/FusionSampler/FusionVoice.cpp
--- a/Engine/Plugins/Runtime/Harmonix/Source/HarmonixDsp/Private/HarmonixDsp/FusionSampler/FusionVoice.cpp
+++ b/Engine/Plugins/Runtime/Harmonix/Source/HarmonixDsp/Private/HarmonixDsp/FusionSampler/FusionVoice.cpp
@@ -65,11 +65,13 @@ void FFusionVoice::SetSampleRate(double InSampleRate)
 
 	float fSamplesPerSecond = (float)SamplesPerSecond;
 
-	Adsrs.Assignable().Prepare(fSamplesPerSecond);
-	Adsrs.Volume().Prepare(fSamplesPerSecond);
-	Lfo1.Prepare(fSamplesPerSecond);
-	Lfo2.Prepare(fSamplesPerSecond);
-
+	AdsrAssignable.Prepare(fSamplesPerSecond);
+	AdsrVolume.Prepare(fSamplesPerSecond);
+	for (int32 Idx = 0; Idx < kNumLfos; ++Idx)
+	{
+		Lfo[Idx].Prepare(fSamplesPerSecond);
+	}
+	
 	static float sRampTime = 15.0f;
 	FilterCoefsRamper.SetRampTimeMs(fSamplesPerSecond / AudioRendering::kMicroSliceSize, sRampTime);
 	FilterGainRamper.SetRampTimeMs(fSamplesPerSecond / AudioRendering::kMicroSliceSize, 10.0f);
@@ -110,12 +112,14 @@ bool FFusionVoice::AssignIDs(FFusionSampler* InSampler, const FKeyzoneSettings*
 	MySampler = InSampler;
 	PitchShifter = Shifter;
 
-	Adsrs.Volume().UseSettings(&(MySampler->AdsrSettings.Volume()));
-	Adsrs.Assignable().UseSettings(&(MySampler->AdsrSettings.Assignable()));
+	AdsrVolume.UseSettings(&(MySampler->AdsrVolumeSettings));
+	AdsrAssignable.UseSettings(&(MySampler->AdsrAssignableSettings));
 	VelocityGain = 0.0f;
 
-	Lfo1.UseSettings(&MySampler->LfoSettings[0]);
-	Lfo2.UseSettings(&MySampler->LfoSettings[1]);
+	for (int32 Idx = 0; Idx < kNumLfos; ++Idx)
+	{
+		Lfo[Idx].UseSettings(&MySampler->LfoSettings[Idx]);
+	}
 
 	SetSampleRate(MySampler->GetSamplesPerSecond());
 
@@ -155,13 +159,9 @@ bool FFusionVoice::AssignIDs(FFusionSampler* InSampler, const FKeyzoneSettings*
 
 void FFusionVoice::SetupLfo(uint8 Index, const FLfoSettings& InSettings)
 {
-	check(Index < (uint8)ELfoIndex::Num);
+	check(Index < kNumLfos);
 
-	switch (Index)
-	{
-	case 0: Lfo1.UseSettings(&InSettings); break;
-	case 1: Lfo2.UseSettings(&InSettings); break;
-	}
+	Lfo[Index].UseSettings(&InSettings);
 }
 
 void FFusionVoice::Attack()
@@ -170,22 +170,22 @@ void FFusionVoice::Attack()
 	checkSlow(KeyZone);
 	checkSlow(KeyZone->AudioSample);
 
-	Adsrs.Volume().Attack();
-	Adsrs.Assignable().Attack();
+	AdsrVolume.Attack();
+	AdsrAssignable.Attack();
 	bWaitingForAttack = false;
 	bHasRenderedAnySamples = false;
 }
 
 void FFusionVoice::Release()
 {
-	Adsrs.Volume().Release();
-	Adsrs.Assignable().Release();
+	AdsrVolume.Release();
+	AdsrAssignable.Release();
 }
 
 void FFusionVoice::FastRelease()
 {
-	Adsrs.Volume().FastRelease();
-	Adsrs.Assignable().FastRelease();
+	AdsrVolume.FastRelease();
+	AdsrAssignable.FastRelease();
 }
 
 void FFusionVoice::Kill()
@@ -212,8 +212,8 @@ void FFusionVoice::Kill()
 	}
 	PitchShifter = nullptr;
 
-	Adsrs.Volume().Kill();
-	Adsrs.Assignable().Kill();
+	AdsrVolume.Kill();
+	AdsrAssignable.Kill();
 	SamplePos = 0.0f;
 	bWaitingForAttack = false;
 	PitchOffsetCents = 0.0f;
@@ -276,10 +276,12 @@ void FFusionVoice::SetSampler(FFusionSampler* InSampler)
 	}
 
 	MySampler = InSampler;
-	Adsrs.Volume().UseSettings(&MySampler->AdsrSettings.Volume(), false);
-	Adsrs.Assignable().UseSettings(&MySampler->AdsrSettings.Assignable(), false);
-	Lfo1.UseSettings(&MySampler->LfoSettings[0]);
-	Lfo2.UseSettings(&MySampler->LfoSettings[1]);
+	AdsrVolume.UseSettings(&MySampler->AdsrVolumeSettings, false);
+	AdsrAssignable.UseSettings(&MySampler->AdsrAssignableSettings, false);
+	for (int32 Idx = 0; Idx < kNumLfos; ++Idx)
+	{
+		Lfo[Idx].UseSettings(&MySampler->LfoSettings[Idx]);
+	}
 
 	if (ActiveRenderer)
 	{
@@ -427,22 +429,16 @@ void FFusionVoice::PrepareWithPitchOffsetAndGain(double InPitchOffsetCents, floa
 	checkSlow(KeyZone);
 	checkSlow(KeyZone->AudioSample);
 
-	if (Lfo1.GetSettings()->ShouldRetrigger)
+	for (int32 Idx = 0; Idx < kNumLfos; ++Idx)
 	{
-		Lfo1.Retrigger();
-	}
-	else
-	{
-		Lfo1.SetPhase(MySampler->Lfos[0].GetPhase());
-	}
-
-	if (Lfo2.GetSettings()->ShouldRetrigger)
-	{
-		Lfo2.Retrigger();
-	}
-	else
-	{
-		Lfo2.SetPhase(MySampler->Lfos[1].GetPhase());
+		if (Lfo[Idx].GetSettings()->ShouldRetrigger)
+		{
+			Lfo[Idx].Retrigger();
+		}
+		else
+		{
+			Lfo[Idx].SetPhase(MySampler->Lfos[Idx].GetPhase());
+		}
 	}
 
 	InPitchOffsetCents = FMath::Clamp(InPitchOffsetCents, -kMaxPitchOffsetCents, kMaxPitchOffsetCents);
@@ -501,13 +497,12 @@ void FFusionVoice::AggregatePansSetTargetAndRamp(float InTotalGain, bool InSnap
 	}
 
 	float PanOffset = MySampler->PanSettings.GetBasicPan();
-	if (Lfo1.GetSettings()->Target == ELfoTarget::Pan)
-	{
-		PanOffset += Lfo1.GetValue();
-	}
-	if (Lfo2.GetSettings()->Target == ELfoTarget::Pan)
+	for (int32 Idx = 0; Idx < kNumLfos; ++Idx)
 	{
-		PanOffset += Lfo2.GetValue();
+		if (Lfo[Idx].GetSettings()->Target == ELfoTarget::Pan)
+		{
+			PanOffset += Lfo[Idx].GetValue();
+		}
 	}
 
 	if (KeyZone->Pan.Mode == EPannerMode::LegacyStereo || KeyZone->Pan.Mode == EPannerMode::Stereo)
@@ -550,7 +545,7 @@ bool FFusionVoice::IsWaitingForAttack() const
 
 bool FFusionVoice::IsInUse() const
 {
-	bool IsAdsrIdle = Adsrs.Volume().GetStage() == Harmonix::Dsp::Modulators::EAdsrStage::Idle;
+	bool IsAdsrIdle = AdsrVolume.GetStage() == Harmonix::Dsp::Modulators::EAdsrStage::Idle;
 	bool IsAudioActive = MaxAudioLevel > 0.0f;
 	return !IsAdsrIdle || IsAudioActive;
 }
@@ -561,19 +556,17 @@ float FFusionVoice::ComputeOctaveShift()
 
 	float OutOctaveShift = 0.0f;
 
-	if (Lfo1.GetSettings()->Target == ELfoTarget::FilterFreq)
-	{
-		OutOctaveShift += kMaxOctaveSwing * Lfo1.GetValue();
-	}
-
-	if (Lfo2.GetSettings()->Target == ELfoTarget::FilterFreq)
+	for (int32 Idx = 0; Idx < kNumLfos; ++Idx)
 	{
-		OutOctaveShift += kMaxOctaveSwing * Lfo2.GetValue();
+		if (Lfo[Idx].GetSettings()->Target == ELfoTarget::FilterFreq)
+		{
+			OutOctaveShift += kMaxOctaveSwing * Lfo[Idx].GetValue();
+		}	
 	}
 
-	if (Adsrs.Assignable().GetSettings()->Target == EAdsrTarget::FilterFreq)
+	if (AdsrAssignable.GetSettings()->Target == EAdsrTarget::FilterFreq && AdsrAssignable.GetSettings()->IsEnabled)
 	{
-		OutOctaveShift += 2 * kMaxOctaveSwing * Adsrs.Assignable().GetValue();
+		OutOctaveShift += 2 * kMaxOctaveSwing * AdsrAssignable.GetValue();
 	}
 
 	return OutOctaveShift;
@@ -688,9 +681,11 @@ uint32 FFusionVoice::Process(uint32 InSliceIndex, uint32 InSubsliceIndex, float*
 		return 0;
 	}
 
-	Lfo1.Advance(static_cast<uint32> (InMaxNumSamples * InSpeed));
-	Lfo2.Advance(static_cast<uint32> (InMaxNumSamples * InSpeed));
-
+	for (int32 Idx = 0; Idx < kNumLfos; ++Idx)
+	{
+		Lfo[Idx].Advance(static_cast<uint32> (InMaxNumSamples * InSpeed));
+	}
+	
 	// compute the sample increment (ratio of input samples per output sample)
 	// get the base increment and then apply pitch bend and keyzone's fine tuning
 	double ResampleInc = ResampleRate;
@@ -705,14 +700,14 @@ uint32 FFusionVoice::Process(uint32 InSliceIndex, uint32 InSubsliceIndex, float*
 		SemiTonesBend += PortaDiffInSemitones;
 	}
 
-	if (Lfo1.GetSettings()->Target == ELfoTarget::Pitch)
+	for (int32 Idx = 0; Idx < kNumLfos; ++Idx)
 	{
-		SemiTonesBend += 2.0 * Lfo1.GetValue();
-	}
-	if (Lfo2.GetSettings()->Target == ELfoTarget::Pitch)
-	{
-		SemiTonesBend += 2.0 * Lfo2.GetValue();
+		if (Lfo[Idx].GetSettings()->Target == ELfoTarget::Pitch)
+		{
+			SemiTonesBend += 2.0 * Lfo[Idx].GetValue();
+		}
 	}
+
 	PitchBend *= FMath::Pow(2.0, SemiTonesBend / 12.0);
 	ResampleInc *= PitchBend;
 	ResampleInc *= KeyZone->FineTuneAdjustment;
@@ -723,7 +718,7 @@ uint32 FFusionVoice::Process(uint32 InSliceIndex, uint32 InSubsliceIndex, float*
 	// we should loop (if possible) as long as we are not currently releasing the voice
 	// (This seem hinky to me. Often looping samples DO NOT have a nice exit 
 	// portion after the loop end. In these cases the sample will just stop short. - Buzz)
-	bool IsInReleaseStage = (Adsrs.Volume().GetStage() == EAdsrStage::Release);
+	bool IsInReleaseStage = (AdsrVolume.GetStage() == EAdsrStage::Release);
 
 	// Forcing this to true so that Adsr release actually has samples to work on!
 	bool ShouldHonorLoopPoints = true; // !IsInReleaseStage; 
@@ -869,12 +864,12 @@ uint32 FFusionVoice::Process(uint32 InSliceIndex, uint32 InSubsliceIndex, float*
 			// write interpolated sample data into the output buffer
 			{
 				//TIME_BLOCK(fusion_voice_process_filter);
-				Adsrs.Volume().Advance(static_cast<uint32> (NumToProcessThisPass * InSpeed));
-				Adsrs.Assignable().Advance(static_cast<uint32> (NumToProcessThisPass * InSpeed));
+				AdsrVolume.Advance(static_cast<uint32> (NumToProcessThisPass * InSpeed));
+				AdsrAssignable.Advance(static_cast<uint32> (NumToProcessThisPass * InSpeed));
 			}
 
 			WorkingGain.Lerp(PrevGain, Panner.GetCurrentGainMatrix(), (float)NumProcessed / (float)InMaxNumSamples);
-			WorkingGain *= Adsrs.Volume().GetValue();
+			WorkingGain *= AdsrVolume.GetValue();
 
 			{
 				SamplePos = ActiveRenderer->Render(
@@ -947,11 +942,11 @@ uint32 FFusionVoice::Process(uint32 InSliceIndex, uint32 InSubsliceIndex, float*
 	bool IsPortamentoEnabled = MySampler->GetIsPortamentoEnabled(); 
 	// could be processing an oscillating filter
 	bool IsAudioLevelLow = MaxAudioLevel < 0.0001; 
-	bool IsOutOfSampleData = Adsrs.Volume().GetStage() == EAdsrStage::Idle || SamplePos >= EndOfSampleData;
+	bool IsOutOfSampleData = AdsrVolume.GetStage() == EAdsrStage::Idle || SamplePos >= EndOfSampleData;
 
 	if (IsAudioLevelLow && IsOutOfSampleData)
 	{
-		if (!IsPortamentoEnabled || Adsrs.Volume().GetStage() == EAdsrStage::Idle)
+		if (!IsPortamentoEnabled || AdsrVolume.GetStage() == EAdsrStage::Idle)
 		{
 			Kill();
 		}
diff --git a/Engine/Plugins/Runtime/Harmonix/Source/HarmonixDsp/Private/HarmonixDsp/FusionSampler/Settings/FusionPatchSettings.cpp b/Engine/Plugins/Runtime/Harmonix/Source/HarmonixDsp/Private/HarmonixDsp/FusionSampler/Settings/FusionPatchSettings.cpp
--- a/Engine/Plugins/Runtime/Harmonix/Source/HarmonixDsp/Private/HarmonixDsp/FusionSampler/Settings/FusionPatchSettings.cpp
+++ b/Engine/Plugins/Runtime/Harmonix/Source/HarmonixDsp/Private/HarmonixDsp/FusionSampler/Settings/FusionPatchSettings.cpp
@@ -10,6 +10,8 @@ FFusionPatchSettings::FFusionPatchSettings()
 	, MaxVoices(32)
 	, KeyzoneSelectMode(EKeyzoneSelectMode::Layers)
 {
-	Lfos.Pan().Target = ELfoTarget::Pan;
-	Lfos.Pitch().Target = ELfoTarget::Pitch;
+	Adsr[0].Target = EAdsrTarget::Volume;
+	Adsr[1].Target = EAdsrTarget::FilterFreq;
+	Lfo[0].Target = ELfoTarget::Pan;
+	Lfo[1].Target = ELfoTarget::Pitch;
 }
\ No newline at end of file
diff --git a/Engine/Plugins/Runtime/Harmonix/Source/HarmonixDsp/Public/HarmonixDsp/Containers/EnumArray.h b/Engine/Plugins/Runtime/Harmonix/Source/HarmonixDsp/Public/HarmonixDsp/Containers/EnumArray.h
deleted file mode 100644
diff --git a/Engine/Plugins/Runtime/Harmonix/Source/HarmonixDsp/Public/HarmonixDsp/FusionSampler/FusionSampler.h b/Engine/Plugins/Runtime/Harmonix/Source/HarmonixDsp/Public/HarmonixDsp/FusionSampler/FusionSampler.h
--- a/Engine/Plugins/Runtime/Harmonix/Source/HarmonixDsp/Public/HarmonixDsp/FusionSampler/FusionSampler.h
+++ b/Engine/Plugins/Runtime/Harmonix/Source/HarmonixDsp/Public/HarmonixDsp/FusionSampler/FusionSampler.h
@@ -7,7 +7,6 @@
 #include "HarmonixDsp/Effects/BiquadFilter.h"
 
 #include "HarmonixDsp/Modulators/Settings/AdsrSettings.h"
-#include "HarmonixDsp/Modulators/Settings/ModulatorSettings.h"
 #include "HarmonixDsp/Modulators/Modulator.h"
 #include "HarmonixDsp/Modulators/Lfo.h"
 
@@ -306,17 +305,14 @@ private:
 
 	//---------------------------------------------
 	// Modulators
-	FAdsrSettingsArray AdsrSettings;
-
-	FLfoSettingsArray LfoSettings;
-
-	Harmonix::Dsp::Modulators::FLfoArray Lfos;
-
-	Harmonix::Dsp::Modulators::FModulatorArray Randomizers;
-	Harmonix::Dsp::Modulators::FModulatorArray VelocityModulators;
-
-	Harmonix::Dsp::Modulators::FModulator Randomizer[(uint8)EModulatorTarget::Num];
-	Harmonix::Dsp::Modulators::FModulator VelocityModulator[(uint8)EModulatorTarget::Num];
+	FAdsrSettings AdsrVolumeSettings;
+	FAdsrSettings AdsrAssignableSettings;
+	static constexpr int32 kNumLfos = 2;
+	FLfoSettings LfoSettings[kNumLfos];
+	Harmonix::Dsp::Modulators::FLfo Lfos[kNumLfos];
+	static constexpr int32 kNumModulators = 2;
+	Harmonix::Dsp::Modulators::FModulator Randomizers[kNumModulators];
+	Harmonix::Dsp::Modulators::FModulator VelocityModulators[kNumModulators];
 
 	// soft limit, may be over max if some voices are in release stage
 	uint32 MaxNumVoices = 100;
diff --git a/Engine/Plugins/Runtime/Harmonix/Source/HarmonixDsp/Public/HarmonixDsp/FusionSampler/FusionVoice.h b/Engine/Plugins/Runtime/Harmonix/Source/HarmonixDsp/Public/HarmonixDsp/FusionSampler/FusionVoice.h
--- a/Engine/Plugins/Runtime/Harmonix/Source/HarmonixDsp/Public/HarmonixDsp/FusionSampler/FusionVoice.h
+++ b/Engine/Plugins/Runtime/Harmonix/Source/HarmonixDsp/Public/HarmonixDsp/FusionSampler/FusionVoice.h
@@ -54,9 +54,9 @@ public:
 
 	virtual TSharedPtr<IStretcherAndPitchShifter, ESPMode::ThreadSafe> GetPitchShifter() const { return PitchShifter;  }
 
-	virtual Harmonix::Dsp::Modulators::EAdsrStage GetAdsrStage() const { return Adsrs.Volume().GetStage(); }
+	virtual Harmonix::Dsp::Modulators::EAdsrStage GetAdsrStage() const { return AdsrVolume.GetStage(); }
 
-	uint32 GetAge() const { return Adsrs.Volume().GetAge(); }
+	uint32 GetAge() const { return AdsrVolume.GetAge(); }
 	FMidiVoiceId GetVoiceID() const { return VoiceID; }
 	uint8 GetTriggeredMidiNote() const { return TriggeredMidiNote; }
 
@@ -172,10 +172,10 @@ private:
 
 	TSharedPtr<IStretcherAndPitchShifter, ESPMode::ThreadSafe> PitchShifter = nullptr;
 
-	Harmonix::Dsp::Modulators::FAdsrArray Adsrs;
-	
-	Harmonix::Dsp::Modulators::FLfo Lfo1;
-	Harmonix::Dsp::Modulators::FLfo Lfo2;
+	Harmonix::Dsp::Modulators::FAdsr AdsrVolume;
+	Harmonix::Dsp::Modulators::FAdsr AdsrAssignable;
+	static const int32 kNumLfos = 2;
+	Harmonix::Dsp::Modulators::FLfo Lfo[kNumLfos];
 
 	FPanner Panner;
 
diff --git a/Engine/Plugins/Runtime/Harmonix/Source/HarmonixDsp/Public/HarmonixDsp/FusionSampler/Settings/FusionPatchSettings.h b/Engine/Plugins/Runtime/Harmonix/Source/HarmonixDsp/Public/HarmonixDsp/FusionSampler/Settings/FusionPatchSettings.h
--- a/Engine/Plugins/Runtime/Harmonix/Source/HarmonixDsp/Public/HarmonixDsp/FusionSampler/Settings/FusionPatchSettings.h
+++ b/Engine/Plugins/Runtime/Harmonix/Source/HarmonixDsp/Public/HarmonixDsp/FusionSampler/Settings/FusionPatchSettings.h
@@ -1,7 +1,6 @@
 // Copyright Epic Games, Inc. All Rights Reserved.
 #pragma once
 
-#include "HarmonixDsp/Containers/EnumArray.h"
 #include "HarmonixDsp/PannerDetails.h"
 #include "HarmonixDsp/Effects/Settings/BiquadFilterSettings.h"
 #include "HarmonixDsp/Effects/Settings/BitCrusherSettings.h"
@@ -67,17 +66,33 @@ public:
 	
 	UPROPERTY(EditDefaultsOnly, Category = "Settings")
 	FBiquadFilterSettings Filter;
-	
+
+	static constexpr int32 kNumAdsrs = 2;
 	UPROPERTY(EditDefaultsOnly, Category = "Settings")
-	FAdsrSettingsArray Adsrs;
+	FAdsrSettings Adsr[kNumAdsrs];
 
+	static constexpr int32 kNumLfos = 2;
 	UPROPERTY(EditDefaultsOnly, Category = "Settings")
-	FLfoSettingsArray Lfos;
+	FLfoSettings Lfo[kNumLfos];
 
+	static constexpr int32 kNumRandomizers = 2;
 	UPROPERTY(EditDefaultsOnly, Category = "Settings")
-	FModulatorSettingsArray Randomizers;
+	FModulatorSettings Randomizer[kNumRandomizers];
 
+	static constexpr int32 kNumModulators = 2;
 	UPROPERTY(EditDefaultsOnly, Category = "Settings")
-	FModulatorSettingsArray VelocityModulators;
+	FModulatorSettings VelocityModulator[kNumModulators];
+	
+	UPROPERTY()
+	FAdsrSettingsArray Adsrs_DEPRECATED;
+
+	UPROPERTY()
+	FLfoSettingsArray Lfos_DEPRECATED;
+
+	UPROPERTY()
+	FModulatorSettingsArray Randomizers_DEPRECATED;
+
+	UPROPERTY()
+	FModulatorSettingsArray VelocityModulators_DEPRECATED;
 
 };
\ No newline at end of file
diff --git a/Engine/Plugins/Runtime/Harmonix/Source/HarmonixDsp/Public/HarmonixDsp/Modulators/Adsr.h b/Engine/Plugins/Runtime/Harmonix/Source/HarmonixDsp/Public/HarmonixDsp/Modulators/Adsr.h
--- a/Engine/Plugins/Runtime/Harmonix/Source/HarmonixDsp/Public/HarmonixDsp/Modulators/Adsr.h
+++ b/Engine/Plugins/Runtime/Harmonix/Source/HarmonixDsp/Public/HarmonixDsp/Modulators/Adsr.h
@@ -1,7 +1,6 @@
 // Copyright Epic Games, Inc. All Rights Reserved.
 #pragma once
 
-#include "HarmonixDsp/Containers/EnumArray.h"
 #include "HarmonixDsp/Modulators/Settings/AdsrSettings.h"
 
 namespace Harmonix::Dsp::Modulators
@@ -175,12 +174,4 @@ private:
 
 	EAdsrMode Mode = EAdsrMode::MinUp;
 };
-
-struct FAdsrArray : TEnumArray<EAdsrIndex, FAdsr>
-{
-public:
-	ENUM_PROPERTY(Volume);
-	ENUM_PROPERTY(Assignable);
-};
-
 }
\ No newline at end of file
diff --git a/Engine/Plugins/Runtime/Harmonix/Source/HarmonixDsp/Public/HarmonixDsp/Modulators/Lfo.h b/Engine/Plugins/Runtime/Harmonix/Source/HarmonixDsp/Public/HarmonixDsp/Modulators/Lfo.h
--- a/Engine/Plugins/Runtime/Harmonix/Source/HarmonixDsp/Public/HarmonixDsp/Modulators/Lfo.h
+++ b/Engine/Plugins/Runtime/Harmonix/Source/HarmonixDsp/Public/HarmonixDsp/Modulators/Lfo.h
@@ -2,7 +2,6 @@
 #pragma once
 
 #include "Math/Interval.h"
-#include "HarmonixDsp/Containers/EnumArray.h"
 #include "HarmonixDsp/Modulators/Settings/LfoSettings.h"
 
 DECLARE_LOG_CATEGORY_EXTERN(LogHarmonixDsp_Lfo, Log, All);
@@ -103,10 +102,4 @@ private:
 	float  TempoWhenCalculated;
 };
 
-struct FLfoArray : TEnumArray<ELfoIndex, FLfo>
-{
-	ENUM_PROPERTY(Pan);
-	ENUM_PROPERTY(Pitch);
-};
-
 }
\ No newline at end of file
diff --git a/Engine/Plugins/Runtime/Harmonix/Source/HarmonixDsp/Public/HarmonixDsp/Modulators/Modulator.h b/Engine/Plugins/Runtime/Harmonix/Source/HarmonixDsp/Public/HarmonixDsp/Modulators/Modulator.h
--- a/Engine/Plugins/Runtime/Harmonix/Source/HarmonixDsp/Public/HarmonixDsp/Modulators/Modulator.h
+++ b/Engine/Plugins/Runtime/Harmonix/Source/HarmonixDsp/Public/HarmonixDsp/Modulators/Modulator.h
@@ -1,9 +1,7 @@
 // Copyright Epic Games, Inc. All Rights Reserved.
 #pragma once
 
-#include "HarmonixDsp/Containers/EnumArray.h"
-
-#include "HarmonixDsp/Modulators/Settings/ModulatorSettings.h"
+#include "Math/Interval.h"
 
 namespace Harmonix::Dsp::Modulators
 {
@@ -68,12 +66,4 @@ protected:
 	TInterval<float> Range;
 };
 
-struct FModulatorArray : TEnumArray<EModulatorIndex, FModulator>
-{
-public:
-
-	ENUM_PROPERTY(StartPoint);
-	ENUM_PROPERTY(Pitch);
-};
-
 }
\ No newline at end of file
diff --git a/Engine/Plugins/Runtime/Harmonix/Source/HarmonixDsp/Public/HarmonixDsp/Modulators/Settings/AdsrSettings.h b/Engine/Plugins/Runtime/Harmonix/Source/HarmonixDsp/Public/HarmonixDsp/Modulators/Settings/AdsrSettings.h
--- a/Engine/Plugins/Runtime/Harmonix/Source/HarmonixDsp/Public/HarmonixDsp/Modulators/Settings/AdsrSettings.h
+++ b/Engine/Plugins/Runtime/Harmonix/Source/HarmonixDsp/Public/HarmonixDsp/Modulators/Settings/AdsrSettings.h
@@ -1,15 +1,11 @@
 // Copyright Epic Games, Inc. All Rights Reserved.
 #pragma once
 
-#include "HarmonixDsp/Containers/EnumArray.h"
-#include "Misc/EnumRange.h"
-
 #include "AdsrSettings.generated.h"
 
 UENUM(BlueprintType)
 enum class EAdsrTarget : uint8
 {
-	None        UMETA(Json="none"),
 	Volume      UMETA(Json="volume"),
 	FilterFreq  UMETA(Json="filter_freq"),
 	Num			UMETA(Hidden),
@@ -22,15 +18,14 @@ struct HARMONIXDSP_API FAdsrSettings
 	GENERATED_BODY()
 
 public:
-
 	static const float kMinTimeSec;
+
+	UPROPERTY(EditDefaultsOnly, Category = "Settings")
+	EAdsrTarget Target = EAdsrTarget::Volume;
 	
 	UPROPERTY(EditDefaultsOnly, Category="Settings")
 	bool IsEnabled = false;
 
-	UPROPERTY(EditDefaultsOnly, Category="Settings")
-	EAdsrTarget Target = EAdsrTarget::Volume;
-
 	UPROPERTY(EditDefaultsOnly, Category="Settings", DisplayName = "Attack (Seconds)", Meta = (UIMin = "0", UIMax = "10", ClampMin = "0", ClampMax = "10"))
 	float AttackTime = 0.0f;
 
@@ -70,7 +65,6 @@ public:
 
 	void ResetToDefaults()
 	{
-		Target = EAdsrTarget::Volume;
 		IsEnabled = false;
 		AttackTime = 0.0f;
 		DecayTime = 0.0f;
@@ -152,28 +146,13 @@ private:
 	void BuildCurveTable(float curve, bool down, float* table);
 };
 
-enum class EAdsrIndex : uint8
-{
-	Volume,
-	Assignable,
-	Num
-};
-
-ENUM_RANGE_BY_COUNT(EAdsrIndex, EAdsrIndex::Num);
-
-USTRUCT(BlueprintType)
+USTRUCT()
 struct FAdsrSettingsArray
 {
 	GENERATED_BODY()
 
-public:
-	IMPLEMENT_ENUM_ARRAY(Array, EAdsrIndex, FAdsrSettings)
-
-	ENUM_PROPERTY(Volume);
-	ENUM_PROPERTY(Assignable);
-
-private:
+	static constexpr int32 Num = 2;
 
-	UPROPERTY(EditDefaultsOnly, Category="Settings")
-	FAdsrSettings Array[(uint8)EAdsrIndex::Num];
+	UPROPERTY()
+	FAdsrSettings Array[Num];
 };
diff --git a/Engine/Plugins/Runtime/Harmonix/Source/HarmonixDsp/Public/HarmonixDsp/Modulators/Settings/LfoSettings.h b/Engine/Plugins/Runtime/Harmonix/Source/HarmonixDsp/Public/HarmonixDsp/Modulators/Settings/LfoSettings.h
--- a/Engine/Plugins/Runtime/Harmonix/Source/HarmonixDsp/Public/HarmonixDsp/Modulators/Settings/LfoSettings.h
+++ b/Engine/Plugins/Runtime/Harmonix/Source/HarmonixDsp/Public/HarmonixDsp/Modulators/Settings/LfoSettings.h
@@ -1,9 +1,6 @@
 // Copyright Epic Games, Inc. All Rights Reserved.
 #pragma once
 
-#include "HarmonixDsp/Containers/EnumArray.h"
-#include "Misc/EnumRange.h"
-
 #include "LfoSettings.generated.h"
 
 UENUM(BlueprintType)
@@ -35,15 +32,16 @@ struct HARMONIXDSP_API FLfoSettings
 	GENERATED_BODY()
 
 public:
+	
+	UPROPERTY(EditDefaultsOnly, Category = "Settings")
+	ELfoTarget Target = ELfoTarget::None;
+	
 	UPROPERTY(EditDefaultsOnly, Category = "Settings")
 	bool IsEnabled = false;
 	
 	UPROPERTY(EditDefaultsOnly, Category = "Settings")
 	EWaveShape Shape = EWaveShape::Sine;
 	
-	UPROPERTY(EditDefaultsOnly, Category = "Settings")
-	ELfoTarget Target = ELfoTarget::None;
-
 	UPROPERTY(EditDefaultsOnly, Category = "Settings")
 	bool ShouldRetrigger = false;
 
@@ -103,28 +101,13 @@ public:
 	}
 };
 
-enum class ELfoIndex : uint8
-{
-	Pan,
-	Pitch,
-	Num
-};
-
-ENUM_RANGE_BY_COUNT(ELfoIndex, ELfoIndex::Num);
-
-USTRUCT(BlueprintType)
+USTRUCT()
 struct FLfoSettingsArray
 {
 	GENERATED_BODY()
 
-public:
-	IMPLEMENT_ENUM_ARRAY(Array, ELfoIndex, FLfoSettings)
-
-	ENUM_PROPERTY(Pan);
-	ENUM_PROPERTY(Pitch);
-
-private:
-
-	UPROPERTY(EditDefaultsOnly, Category="Settings")
-	FLfoSettings Array[(uint8)ELfoIndex::Num];
+	static constexpr int32 Num = 2;
+	
+	UPROPERTY()
+	FLfoSettings Array[Num];
 };
\ No newline at end of file
diff --git a/Engine/Plugins/Runtime/Harmonix/Source/HarmonixDsp/Public/HarmonixDsp/Modulators/Settings/ModulatorSettings.h b/Engine/Plugins/Runtime/Harmonix/Source/HarmonixDsp/Public/HarmonixDsp/Modulators/Settings/ModulatorSettings.h
--- a/Engine/Plugins/Runtime/Harmonix/Source/HarmonixDsp/Public/HarmonixDsp/Modulators/Settings/ModulatorSettings.h
+++ b/Engine/Plugins/Runtime/Harmonix/Source/HarmonixDsp/Public/HarmonixDsp/Modulators/Settings/ModulatorSettings.h
@@ -1,7 +1,6 @@
 // Copyright Epic Games, Inc. All Rights Reserved.
 #pragma once
 
-#include "HarmonixDsp/Containers/EnumArray.h"
 #include "Misc/EnumRange.h"
 
 #include "ModulatorSettings.generated.h"
@@ -21,7 +20,6 @@ struct HARMONIXDSP_API FModulatorSettings
 	GENERATED_BODY()
 
 public:
-
 	FModulatorSettings();
 
 	UPROPERTY(EditDefaultsOnly, Category="Harmonix|DSP")
@@ -32,32 +30,14 @@ public:
 
 	UPROPERTY(EditDefaultsOnly, Category = "Harmonix|DSP", Meta = (UIMin = "0.0", UIMax = "1.0", ClampMin = "0.0", ClampMax = "1.0"))
 	float Depth = 0.5f;
-
-};
-
-enum class EModulatorIndex : uint8
-{
-	StartPoint,
-	Pitch,
-	Num
 };
 
-ENUM_RANGE_BY_COUNT(EModulatorIndex, EModulatorIndex::Num);
-
-USTRUCT(BlueprintType)
+USTRUCT()
 struct FModulatorSettingsArray
 {
 	GENERATED_BODY()
 
-public:
-
-	IMPLEMENT_ENUM_ARRAY(Array, EModulatorIndex, FModulatorSettings)
-
-	ENUM_PROPERTY(StartPoint);
-	ENUM_PROPERTY(Pitch);
-
-private:
-
-	UPROPERTY(EditDefaultsOnly, Category="Settings")
-	FModulatorSettings Array[(uint8)EModulatorIndex::Num];
+	static constexpr int32 Num = 2;
+	UPROPERTY()
+	FModulatorSettings Array[Num];
 };
\ No newline at end of file
diff --git a/Engine/Plugins/Runtime/Harmonix/Source/HarmonixDspEditor/Private/SettingsJsonImporter.cpp b/Engine/Plugins/Runtime/Harmonix/Source/HarmonixDspEditor/Private/SettingsJsonImporter.cpp
--- a/Engine/Plugins/Runtime/Harmonix/Source/HarmonixDspEditor/Private/SettingsJsonImporter.cpp
+++ b/Engine/Plugins/Runtime/Harmonix/Source/HarmonixDspEditor/Private/SettingsJsonImporter.cpp
@@ -147,12 +147,14 @@ bool FSettingsJsonImporter::TryParseJson(TSharedPtr<FJsonObject> JsonObj, FFusio
 			}
 		}
 	}
-
-	{
-	int AdsrIdx = -1;
+	
+	int32 AdsrIdx = -1;
 	for (TSharedPtr<FJsonValue> ArrayValue : IterField(JsonObj, "adsrs"))
 	{
 		++AdsrIdx;
+		if (AdsrIdx >= FFusionPatchSettings::kNumAdsrs)
+			break;
+	
 		TSharedPtr<FJsonObject> AdsrObj;
 		if (!TryGetObject(ArrayValue, AdsrObj))
 			continue;
@@ -161,18 +163,19 @@ bool FSettingsJsonImporter::TryParseJson(TSharedPtr<FJsonObject> JsonObj, FFusio
 		if (!TryGetObjectField(AdsrObj, "adsr", AdsrValue))
 			continue;
 
-		FAdsrSettings& AdsrSettings = AdsrIdx == 0 ? PatchSettings.Adsrs.Volume() : PatchSettings.Adsrs.Assignable();
+		FAdsrSettings& AdsrSettings = PatchSettings.Adsr[AdsrIdx];
 		if (!TryParseJson(AdsrValue, AdsrSettings))
 			continue;
-
+	
 		// old defaults
-		if (AdsrIdx == 0 && !AdsrValue->HasField(TEXT("target")))
+		if (AdsrSettings.Target == EAdsrTarget::Invalid)
 		{
 			AdsrSettings.Target = EAdsrTarget::Volume;
 			AdsrSettings.Depth = 1.0f;
 			AdsrSettings.IsEnabled = true;
 		}
-	}
+	
+		AdsrSettings.Calculate();
 	}
 
 	{
@@ -181,7 +184,7 @@ bool FSettingsJsonImporter::TryParseJson(TSharedPtr<FJsonObject> JsonObj, FFusio
 	{
 		++LfoIdx;
 
-		if (!ensure(LfoIdx < PatchSettings.Lfos.Num()))
+		if (!ensure(LfoIdx < FFusionPatchSettings::kNumLfos))
 			break;
 
 		TSharedPtr<FJsonObject> LfoObj;
@@ -192,7 +195,7 @@ bool FSettingsJsonImporter::TryParseJson(TSharedPtr<FJsonObject> JsonObj, FFusio
 		if (!TryGetObjectField(LfoObj, "lfo", LfoValue))
 			continue;
 
-		FLfoSettings& LfoSettings = PatchSettings.Lfos[LfoIdx];
+		FLfoSettings& LfoSettings = PatchSettings.Lfo[LfoIdx];
 
 		if (!TryParseJson(LfoValue, LfoSettings))
 			continue;
@@ -223,7 +226,7 @@ bool FSettingsJsonImporter::TryParseJson(TSharedPtr<FJsonObject> JsonObj, FFusio
 	for (TSharedPtr<FJsonValue> ArrayValue : IterField(JsonObj, "randomizers"))
 	{
 		++RandomizerIdx;
-		if (!ensure(RandomizerIdx < PatchSettings.Randomizers.Num()))
+		if (!ensure(RandomizerIdx < FFusionPatchSettings::kNumRandomizers))
 			break;
 
 		TSharedPtr<FJsonObject> RandomizerObj;
@@ -234,7 +237,7 @@ bool FSettingsJsonImporter::TryParseJson(TSharedPtr<FJsonObject> JsonObj, FFusio
 		if (!TryGetObjectField(RandomizerObj, "modulator", RandomizerValue))
 			continue;
 
-		TryParseJson(RandomizerValue, PatchSettings.Randomizers[RandomizerIdx]);
+		TryParseJson(RandomizerValue, PatchSettings.Randomizer[RandomizerIdx]);
 	}
 	}
 
@@ -243,7 +246,7 @@ bool FSettingsJsonImporter::TryParseJson(TSharedPtr<FJsonObject> JsonObj, FFusio
 	for (TSharedPtr<FJsonValue> ArrayValue : IterField(JsonObj, "velocity_mods"))
 	{
 		++ModIdx;
-		if (!ensure(ModIdx < PatchSettings.VelocityModulators.Num()))
+		if (!ensure(ModIdx < FFusionPatchSettings::kNumModulators))
 			break;
 
 		TSharedPtr<FJsonObject> VelocityModObj;
@@ -254,7 +257,7 @@ bool FSettingsJsonImporter::TryParseJson(TSharedPtr<FJsonObject> JsonObj, FFusio
 		if (!TryGetObjectField(VelocityModObj, "modulator", VelocityModValue))
 			continue;
 
-		TryParseJson(VelocityModValue, PatchSettings.VelocityModulators[ModIdx]);
+		TryParseJson(VelocityModValue, PatchSettings.VelocityModulator[ModIdx]);
 	}
 	}
 
@@ -279,7 +282,7 @@ bool FSettingsJsonImporter::TryParseJson(TSharedPtr<FJsonObject> JsonObj, FAdsrS
 {
 	// simplify the conversion process :)
 	static TFunction<float(float)> MsToSec = [](float Ms) { return Ms * 0.001f; };
-	TryGetEnumField(JsonObj, "target", AdsrSettings.Target, EAdsrTarget::None);
+	TryGetEnumField(JsonObj, "target", AdsrSettings.Target, EAdsrTarget::Invalid);
 	TryGetNumberField(JsonObj, "depth", AdsrSettings.Depth, 0.0f);
 	TryGetBoolField(JsonObj, "enabled", AdsrSettings.IsEnabled, false);
 	TryGetNumberField(JsonObj, "attack", AdsrSettings.AttackTime, MsToSec, 1.0f);
@@ -290,8 +293,6 @@ bool FSettingsJsonImporter::TryParseJson(TSharedPtr<FJsonObject> JsonObj, FAdsrS
 	TryGetNumberField(JsonObj, "release", AdsrSettings.ReleaseTime, MsToSec, 1.0f);
 	TryGetNumberField(JsonObj, "release_curve", AdsrSettings.ReleaseCurve, 0.0f);
 
-	AdsrSettings.Calculate();
-
 	return true;
 }
 
diff --git a/Engine/Plugins/Runtime/Harmonix/Source/HarmonixMetasound/Private/HarmonixMetasound/Nodes/FusionSamplerNode.cpp b/Engine/Plugins/Runtime/Harmonix/Source/HarmonixMetasound/Private/HarmonixMetasound/Nodes/FusionSamplerNode.cpp
--- a/Engine/Plugins/Runtime/Harmonix/Source/HarmonixMetasound/Private/HarmonixMetasound/Nodes/FusionSamplerNode.cpp
+++ b/Engine/Plugins/Runtime/Harmonix/Source/HarmonixMetasound/Private/HarmonixMetasound/Nodes/FusionSamplerNode.cpp
@@ -228,13 +228,14 @@ namespace HarmonixMetasound
 
 		void UpdatePatchOverrides()
 		{
+			static_assert(FFusionPatchSettings::kNumLfos == 2);
 			if (*Lfo0RateInPin > 0.0f)
 			{
 				SetController(MidiConstants::LFO0Frequency, *Lfo0RateInPin);
 			}
 			else if (FusionPatchDataPtr)
 			{
-				SetController(MidiConstants::LFO0Frequency, FusionPatchDataPtr->GetSettings().Lfos[0].Freq);
+				SetController(MidiConstants::LFO0Frequency, FusionPatchDataPtr->GetSettings().Lfo[0].Freq);
 			}
 			if (*Lfo0DepthInPin >= 0.0f)
 			{
@@ -242,7 +243,7 @@ namespace HarmonixMetasound
 			}
 			else if (FusionPatchDataPtr)
 			{
-				SetController(MidiConstants::LFO0Depth, FusionPatchDataPtr->GetSettings().Lfos[0].Depth);
+				SetController(MidiConstants::LFO0Depth, FusionPatchDataPtr->GetSettings().Lfo[0].Depth);
 			}
 			if (*Lfo1RateInPin > 0.0f)
 			{
@@ -250,7 +251,7 @@ namespace HarmonixMetasound
 			}
 			else if (FusionPatchDataPtr)
 			{
-				SetController(MidiConstants::LFO1Frequency, FusionPatchDataPtr->GetSettings().Lfos[1].Freq);
+				SetController(MidiConstants::LFO1Frequency, FusionPatchDataPtr->GetSettings().Lfo[1].Freq);
 			}
 			if (*Lfo1DepthInPin >= 0.0f)
 			{
@@ -258,7 +259,7 @@ namespace HarmonixMetasound
 			}
 			else if (FusionPatchDataPtr)
 			{
-				SetController(MidiConstants::LFO1Depth, FusionPatchDataPtr->GetSettings().Lfos[1].Depth);
+				SetController(MidiConstants::LFO1Depth, FusionPatchDataPtr->GetSettings().Lfo[1].Depth);
 			}
 
 			float FineTuneCentsTotal = 0;
