[Harmonix] MIDI Channel Filter node

#tests automated tests

diff --git a/Engine/Plugins/Runtime/Harmonix/Source/HarmonixMetasound/Private/HarmonixMetasound/Common.cpp b/Engine/Plugins/Runtime/Harmonix/Source/HarmonixMetasound/Private/HarmonixMetasound/Common.cpp
--- a/Engine/Plugins/Runtime/Harmonix/Source/HarmonixMetasound/Private/HarmonixMetasound/Common.cpp
+++ b/Engine/Plugins/Runtime/Harmonix/Source/HarmonixMetasound/Private/HarmonixMetasound/Common.cpp
@@ -22,7 +22,7 @@ namespace HarmonixMetasound
 		{
 			DEFINE_INPUT_METASOUND_PARAM(Enable,                        "Enable",                             "Enables processing.")
 			DEFINE_INPUT_METASOUND_PARAM(MidiStream,                    "Midi Stream",						  "Midi event stream.")
-			DEFINE_INPUT_METASOUND_PARAM(MidiChannelNumber,             "Midi Channel",                       "Midi channel to render (0 for all/omni).")
+			DEFINE_INPUT_METASOUND_PARAM(MidiChannelNumber,             "Midi Channel",                       "Midi channel to render (1-16, 0 for all).")
 			DEFINE_INPUT_METASOUND_PARAM(MidiChannelFilterSpecifier,    "Midi Channel Filter",                "Midi channels to process. '*' for all, or a string like '1,3,4-8' to identify specific channels.")
 			DEFINE_INPUT_METASOUND_PARAM(MidiTrackNumber,               "Track Number",                       "Track number (1 based).")
 			DEFINE_INPUT_METASOUND_PARAM(MidiTrackIndexFilterSpecifier, "Midi Track Filter",                  "Midi tracks to process. '*' for all, or a string like '1,3,4-8' to identify specific tracks.")
diff --git a/Engine/Plugins/Runtime/Harmonix/Source/HarmonixMetasound/Private/HarmonixMetasound/DataTypes/MidiStream.cpp b/Engine/Plugins/Runtime/Harmonix/Source/HarmonixMetasound/Private/HarmonixMetasound/DataTypes/MidiStream.cpp
--- a/Engine/Plugins/Runtime/Harmonix/Source/HarmonixMetasound/Private/HarmonixMetasound/DataTypes/MidiStream.cpp
+++ b/Engine/Plugins/Runtime/Harmonix/Source/HarmonixMetasound/Private/HarmonixMetasound/DataTypes/MidiStream.cpp
@@ -123,6 +123,43 @@ namespace HarmonixMetasound
 		CopyMidiEvents(SameClockStreams);
 	}
 
+	void FMidiStream::Copy(const FMidiStream& InStream, const FCopyMidiEventsPredicate& Predicate, bool IncludeTransportEvents)
+	{
+		if (IncludeTransportEvents)
+		{
+			CopyTransportEvents(InStream);
+		}
+
+		CopyMidiEvents(InStream, Predicate);
+	}
+
+	void FMidiStream::CopyTransportEvents(const FMidiStream& InStream)
+	{
+		const TArray<FMidiTimestampTransportState>& TransportChanges = InStream.GetTransportChangesInBlock();
+		TransportChangesInBlock.Empty(TransportChanges.Num());
+		for (const FMidiTimestampTransportState& Change : TransportChanges)
+		{
+			AddTransportStateChangeMessage(Change.BlockSampleFrameIndex, Change.TransportState);
+		}
+	}
+
+	void FMidiStream::CopyMidiEvents(const FMidiStream& InStream, const FCopyMidiEventsPredicate& Predicate)
+	{
+		const TArray<FMidiStreamEvent>& MidiEvents = InStream.GetEventsInBlock();
+		if (!ensureAlwaysMsgf(EventsInBlock.Num()==0, TEXT("Existing midi events in midi stream are being destroyed during copy!")))
+		{
+			EventsInBlock.Empty(FMath::Max(MidiEvents.Num(), 32));
+			ActiveVoices.Reset();
+		}
+		for (const FMidiStreamEvent& Event : MidiEvents)
+		{
+			if (Predicate(Event))
+			{
+				AddMidiEvent(Event);
+			}
+		}
+	}
+
 	void FMidiStream::Copy(FMidiStreamReadRef InStream, bool EventsOnly)
 	{
 		if (GetMidiClockSource() && GetMidiClockSource()->Get() != InStream->GetMidiClockSource()->Get())
@@ -211,12 +248,7 @@ namespace HarmonixMetasound
 
 	void FMidiStream::CopyTransportEvents(const FMidiStreamReadRef& InStream)
 	{
-		const TArray<FMidiTimestampTransportState>& TransportChanges = InStream->GetTransportChangesInBlock();
-		TransportChangesInBlock.Empty(TransportChanges.Num());
-		for (auto& Change : TransportChanges)
-		{
-			AddTransportStateChangeMessage(Change.BlockSampleFrameIndex, Change.TransportState);
-		}
+		CopyTransportEvents(*InStream);
 	}
 
 	void FMidiStream::CopyMidiEvents(const TArray<FMidiStreamReadRef>& InStreams)
diff --git a/Engine/Plugins/Runtime/Harmonix/Source/HarmonixMetasound/Private/HarmonixMetasound/MidiOps/MidiChannelFilter.cpp b/Engine/Plugins/Runtime/Harmonix/Source/HarmonixMetasound/Private/HarmonixMetasound/MidiOps/MidiChannelFilter.cpp
new file mode 100644
--- /dev/null
+++ b/Engine/Plugins/Runtime/Harmonix/Source/HarmonixMetasound/Private/HarmonixMetasound/MidiOps/MidiChannelFilter.cpp
@@ -0,0 +1,77 @@
+// Copyright Epic Games, Inc. All Rights Reserved.
+
+#include "HarmonixMetasound/MidiOps/MidiChannelFilter.h"
+
+namespace Harmonix::Midi::Ops
+{
+	void FMidiChannelFilter::Process(const HarmonixMetasound::FMidiStream& InStream, HarmonixMetasound::FMidiStream& OutStream)
+	{
+		OutStream.Copy(
+			InStream,
+			[this](const HarmonixMetasound::FMidiStreamEvent& Event)
+			{
+				if (Event.MidiMessage.IsStd())
+				{
+					const uint8 ChannelIdx = Event.MidiMessage.GetStdChannel();
+					check(ChannelIdx < 16);
+
+					const bool ChannelIsEnabled = GetChannelEnabled(ChannelIdx + 1);
+
+					// If this is a note on or off, keep track of active voices so we can unstick notes if the filter changes later
+					if (ChannelIsEnabled)
+					{
+						if (Event.MidiMessage.IsNoteOn())
+						{
+							ActiveVoices[ChannelIdx].AddUnique(Event.GetVoiceId());
+						}
+						else if (Event.MidiMessage.IsNoteOff())
+						{
+							ActiveVoices[ChannelIdx].Remove(Event.GetVoiceId());
+						}
+					}
+					
+					return ChannelIsEnabled;
+				}
+
+				// Pass non-std messages
+				return true;
+			},
+			true /* include transport events */);
+
+		// If we turned off some channels, send note offs for any active voices in that channel
+		for (uint8 ChannelIdx = 0; ChannelIdx < MidiConstants::kNumChannels; ++ChannelIdx)
+		{
+			if (!GetChannelEnabled(ChannelIdx + 1))
+			{
+				for (const FMidiVoiceId& VoiceId : ActiveVoices[ChannelIdx])
+				{
+					HarmonixMetasound::FMidiStreamEvent NoteOffEvent{ static_cast<uint32>(0), FMidiMsg::CreateNoteOff(ChannelIdx, 0) };
+					NoteOffEvent.SetVoiceId(VoiceId);
+					OutStream.InsertMidiEvent(NoteOffEvent);
+				}
+
+				ActiveVoices[ChannelIdx].Reset();
+			}
+		}
+	}
+
+	void FMidiChannelFilter::SetChannelEnabled(const uint8 Channel, const bool Enabled)
+	{
+		const uint16 ChannelMask = MidiChannelToBitmask(Channel);
+
+		if (Enabled)
+		{
+			EnabledChannels |= ChannelMask;
+		}
+		else
+		{
+			EnabledChannels &= ~ChannelMask;
+		}
+	}
+
+	bool FMidiChannelFilter::GetChannelEnabled(const uint8 Channel) const
+	{
+		const uint16 ChannelMask = MidiChannelToBitmask(Channel);
+		return (EnabledChannels & ChannelMask) == ChannelMask;
+	}
+}
diff --git a/Engine/Plugins/Runtime/Harmonix/Source/HarmonixMetasound/Private/HarmonixMetasound/MidiOps/MidiOps.cpp b/Engine/Plugins/Runtime/Harmonix/Source/HarmonixMetasound/Private/HarmonixMetasound/MidiOps/MidiOps.cpp
new file mode 100644
--- /dev/null
+++ b/Engine/Plugins/Runtime/Harmonix/Source/HarmonixMetasound/Private/HarmonixMetasound/MidiOps/MidiOps.cpp
@@ -0,0 +1,28 @@
+// Copyright Epic Games, Inc. All Rights Reserved.
+
+#include "HarmonixMetasound/MidiOps/MidiOps.h"
+
+DEFINE_LOG_CATEGORY(LogHarmonixMidiOps)
+
+namespace Harmonix::Midi::Ops
+{
+	uint16 MidiChannelToBitmask(const uint8 MidiChannel)
+	{
+		// special case: 0 is all channels
+		if (MidiChannel == 0)
+		{
+			return AllChannelsOn;
+		}
+
+		// bad input: > 16 is no channels
+		if (MidiChannel > 16)
+		{
+			UE_LOG(LogHarmonixMidiOps, Warning, TEXT("Provided a MIDI channel > 16. You probably didn't mean to do that. Failing gracefully, probably."));
+			return AllChannelsOff;
+		}
+
+		// shift to 0-15 range, then make the mask
+		return 1 << (MidiChannel - 1);
+	}
+
+}
diff --git a/Engine/Plugins/Runtime/Harmonix/Source/HarmonixMetasound/Private/HarmonixMetasound/Nodes/MidiChannelFilterNode.cpp b/Engine/Plugins/Runtime/Harmonix/Source/HarmonixMetasound/Private/HarmonixMetasound/Nodes/MidiChannelFilterNode.cpp
new file mode 100644
--- /dev/null
+++ b/Engine/Plugins/Runtime/Harmonix/Source/HarmonixMetasound/Private/HarmonixMetasound/Nodes/MidiChannelFilterNode.cpp
@@ -0,0 +1,183 @@
+// Copyright Epic Games, Inc. All Rights Reserved.
+
+#include "HarmonixMetasound/Nodes/MidiChannelFilterNode.h"
+
+#include "MetasoundExecutableOperator.h"
+#include "MetasoundFacade.h"
+#include "MetasoundNodeRegistrationMacro.h"
+#include "MetasoundStandardNodesCategories.h"
+
+#include "HarmonixMetasound/DataTypes/MidiStream.h"
+#include "HarmonixMetasound/MidiOps/MidiChannelFilter.h"
+
+#define LOCTEXT_NAMESPACE "HarmonixMetaSound"
+
+namespace HarmonixMetasound::Nodes::MidiChannelFilter
+{
+	Metasound::FNodeClassName GetClassName()
+	{
+		return { HarmonixNodeNamespace, TEXT("MidiChannelFilterNode"), ""};
+	}
+
+	int32 GetCurrentMajorVersion()
+	{
+		return 0;
+	}
+
+	namespace Inputs
+	{
+		DEFINE_METASOUND_PARAM_ALIAS(Enable, CommonPinNames::Inputs::Enable);
+		DEFINE_METASOUND_PARAM_ALIAS(MidiStream, CommonPinNames::Inputs::MidiStream);
+		DEFINE_METASOUND_PARAM_ALIAS(Channel, CommonPinNames::Inputs::MidiChannelNumber);
+	}
+
+	namespace Outputs
+	{
+		DEFINE_METASOUND_PARAM_ALIAS(MidiStream, CommonPinNames::Outputs::MidiStream);
+	}
+
+	class FMidiChannelFilterOperator final : public Metasound::TExecutableOperator<FMidiChannelFilterOperator>
+	{
+	public:
+		static const Metasound::FVertexInterface& GetVertexInterface()
+		{
+			auto InitVertexInterface = []() -> Metasound::FVertexInterface
+			{
+				using namespace Metasound;
+				
+				return {
+					FInputVertexInterface
+					{
+						TInputDataVertex<bool>(METASOUND_GET_PARAM_NAME_AND_METADATA(Inputs::Enable), true),
+						TInputDataVertex<FMidiStream>(METASOUND_GET_PARAM_NAME_AND_METADATA(Inputs::MidiStream)),
+						TInputDataVertex<int32>(METASOUND_GET_PARAM_NAME_AND_METADATA(Inputs::Channel), 0)
+					},
+					FOutputVertexInterface
+					{
+						TOutputDataVertex<FMidiStream>(METASOUND_GET_PARAM_NAME_AND_METADATA(Outputs::MidiStream))
+					}
+				};
+			};
+			
+			static const Metasound::FVertexInterface Interface = InitVertexInterface();
+			return Interface;
+		}
+		
+		static const Metasound::FNodeClassMetadata& GetNodeInfo()
+		{
+			auto InitNodeInfo = []() -> Metasound::FNodeClassMetadata
+			{
+				Metasound::FNodeClassMetadata Info;
+				Info.ClassName = GetClassName();
+				Info.MajorVersion = 0;
+				Info.MinorVersion = 1;
+				Info.DisplayName = METASOUND_LOCTEXT("MidiChannelFilterNode_DisplayName", "MIDI Channel Filter");
+				Info.Description = METASOUND_LOCTEXT("MidiChannelFilterNode_Description", "Passes through MIDI events on the specified channel (or all channels)");
+				Info.Author = Metasound::PluginAuthor;
+				Info.PromptIfMissing = Metasound::PluginNodeMissingPrompt;
+				Info.DefaultInterface = GetVertexInterface();
+				Info.CategoryHierarchy = { MetasoundNodeCategories::Harmonix, Metasound::NodeCategories::Music };
+				return Info;
+			};
+
+			static const Metasound::FNodeClassMetadata Info = InitNodeInfo();
+
+			return Info;
+		}
+
+		struct FInputs
+		{
+			Metasound::FBoolReadRef Enabled;
+			FMidiStreamReadRef MidiStream;
+			Metasound::FInt32ReadRef Channel;
+		};
+
+		struct FOutputs
+		{
+			FMidiStreamWriteRef MidiStream;
+		};
+
+		static TUniquePtr<IOperator> CreateOperator(const Metasound::FBuildOperatorParams& InParams, Metasound::FBuildResults& OutResults)
+		{
+			const Metasound::FOperatorSettings& OperatorSettings = InParams.OperatorSettings;
+			const Metasound::FInputVertexInterfaceData& InputData = InParams.InputData;
+			
+			FInputs Inputs
+			{
+				InputData.GetOrCreateDefaultDataReadReference<bool>(Inputs::EnableName, OperatorSettings),
+				InputData.GetOrConstructDataReadReference<FMidiStream>(Inputs::MidiStreamName, OperatorSettings),
+				InputData.GetOrCreateDefaultDataReadReference<int32>(Inputs::ChannelName, OperatorSettings)
+			};
+
+			FOutputs Outputs
+			{
+				FMidiStreamWriteRef::CreateNew(OperatorSettings)
+			};
+
+			return MakeUnique<FMidiChannelFilterOperator>(InParams, MoveTemp(Inputs), MoveTemp(Outputs));
+		}
+
+		FMidiChannelFilterOperator(const Metasound::FBuildOperatorParams& Params, FInputs&& InInputs, FOutputs&& InOutputs)
+		: Inputs(MoveTemp(InInputs))
+		, Outputs(MoveTemp(InOutputs))
+		{
+			Reset(Params);
+		}
+
+		virtual void BindInputs(Metasound::FInputVertexInterfaceData& InVertexData) override
+		{
+			InVertexData.BindReadVertex(Inputs::EnableName, Inputs.Enabled);
+			InVertexData.BindReadVertex(Inputs::MidiStreamName, Inputs.MidiStream);
+			InVertexData.BindReadVertex(Inputs::ChannelName, Inputs.Channel);
+		}
+
+		virtual void BindOutputs(Metasound::FOutputVertexInterfaceData& InVertexData) override
+		{
+			InVertexData.BindReadVertex(Outputs::MidiStreamName, Outputs.MidiStream);
+		}
+
+		void Reset(const FResetParams&)
+		{
+			LastChannelValue = -1;
+		}
+
+		void Execute()
+		{
+			if (const int32 ClampedInputChannel = FMath::Clamp(*Inputs.Channel, 0, 16); ClampedInputChannel != LastChannelValue)
+			{
+				// Reset the filter in case we had "all" enabled
+				Filter.SetChannelEnabled(0, false);
+				// Set the channel (or all) enabled
+				Filter.SetChannelEnabled(ClampedInputChannel, true);
+				// Keep track of the last input value so we don't have to do this unless it changes
+				LastChannelValue = ClampedInputChannel;
+			}
+
+			Outputs.MidiStream->PrepareBlock();
+
+			if (*Inputs.Enabled)
+			{
+				Filter.Process(*Inputs.MidiStream, *Outputs.MidiStream);
+			}
+		}
+
+	private:
+		FInputs Inputs;
+		FOutputs Outputs;
+		Harmonix::Midi::Ops::FMidiChannelFilter Filter;
+		int32 LastChannelValue = -1;
+	};
+
+	class FMidiChannelFilterNode final : public Metasound::FNodeFacade
+	{
+	public:
+		explicit FMidiChannelFilterNode(const Metasound::FNodeInitData& InInitData)
+			: FNodeFacade(InInitData.InstanceName, InInitData.InstanceID, Metasound::TFacadeOperatorClass<FMidiChannelFilterOperator>())
+		{}
+		virtual ~FMidiChannelFilterNode() override = default;
+	};
+
+	METASOUND_REGISTER_NODE(FMidiChannelFilterNode)
+}
+
+#undef LOCTEXT_NAMESPACE
diff --git a/Engine/Plugins/Runtime/Harmonix/Source/HarmonixMetasound/Public/HarmonixMetasound/DataTypes/MidiStream.h b/Engine/Plugins/Runtime/Harmonix/Source/HarmonixMetasound/Public/HarmonixMetasound/DataTypes/MidiStream.h
--- a/Engine/Plugins/Runtime/Harmonix/Source/HarmonixMetasound/Public/HarmonixMetasound/DataTypes/MidiStream.h
+++ b/Engine/Plugins/Runtime/Harmonix/Source/HarmonixMetasound/Public/HarmonixMetasound/DataTypes/MidiStream.h
@@ -97,11 +97,6 @@ namespace HarmonixMetasound
 			MidiClockSource.Reset();
 		}
 
-		using MidiSinkFunction = std::function<void()>;
-		void ExecuteBlock(MidiSinkFunction Destination) const
-		{
-		}
-
 		void PrepareBlock();
 
 		void AddTransportStateChangeMessage(int32 SampleFrameIndexInBlock, EMusicPlayerTransportState State);
@@ -127,6 +122,11 @@ namespace HarmonixMetasound
 			return MidiClockSource.GetPtrOrNull();
 		}
 
+		using FCopyMidiEventsPredicate = TFunctionRef<bool(const FMidiStreamEvent&)>;
+		void Copy(const FMidiStream& InStream, const FCopyMidiEventsPredicate& Predicate, bool IncludeTransportEvents);
+		void CopyTransportEvents(const FMidiStream& InStream);
+		void CopyMidiEvents(const FMidiStream& InStream, const FCopyMidiEventsPredicate& Predicate);
+
 		// Forward declare the read ref for copying other streams
 		using FMidiStreamReadRef = Metasound::TDataReadReference<FMidiStream>;
 
diff --git a/Engine/Plugins/Runtime/Harmonix/Source/HarmonixMetasound/Public/HarmonixMetasound/MidiOps/MidiChannelFilter.h b/Engine/Plugins/Runtime/Harmonix/Source/HarmonixMetasound/Public/HarmonixMetasound/MidiOps/MidiChannelFilter.h
new file mode 100644
--- /dev/null
+++ b/Engine/Plugins/Runtime/Harmonix/Source/HarmonixMetasound/Public/HarmonixMetasound/MidiOps/MidiChannelFilter.h
@@ -0,0 +1,45 @@
+// Copyright Epic Games, Inc. All Rights Reserved.
+
+#pragma once
+
+#include "MidiOps.h"
+#include "HarmonixMetasound/DataTypes/MidiStream.h"
+
+namespace Harmonix::Midi::Ops
+{
+	/**
+	 * Includes MIDI events from enabled channels only
+	 */
+	class HARMONIXMETASOUND_API FMidiChannelFilter
+	{
+	public:
+		/**
+		 * Pass the filtered events from the input stream to the output stream
+		 * @param InStream The MIDI stream to filter
+		 * @param OutStream The filtered MIDI stream
+		 */
+		void Process(const HarmonixMetasound::FMidiStream& InStream, HarmonixMetasound::FMidiStream& OutStream);
+
+		/**
+		 * Enable or disable a channel (or all channels)
+		 * @param Channel The MIDI channel (1-16, 0 is all channels)
+		 * @param Enabled Whether the channel(s) should be enabled
+		 */
+		void SetChannelEnabled(uint8 Channel, bool Enabled);
+
+		/**
+		 * Find out if a channel is enabled
+		 * @param Channel The MIDI channel (1-16, 0 is all channels)
+		 * @return true if the channel is (or all channels are) enabled
+		 */
+		bool GetChannelEnabled(uint8 Channel) const;
+
+	private:
+		// bitfield to store which channels are enabled
+		uint16 EnabledChannels = AllChannelsOff;
+
+		// Keep track of active voices in case we turn off a channel.
+		// If we do turn off a channel, we will immediately send note offs for any active voices in that channel so they don't become stuck.
+		TArray<FMidiVoiceId> ActiveVoices[MidiConstants::kNumChannels];
+	};
+}
diff --git a/Engine/Plugins/Runtime/Harmonix/Source/HarmonixMetasound/Public/HarmonixMetasound/MidiOps/MidiOps.h b/Engine/Plugins/Runtime/Harmonix/Source/HarmonixMetasound/Public/HarmonixMetasound/MidiOps/MidiOps.h
new file mode 100644
--- /dev/null
+++ b/Engine/Plugins/Runtime/Harmonix/Source/HarmonixMetasound/Public/HarmonixMetasound/MidiOps/MidiOps.h
@@ -0,0 +1,20 @@
+// Copyright Epic Games, Inc. All Rights Reserved.
+
+#pragma once
+
+#include "Logging/LogMacros.h"
+
+DECLARE_LOG_CATEGORY_EXTERN(LogHarmonixMidiOps, Log, All)
+
+namespace Harmonix::Midi::Ops
+{
+	constexpr uint16 AllChannelsOn = 0xffff;
+	constexpr uint16 AllChannelsOff = 0;
+	
+	/**
+	 * Get a bitmask from a MIDI channel number (e.g. for use when creating a bitfield for disabling/enabling channels)
+	 * @param MidiChannel The MIDI channel (1-16, or 0 for all)
+	 * @return The bitmask for the channel(s)
+	 */
+	uint16 MidiChannelToBitmask(uint8 MidiChannel);
+}
diff --git a/Engine/Plugins/Runtime/Harmonix/Source/HarmonixMetasound/Public/HarmonixMetasound/Nodes/MidiChannelFilterNode.h b/Engine/Plugins/Runtime/Harmonix/Source/HarmonixMetasound/Public/HarmonixMetasound/Nodes/MidiChannelFilterNode.h
new file mode 100644
--- /dev/null
+++ b/Engine/Plugins/Runtime/Harmonix/Source/HarmonixMetasound/Public/HarmonixMetasound/Nodes/MidiChannelFilterNode.h
@@ -0,0 +1,25 @@
+// Copyright Epic Games, Inc. All Rights Reserved.
+
+#pragma once
+
+#include "MetasoundNodeInterface.h"
+
+#include "HarmonixMetasound/Common.h"
+
+namespace HarmonixMetasound::Nodes::MidiChannelFilter
+{
+	HARMONIXMETASOUND_API Metasound::FNodeClassName GetClassName();
+	HARMONIXMETASOUND_API int32 GetCurrentMajorVersion();
+
+	namespace Inputs
+	{
+		DECLARE_METASOUND_PARAM_EXTERN(Enable);
+		DECLARE_METASOUND_PARAM_EXTERN(MidiStream);
+		DECLARE_METASOUND_PARAM_EXTERN(Channel);
+	}
+
+	namespace Outputs
+	{
+		DECLARE_METASOUND_PARAM_EXTERN(MidiStream);
+	}
+}
diff --git a/Engine/Plugins/Runtime/Harmonix/Source/HarmonixMetasoundTests/Private/MidiOps/MidiChannelFilterTests.cpp b/Engine/Plugins/Runtime/Harmonix/Source/HarmonixMetasoundTests/Private/MidiOps/MidiChannelFilterTests.cpp
new file mode 100644
--- /dev/null
+++ b/Engine/Plugins/Runtime/Harmonix/Source/HarmonixMetasoundTests/Private/MidiOps/MidiChannelFilterTests.cpp
@@ -0,0 +1,96 @@
+// Copyright Epic Games, Inc. All Rights Reserved.
+
+#include "HarmonixMetasound/MidiOps/MidiChannelFilter.h"
+#include "Misc/AutomationTest.h"
+
+#if WITH_DEV_AUTOMATION_TESTS
+
+namespace Harmonix::Midi::Ops::Tests
+{
+	IMPLEMENT_SIMPLE_AUTOMATION_TEST(
+		FMidiChannelFilterSetGetTest,
+		"Harmonix.Midi.Ops.MidiChannelFilter.SetGet",
+		EAutomationTestFlags::EditorContext | EAutomationTestFlags::EngineFilter)
+	bool FMidiChannelFilterSetGetTest::RunTest(const FString&)
+	{
+		FMidiChannelFilter Filter;
+
+		// Default: all channels off
+		UTEST_FALSE("All channels disabled", Filter.GetChannelEnabled(0));
+
+		// Set all channels on and expect them to be enabled
+		Filter.SetChannelEnabled(0, true);
+		UTEST_TRUE("All channels enabled", Filter.GetChannelEnabled(0));
+
+		// Set all channels off and expect them to be disabled
+		Filter.SetChannelEnabled(0, false);
+		UTEST_FALSE("All channels disabled", Filter.GetChannelEnabled(0));
+
+		// Set each channel on and off and check them
+		for (uint8 Channel = 1; Channel <= 16; ++Channel)
+		{
+			Filter.SetChannelEnabled(Channel, true);
+			UTEST_TRUE(FString::Printf(TEXT("Channel %i enabled"), Channel), Filter.GetChannelEnabled(Channel));
+			Filter.SetChannelEnabled(Channel, false);
+			UTEST_FALSE(FString::Printf(TEXT("Channel %i disabled"), Channel), Filter.GetChannelEnabled(Channel));
+		}
+		
+		return true;
+	}
+	
+	IMPLEMENT_SIMPLE_AUTOMATION_TEST(
+		FMidiChannelFilterProcessTest,
+		"Harmonix.Midi.Ops.MidiChannelFilter.Process",
+		EAutomationTestFlags::EditorContext | EAutomationTestFlags::EngineFilter)
+	bool FMidiChannelFilterProcessTest::RunTest(const FString&)
+	{
+		FMidiChannelFilter Filter;
+		
+		Metasound::FOperatorSettings OperatorSettings{ 48000, 100 };
+		HarmonixMetasound::FMidiStream InputStream{ OperatorSettings };
+		HarmonixMetasound::FMidiStream OutputStream{ OperatorSettings };
+
+		// Add some events to the input stream
+		InputStream.AddTransportStateChangeMessage(0, HarmonixMetasound::EMusicPlayerTransportState::Playing);
+		constexpr uint8 NoteChannelIdx = 1;
+		constexpr uint8 NoteNumber = 76;
+		constexpr uint8 NoteVelocity = 67;
+		InputStream.AddMidiEvent({ static_cast<uint32>(0), FMidiMsg::CreateNoteOn(NoteChannelIdx, NoteNumber, NoteVelocity)});
+		{
+			HarmonixMetasound::FMidiStreamEvent NoteOffEvent{ static_cast<uint32>(0), FMidiMsg::CreateNoteOff(NoteChannelIdx, NoteNumber) };
+			NoteOffEvent.BlockSampleFrameIndex = 1;
+			InputStream.AddMidiEvent(NoteOffEvent);
+		}
+
+		// Default: only transport events pass
+		Filter.Process(InputStream, OutputStream);
+		UTEST_EQUAL("Got transport events", OutputStream.GetTransportChangesInBlock().Num(), 1);
+		UTEST_EQUAL("No MIDI events", OutputStream.GetEventsInBlock().Num(), 0);
+
+		// Channel mismatch: no MIDI events pass, transport state hasn't changed, so we won't get another one of those
+		constexpr uint8 NoteChannel = NoteChannelIdx + 1;
+		Filter.SetChannelEnabled(NoteChannel + 1, true);
+		OutputStream.PrepareBlock();
+		Filter.Process(InputStream, OutputStream);
+		UTEST_EQUAL("No MIDI events", OutputStream.GetEventsInBlock().Num(), 0);
+
+		// Channel match: MIDI events pass, transport state hasn't changed, so we won't get another one of those
+		Filter.SetChannelEnabled(NoteChannel, true);
+		OutputStream.PrepareBlock();
+		Filter.Process(InputStream, OutputStream);
+		UTEST_EQUAL("No MIDI events", OutputStream.GetEventsInBlock().Num(), 2);
+		for (const HarmonixMetasound::FMidiStreamEvent& Event : OutputStream.GetEventsInBlock())
+		{
+			UTEST_EQUAL("Channel matches", Event.MidiMessage.GetStdChannel(), NoteChannelIdx);
+			UTEST_EQUAL("Note number matches", Event.MidiMessage.GetStdData1(), NoteNumber);
+			if (Event.MidiMessage.IsNoteOn())
+			{
+				UTEST_EQUAL("Velocity matches", Event.MidiMessage.GetStdData2(), NoteVelocity);
+			}
+		}
+		
+		return true;
+	}
+}
+
+#endif
\ No newline at end of file
diff --git a/Engine/Plugins/Runtime/Harmonix/Source/HarmonixMetasoundTests/Private/Nodes/MidiChannelFilterNodeTests.cpp b/Engine/Plugins/Runtime/Harmonix/Source/HarmonixMetasoundTests/Private/Nodes/MidiChannelFilterNodeTests.cpp
new file mode 100644
--- /dev/null
+++ b/Engine/Plugins/Runtime/Harmonix/Source/HarmonixMetasoundTests/Private/Nodes/MidiChannelFilterNodeTests.cpp
@@ -0,0 +1,86 @@
+// Copyright Epic Games, Inc. All Rights Reserved.
+
+#include "NodeTestGraphBuilder.h"
+#include "HarmonixMetasound/MidiOps/MidiChannelFilter.h"
+#include "HarmonixMetasound/Nodes/MidiChannelFilterNode.h"
+#include "Misc/AutomationTest.h"
+
+#if WITH_DEV_AUTOMATION_TESTS
+
+namespace HarmonixMetasound::Nodes::MidiChannelFilter::Tests
+{
+	IMPLEMENT_SIMPLE_AUTOMATION_TEST(
+	FMidiChannelFilterNodeParityTest,
+	"Harmonix.Metasound.Nodes.MidiChannelFilter.Parity",
+	EAutomationTestFlags::EditorContext | EAutomationTestFlags::EngineFilter)
+	bool FMidiChannelFilterNodeParityTest::RunTest(const FString&)
+	{
+		const TUniquePtr<Metasound::FMetasoundGenerator> Generator =
+			Metasound::Test::FNodeTestGraphBuilder::MakeSingleNodeGraph(GetClassName(), GetCurrentMajorVersion());
+		TOptional<Metasound::TDataReadReference<FMidiStream>> NodeMidiOutput = Generator->GetOutputReadReference<FMidiStream>(Outputs::MidiStreamName);
+		UTEST_TRUE("Got node MIDI output", NodeMidiOutput.IsSet());
+		
+		Harmonix::Midi::Ops::FMidiChannelFilter FilterForComparison;
+		FMidiStream InputStream{ Generator->OperatorSettings };
+		FMidiStream OutputStream{ Generator->OperatorSettings };
+
+		FMidiStreamEvent NoteOnEvent{ static_cast<uint32>(0), FMidiMsg::CreateNoteOn(2, 43, 45) };
+		const uint8 ChannelOneIndexed = NoteOnEvent.MidiMessage.GetStdChannel() + 1;
+		
+		// Channel mismatch, no events output
+		{
+			// Set the wrong channel
+			Generator->SetInputValue<int32>(Inputs::ChannelName, ChannelOneIndexed + 1);
+			FilterForComparison.SetChannelEnabled(ChannelOneIndexed + 1, true);
+			
+			// Add an event to the node input
+			Generator->ApplyToInputValue<FMidiStream>(Inputs::MidiStreamName, [&NoteOnEvent](FMidiStream& Stream)
+			{
+				Stream.AddMidiEvent(NoteOnEvent);
+			});
+
+			// Process the node
+			Metasound::FAudioBuffer Buffer{ Generator->OperatorSettings };
+			Generator->OnGenerateAudio(Buffer.GetData(), Buffer.Num());
+
+			// Process the filter
+			InputStream.AddMidiEvent(NoteOnEvent);
+			FilterForComparison.Process(InputStream, OutputStream);
+
+			// Make sure we got no events
+			UTEST_EQUAL("Got no events (test)", (*NodeMidiOutput)->GetEventsInBlock().Num(), 0);
+			UTEST_EQUAL("Got no events (control)", OutputStream.GetEventsInBlock().Num(), 0);
+
+			// Leave the inputs where they are
+		}
+
+		// Turn on the channel for the note on and expect the message to come out the other end
+		{
+			Generator->SetInputValue<int32>(Inputs::ChannelName, ChannelOneIndexed);
+			FilterForComparison.SetChannelEnabled(0, false); // Reset
+			FilterForComparison.SetChannelEnabled(ChannelOneIndexed, true);
+
+			// Process the node
+			Metasound::FAudioBuffer Buffer{ Generator->OperatorSettings };
+			Generator->OnGenerateAudio(Buffer.GetData(), Buffer.Num());
+
+			// Process the filter
+			FilterForComparison.Process(InputStream, OutputStream);
+
+			const TArray<FMidiStreamEvent>& NodeEvents = (*NodeMidiOutput)->GetEventsInBlock();
+			UTEST_EQUAL("Got one event (test)", NodeEvents.Num(), 1);
+			UTEST_TRUE("Event is note on (test)", NodeEvents[0].MidiMessage.IsNoteOn());
+			const TArray<FMidiStreamEvent>& ControlEvents = OutputStream.GetEventsInBlock();
+			UTEST_EQUAL("Got one event (control)", ControlEvents.Num(), NodeEvents.Num());
+			UTEST_TRUE("Event is note on (control)", ControlEvents[0].MidiMessage.IsNoteOn());
+
+			UTEST_EQUAL("Same channel", NodeEvents[0].MidiMessage.GetStdChannel(), ControlEvents[0].MidiMessage.GetStdChannel());
+			UTEST_EQUAL("Same note number", NodeEvents[0].MidiMessage.GetStdData1(), ControlEvents[0].MidiMessage.GetStdData1());
+			UTEST_EQUAL("Same velocity", NodeEvents[0].MidiMessage.GetStdData2(), ControlEvents[0].MidiMessage.GetStdData2());
+		}
+		
+		return true;
+	}
+}
+
+#endif
\ No newline at end of file
