[Harmonix] MIDI Note Filter Node

#tests automated tests

diff --git a/Engine/Plugins/Runtime/Harmonix/Source/HarmonixMetasound/Private/HarmonixMetasound/MidiOps/MidiNoteFilter.cpp b/Engine/Plugins/Runtime/Harmonix/Source/HarmonixMetasound/Private/HarmonixMetasound/MidiOps/MidiNoteFilter.cpp
new file mode 100644
--- /dev/null
+++ b/Engine/Plugins/Runtime/Harmonix/Source/HarmonixMetasound/Private/HarmonixMetasound/MidiOps/MidiNoteFilter.cpp
@@ -0,0 +1,84 @@
+// Copyright Epic Games, Inc. All Rights Reserved.
+
+#include "HarmonixMetasound/MidiOps/MidiNoteFilter.h"
+
+namespace Harmonix::Midi::Ops
+{
+	void FMidiNoteFilter::Process(const HarmonixMetasound::FMidiStream& InStream, HarmonixMetasound::FMidiStream& OutStream)
+	{
+		OutStream.Copy(
+			InStream,
+			[this](const HarmonixMetasound::FMidiStreamEvent& Event)
+			{
+				// Note ons must pass note number and velocity check
+				if (Event.MidiMessage.IsNoteOn())
+				{
+					const FMidiVoiceId VoiceId = Event.GetVoiceId();
+					const uint8 Channel = Event.MidiMessage.GetStdChannel();
+					const uint8 NoteNumber = Event.MidiMessage.GetStdData1();
+					const uint8 Velocity = Event.MidiMessage.GetStdData2();
+					const bool bPassesFilter =
+						NoteNumber >= MinNoteNumber
+						&& NoteNumber <= MaxNoteNumber
+						&& Velocity >= MinVelocity
+						&& Velocity <= MaxVelocity;
+
+					if (bPassesFilter)
+					{
+						ActiveVoices.FindOrAdd(VoiceId) = { Channel, NoteNumber, Velocity };
+					}
+
+					return bPassesFilter;
+				}
+
+				// Note offs must pass note number check
+				if (Event.MidiMessage.IsNoteOff())
+				{
+					const FMidiVoiceId VoiceId = Event.GetVoiceId();
+					const uint8 NoteNumber = Event.MidiMessage.GetStdData1();
+					const bool bPassesFilter = NoteNumber >= MinNoteNumber && NoteNumber <= MaxNoteNumber;
+
+					if (bPassesFilter)
+					{
+						ActiveVoices.Remove(VoiceId);
+					}
+					
+					return bPassesFilter;
+				}
+
+				// Handle all notes off/kill
+				if (Event.MidiMessage.IsNoteMessage())
+				{
+					ActiveVoices.Reset();
+					return true;
+				}
+
+				return IncludeOtherEvents.Get();
+			},
+			true /* include transport events */);
+
+		// If active notes we're tracking no longer pass the filter, send a note off for each
+		for (auto It = ActiveVoices.CreateIterator(); It; ++It)
+		{
+			const uint8 NoteNumber = It.Value().NoteNumber;
+			const uint8 Velocity = It.Value().Velocity;
+			const bool bPassesFilter =
+				NoteNumber >= MinNoteNumber
+				&& NoteNumber <= MaxNoteNumber
+				&& Velocity >= MinVelocity
+				&& Velocity <= MaxVelocity;
+
+			if (!bPassesFilter)
+			{
+				// Send a note off
+				const uint8 Channel = It.Value().Channel;
+				HarmonixMetasound::FMidiStreamEvent NoteOffEvent{ static_cast<uint32>(0), FMidiMsg::CreateNoteOff(Channel, NoteNumber) };
+				NoteOffEvent.SetVoiceId(It.Key());
+				OutStream.InsertMidiEvent(NoteOffEvent);
+				
+				// Stop tracking
+				It.RemoveCurrent();
+			}
+		}
+	}
+}
diff --git a/Engine/Plugins/Runtime/Harmonix/Source/HarmonixMetasound/Private/HarmonixMetasound/Nodes/MidiNoteFilterNode.cpp b/Engine/Plugins/Runtime/Harmonix/Source/HarmonixMetasound/Private/HarmonixMetasound/Nodes/MidiNoteFilterNode.cpp
new file mode 100644
--- /dev/null
+++ b/Engine/Plugins/Runtime/Harmonix/Source/HarmonixMetasound/Private/HarmonixMetasound/Nodes/MidiNoteFilterNode.cpp
@@ -0,0 +1,199 @@
+// Copyright Epic Games, Inc. All Rights Reserved.
+
+#include "HarmonixMetasound/Nodes/MidiNoteFilterNode.h"
+
+#include "MetasoundExecutableOperator.h"
+#include "MetasoundFacade.h"
+#include "MetasoundNodeRegistrationMacro.h"
+#include "MetasoundStandardNodesCategories.h"
+
+#include "HarmonixMetasound/DataTypes/MidiStream.h"
+#include "HarmonixMetasound/MidiOps/MidiNoteFilter.h"
+
+#define LOCTEXT_NAMESPACE "HarmonixMetaSound"
+
+namespace HarmonixMetasound::Nodes::MidiNoteFilter
+{
+	Metasound::FNodeClassName GetClassName()
+	{
+		return { HarmonixNodeNamespace, TEXT("MidiNoteFilterNode"), ""};
+	}
+
+	int32 GetCurrentMajorVersion()
+	{
+		return 0;
+	}
+
+	namespace Inputs
+	{
+		DEFINE_METASOUND_PARAM_ALIAS(Enable, CommonPinNames::Inputs::Enable);
+		DEFINE_METASOUND_PARAM_ALIAS(MidiStream, CommonPinNames::Inputs::MidiStream);
+		DEFINE_METASOUND_PARAM_ALIAS(MinNoteNumber, CommonPinNames::Inputs::MinMidiNote);
+		DEFINE_METASOUND_PARAM_ALIAS(MaxNoteNumber, CommonPinNames::Inputs::MaxMidiNote);
+		DEFINE_METASOUND_PARAM_ALIAS(MinVelocity, CommonPinNames::Inputs::MinMidiVelocity);
+		DEFINE_METASOUND_PARAM_ALIAS(MaxVelocity, CommonPinNames::Inputs::MaxMidiVelocity);
+		DEFINE_INPUT_METASOUND_PARAM(IncludeOtherEvents, "Include Other Events", "Toggle on to include non-note events in the output.");
+	}
+
+	namespace Outputs
+	{
+		DEFINE_METASOUND_PARAM_ALIAS(MidiStream, CommonPinNames::Outputs::MidiStream);
+	}
+
+	class FMidiNoteFilterOperator final : public Metasound::TExecutableOperator<FMidiNoteFilterOperator>
+	{
+	public:
+		static const Metasound::FVertexInterface& GetVertexInterface()
+		{
+			auto InitVertexInterface = []() -> Metasound::FVertexInterface
+			{
+				using namespace Metasound;
+
+				const Harmonix::Midi::Ops::FMidiNoteFilter FilterForDefaults;
+				
+				return {
+					FInputVertexInterface
+					{
+						TInputDataVertex<bool>(METASOUND_GET_PARAM_NAME_AND_METADATA(Inputs::Enable), true),
+						TInputDataVertex<FMidiStream>(METASOUND_GET_PARAM_NAME_AND_METADATA(Inputs::MidiStream)),
+						TInputDataVertex<int32>(METASOUND_GET_PARAM_NAME_AND_METADATA(Inputs::MinNoteNumber), static_cast<int32>(FilterForDefaults.MinNoteNumber.Get())),
+						TInputDataVertex<int32>(METASOUND_GET_PARAM_NAME_AND_METADATA(Inputs::MaxNoteNumber), static_cast<int32>(FilterForDefaults.MaxNoteNumber.Get())),
+						TInputDataVertex<int32>(METASOUND_GET_PARAM_NAME_AND_METADATA(Inputs::MinVelocity), static_cast<int32>(FilterForDefaults.MinVelocity.Get())),
+						TInputDataVertex<int32>(METASOUND_GET_PARAM_NAME_AND_METADATA(Inputs::MaxVelocity), static_cast<int32>(FilterForDefaults.MaxVelocity.Get())),
+						TInputDataVertex<bool>(METASOUND_GET_PARAM_NAME_AND_METADATA(Inputs::IncludeOtherEvents), FilterForDefaults.IncludeOtherEvents.Get()),
+					},
+					FOutputVertexInterface
+					{
+						TOutputDataVertex<FMidiStream>(METASOUND_GET_PARAM_NAME_AND_METADATA(Outputs::MidiStream))
+					}
+				};
+			};
+			
+			static const Metasound::FVertexInterface Interface = InitVertexInterface();
+			return Interface;
+		}
+		
+		static const Metasound::FNodeClassMetadata& GetNodeInfo()
+		{
+			auto InitNodeInfo = []() -> Metasound::FNodeClassMetadata
+			{
+				Metasound::FNodeClassMetadata Info;
+				Info.ClassName = GetClassName();
+				Info.MajorVersion = 0;
+				Info.MinorVersion = 1;
+				Info.DisplayName = METASOUND_LOCTEXT("MidiNoteFilterNode_DisplayName", "MIDI Note Filter");
+				Info.Description = METASOUND_LOCTEXT("MidiNoteFilterNode_Description", "Passes through MIDI notes based on note number and velocity");
+				Info.Author = Metasound::PluginAuthor;
+				Info.PromptIfMissing = Metasound::PluginNodeMissingPrompt;
+				Info.DefaultInterface = GetVertexInterface();
+				Info.CategoryHierarchy = { MetasoundNodeCategories::Harmonix, Metasound::NodeCategories::Music };
+				return Info;
+			};
+
+			static const Metasound::FNodeClassMetadata Info = InitNodeInfo();
+
+			return Info;
+		}
+
+		struct FInputs
+		{
+			Metasound::FBoolReadRef Enabled;
+			FMidiStreamReadRef MidiStream;
+			Metasound::FInt32ReadRef MinNoteNumber;
+			Metasound::FInt32ReadRef MaxNoteNumber;
+			Metasound::FInt32ReadRef MinVelocity;
+			Metasound::FInt32ReadRef MaxVelocity;
+			Metasound::FBoolReadRef IncludeOtherEvents;
+		};
+
+		struct FOutputs
+		{
+			FMidiStreamWriteRef MidiStream;
+		};
+
+		static TUniquePtr<IOperator> CreateOperator(const Metasound::FBuildOperatorParams& InParams, Metasound::FBuildResults& OutResults)
+		{
+			const Metasound::FOperatorSettings& OperatorSettings = InParams.OperatorSettings;
+			const Metasound::FInputVertexInterfaceData& InputData = InParams.InputData;
+			
+			FInputs Inputs
+			{
+				InputData.GetOrCreateDefaultDataReadReference<bool>(Inputs::EnableName, OperatorSettings),
+				InputData.GetOrConstructDataReadReference<FMidiStream>(Inputs::MidiStreamName, OperatorSettings),
+				InputData.GetOrCreateDefaultDataReadReference<int32>(Inputs::MinNoteNumberName, OperatorSettings),
+				InputData.GetOrCreateDefaultDataReadReference<int32>(Inputs::MaxNoteNumberName, OperatorSettings),
+				InputData.GetOrCreateDefaultDataReadReference<int32>(Inputs::MinVelocityName, OperatorSettings),
+				InputData.GetOrCreateDefaultDataReadReference<int32>(Inputs::MaxVelocityName, OperatorSettings),
+				InputData.GetOrCreateDefaultDataReadReference<bool>(Inputs::IncludeOtherEventsName, OperatorSettings)
+			};
+
+			FOutputs Outputs
+			{
+				FMidiStreamWriteRef::CreateNew(OperatorSettings)
+			};
+
+			return MakeUnique<FMidiNoteFilterOperator>(InParams, MoveTemp(Inputs), MoveTemp(Outputs));
+		}
+
+		FMidiNoteFilterOperator(const Metasound::FBuildOperatorParams& Params, FInputs&& InInputs, FOutputs&& InOutputs)
+		: Inputs(MoveTemp(InInputs))
+		, Outputs(MoveTemp(InOutputs))
+		{
+			Reset(Params);
+		}
+
+		virtual void BindInputs(Metasound::FInputVertexInterfaceData& InVertexData) override
+		{
+			InVertexData.BindReadVertex(Inputs::EnableName, Inputs.Enabled);
+			InVertexData.BindReadVertex(Inputs::MidiStreamName, Inputs.MidiStream);
+			InVertexData.BindReadVertex(Inputs::MinNoteNumberName, Inputs.MinNoteNumber);
+			InVertexData.BindReadVertex(Inputs::MaxNoteNumberName, Inputs.MaxNoteNumber);
+			InVertexData.BindReadVertex(Inputs::MinVelocityName, Inputs.MinVelocity);
+			InVertexData.BindReadVertex(Inputs::MaxVelocityName, Inputs.MaxVelocity);
+			InVertexData.BindReadVertex(Inputs::IncludeOtherEventsName, Inputs.IncludeOtherEvents);
+		}
+
+		virtual void BindOutputs(Metasound::FOutputVertexInterfaceData& InVertexData) override
+		{
+			InVertexData.BindReadVertex(Outputs::MidiStreamName, Outputs.MidiStream);
+		}
+
+		void Reset(const FResetParams&)
+		{
+		}
+
+		void Execute()
+		{
+			Filter.MinNoteNumber = *Inputs.MinNoteNumber;
+			Filter.MaxNoteNumber = *Inputs.MaxNoteNumber;
+			Filter.MinVelocity = *Inputs.MinVelocity;
+			Filter.MaxVelocity = *Inputs.MaxVelocity;
+			Filter.IncludeOtherEvents.Set(*Inputs.IncludeOtherEvents);
+
+			Outputs.MidiStream->PrepareBlock();
+
+			if (*Inputs.Enabled)
+			{
+				Filter.Process(*Inputs.MidiStream, *Outputs.MidiStream);
+			}
+		}
+
+	private:
+		FInputs Inputs;
+		FOutputs Outputs;
+		Harmonix::Midi::Ops::FMidiNoteFilter Filter;
+	};
+
+	class FMidiNoteFilterNode final : public Metasound::FNodeFacade
+	{
+	public:
+		explicit FMidiNoteFilterNode(const Metasound::FNodeInitData& InInitData)
+			: FNodeFacade(InInitData.InstanceName, InInitData.InstanceID, Metasound::TFacadeOperatorClass<FMidiNoteFilterOperator>())
+		{}
+		virtual ~FMidiNoteFilterNode() override = default;
+	};
+
+	METASOUND_REGISTER_NODE(FMidiNoteFilterNode)
+}
+
+#undef LOCTEXT_NAMESPACE
diff --git a/Engine/Plugins/Runtime/Harmonix/Source/HarmonixMetasound/Public/HarmonixMetasound/MidiOps/MidiNoteFilter.h b/Engine/Plugins/Runtime/Harmonix/Source/HarmonixMetasound/Public/HarmonixMetasound/MidiOps/MidiNoteFilter.h
new file mode 100644
--- /dev/null
+++ b/Engine/Plugins/Runtime/Harmonix/Source/HarmonixMetasound/Public/HarmonixMetasound/MidiOps/MidiNoteFilter.h
@@ -0,0 +1,61 @@
+// Copyright Epic Games, Inc. All Rights Reserved.
+
+#pragma once
+
+#include "HarmonixDsp/Parameters/Parameter.h"
+
+#include "HarmonixMetasound/DataTypes/MidiStream.h"
+
+namespace Harmonix::Midi::Ops
+{
+	/**
+	 * Filters MIDI note events based on note number and velocity
+	 */
+	class HARMONIXMETASOUND_API FMidiNoteFilter
+	{
+	public:
+		/**
+		 * Pass the filtered events from the input stream to the output stream
+		 * @param InStream The MIDI stream to filter
+		 * @param OutStream The filtered MIDI stream
+		 */
+		void Process(const HarmonixMetasound::FMidiStream& InStream, HarmonixMetasound::FMidiStream& OutStream);
+
+		/**
+		 * Includes non-note events from the input into the output
+		 */
+		Dsp::Parameters::FBoolParameter IncludeOtherEvents{ true };
+
+		using FMidiDataParam = Dsp::Parameters::TParameter<uint8>;
+
+		/**
+		 * The minimum note number the note can have to be included (inclusive)
+		 */
+		FMidiDataParam MinNoteNumber{ 0, 127, 0 };
+
+		/**
+		 * The maximum note number the note can have to be included (inclusive)
+		 */
+		FMidiDataParam MaxNoteNumber{ 0, 127, 127 };
+
+		/**
+		 * The minimum velocity the note can have to be included (inclusive)
+		 */
+		FMidiDataParam MinVelocity{ 0, 127, 0 };
+
+		/**
+		 * The maximum velocity the note can have to be included (inclusive)
+		 */
+		FMidiDataParam MaxVelocity{ 0, 127, 127 };
+
+	private:
+		// Keep track of active notes so if the filter changes such that they would become stuck, we can unstick them immediately.
+		struct FActiveVoice
+		{
+			uint8 Channel;
+			uint8 NoteNumber;
+			uint8 Velocity;
+		};
+		TMap<FMidiVoiceId, FActiveVoice> ActiveVoices;
+	};
+}
diff --git a/Engine/Plugins/Runtime/Harmonix/Source/HarmonixMetasound/Public/HarmonixMetasound/Nodes/MidiNoteFilterNode.h b/Engine/Plugins/Runtime/Harmonix/Source/HarmonixMetasound/Public/HarmonixMetasound/Nodes/MidiNoteFilterNode.h
new file mode 100644
--- /dev/null
+++ b/Engine/Plugins/Runtime/Harmonix/Source/HarmonixMetasound/Public/HarmonixMetasound/Nodes/MidiNoteFilterNode.h
@@ -0,0 +1,29 @@
+// Copyright Epic Games, Inc. All Rights Reserved.
+
+#pragma once
+
+#include "MetasoundNodeInterface.h"
+
+#include "HarmonixMetasound/Common.h"
+
+namespace HarmonixMetasound::Nodes::MidiNoteFilter
+{
+	HARMONIXMETASOUND_API Metasound::FNodeClassName GetClassName();
+	HARMONIXMETASOUND_API int32 GetCurrentMajorVersion();
+
+	namespace Inputs
+	{
+		DECLARE_METASOUND_PARAM_EXTERN(Enable);
+		DECLARE_METASOUND_PARAM_EXTERN(MidiStream);
+		DECLARE_METASOUND_PARAM_EXTERN(MinNoteNumber);
+		DECLARE_METASOUND_PARAM_EXTERN(MaxNoteNumber);
+		DECLARE_METASOUND_PARAM_EXTERN(MinVelocity);
+		DECLARE_METASOUND_PARAM_EXTERN(MaxVelocity);
+		DECLARE_METASOUND_PARAM_EXTERN(IncludeOtherEvents);
+	}
+
+	namespace Outputs
+	{
+		DECLARE_METASOUND_PARAM_EXTERN(MidiStream);
+	}
+}
diff --git a/Engine/Plugins/Runtime/Harmonix/Source/HarmonixMetasoundTests/Private/MidiOps/MidiNoteFilterTests.cpp b/Engine/Plugins/Runtime/Harmonix/Source/HarmonixMetasoundTests/Private/MidiOps/MidiNoteFilterTests.cpp
new file mode 100644
--- /dev/null
+++ b/Engine/Plugins/Runtime/Harmonix/Source/HarmonixMetasoundTests/Private/MidiOps/MidiNoteFilterTests.cpp
@@ -0,0 +1,120 @@
+// Copyright Epic Games, Inc. All Rights Reserved.
+
+#include "HarmonixMetasound/MidiOps/MidiNoteFilter.h"
+#include "Misc/AutomationTest.h"
+
+#if WITH_DEV_AUTOMATION_TESTS
+
+namespace Harmonix::Midi::Ops::Tests
+{
+	IMPLEMENT_SIMPLE_AUTOMATION_TEST(
+		FMidiNoteFilterProcessTest,
+		"Harmonix.Midi.Ops.MidiNoteFilter.Process",
+		EAutomationTestFlags::EditorContext | EAutomationTestFlags::EngineFilter)
+	bool FMidiNoteFilterProcessTest::RunTest(const FString&)
+	{
+		FMidiNoteFilter Filter;
+		
+		Metasound::FOperatorSettings OperatorSettings{ 48000, 100 };
+		HarmonixMetasound::FMidiStream InputStream{ OperatorSettings };
+		HarmonixMetasound::FMidiStream OutputStream{ OperatorSettings };
+
+		FMidiMsg NoteOnMsg = FMidiMsg::CreateNoteOn(2, 3, 4);
+		HarmonixMetasound::FMidiStreamEvent NoteOnEvent{ static_cast<uint32>(0), NoteOnMsg };
+		
+		// Default: All events pass
+		{
+			FMidiMsg CCMsg = FMidiMsg::CreateControlChange(3, 4, 5);
+			InputStream.AddMidiEvent(NoteOnEvent);
+			InputStream.AddMidiEvent({ static_cast<uint32>(0), CCMsg });
+
+			Filter.Process(InputStream, OutputStream);
+
+			bool GotNoteOn = false;
+			bool GotCC = false;
+			
+			for (const HarmonixMetasound::FMidiStreamEvent& Event : OutputStream.GetEventsInBlock())
+			{
+				if (Event.MidiMessage.IsNoteOn())
+				{
+					UTEST_EQUAL("Channel matches", Event.MidiMessage.GetStdChannel(), NoteOnMsg.GetStdChannel());
+					UTEST_EQUAL("Note number matches", Event.MidiMessage.GetStdData1(), NoteOnMsg.GetStdData1());
+					UTEST_EQUAL("Velocity matches", Event.MidiMessage.GetStdData2(), NoteOnMsg.GetStdData2());
+					GotNoteOn = true;
+				}
+				else if (Event.MidiMessage.IsControlChange())
+				{
+					UTEST_EQUAL("Channel matches", Event.MidiMessage.GetStdChannel(), CCMsg.GetStdChannel());
+					UTEST_EQUAL("Control number matches", Event.MidiMessage.GetStdData1(), CCMsg.GetStdData1());
+					UTEST_EQUAL("Value matches", Event.MidiMessage.GetStdData2(), CCMsg.GetStdData2());
+					GotCC = true;
+				}
+				else
+				{
+					UTEST_TRUE("Got an unexpected message", false);
+				}
+			}
+
+			UTEST_TRUE("Got note on", GotNoteOn);
+			UTEST_TRUE("Got CC", GotCC);
+		}
+
+		// Change the filter such that the first note is filtered out. We should get a note off.
+		{
+			// Clear out the last block's events
+			InputStream.PrepareBlock();
+			UTEST_EQUAL("No events in input", InputStream.GetEventsInBlock().Num(), 0);
+			OutputStream.PrepareBlock();
+			UTEST_EQUAL("No events in output", InputStream.GetEventsInBlock().Num(), 0);
+			
+			Filter.MinNoteNumber = NoteOnMsg.GetStdData1() + 1;
+			
+			Filter.Process(InputStream, OutputStream);
+
+			const TArray<HarmonixMetasound::FMidiStreamEvent>& EventsOut = OutputStream.GetEventsInBlock();
+			UTEST_EQUAL("Got one event out", EventsOut.Num(), 1);
+			UTEST_TRUE("Event is note off", EventsOut[0].MidiMessage.IsNoteOff());
+			UTEST_EQUAL("Note off is on correct channel", EventsOut[0].MidiMessage.GetStdChannel(), NoteOnMsg.GetStdChannel());
+			UTEST_EQUAL("Note off is on correct note number", EventsOut[0].MidiMessage.GetStdData1(), NoteOnMsg.GetStdData1());
+			UTEST_EQUAL("Note off has correct voice id", EventsOut[0].GetVoiceId(), NoteOnEvent.GetVoiceId());
+		}
+
+		// Fire the note on again, this time expect it to get filtered out
+		{
+			// Clear out the last block's events
+			InputStream.PrepareBlock();
+			UTEST_EQUAL("No events in input", InputStream.GetEventsInBlock().Num(), 0);
+			OutputStream.PrepareBlock();
+			UTEST_EQUAL("No events in output", InputStream.GetEventsInBlock().Num(), 0);
+			
+			InputStream.AddMidiEvent(NoteOnEvent);
+			
+			Filter.Process(InputStream, OutputStream);
+			
+			const TArray<HarmonixMetasound::FMidiStreamEvent>& EventsOut = OutputStream.GetEventsInBlock();
+			UTEST_EQUAL("Got no events", EventsOut.Num(), 0);
+		}
+
+		// Toggle "IncludeOtherEvents" off, send the CC again, this time expect it to get filtered out
+		{
+			// Clear out the last block's events
+			InputStream.PrepareBlock();
+			UTEST_EQUAL("No events in input", InputStream.GetEventsInBlock().Num(), 0);
+			OutputStream.PrepareBlock();
+			UTEST_EQUAL("No events in output", InputStream.GetEventsInBlock().Num(), 0);
+			
+			InputStream.AddMidiEvent({ static_cast<uint32>(0), FMidiMsg::CreateControlChange(4, 5, 6) });
+
+			Filter.IncludeOtherEvents.Set(false);
+			
+			Filter.Process(InputStream, OutputStream);
+
+			const TArray<HarmonixMetasound::FMidiStreamEvent>& EventsOut = OutputStream.GetEventsInBlock();
+			UTEST_EQUAL("Got no events", EventsOut.Num(), 0);
+		}
+
+		return true;
+	}
+}
+
+#endif
\ No newline at end of file
diff --git a/Engine/Plugins/Runtime/Harmonix/Source/HarmonixMetasoundTests/Private/Nodes/MidiNoteFilterNodeTests.cpp b/Engine/Plugins/Runtime/Harmonix/Source/HarmonixMetasoundTests/Private/Nodes/MidiNoteFilterNodeTests.cpp
new file mode 100644
--- /dev/null
+++ b/Engine/Plugins/Runtime/Harmonix/Source/HarmonixMetasoundTests/Private/Nodes/MidiNoteFilterNodeTests.cpp
@@ -0,0 +1,109 @@
+// Copyright Epic Games, Inc. All Rights Reserved.
+
+#include "NodeTestGraphBuilder.h"
+#include "HarmonixMetasound/MidiOps/MidiNoteFilter.h"
+#include "HarmonixMetasound/Nodes/MidiNoteFilterNode.h"
+#include "Misc/AutomationTest.h"
+
+#if WITH_DEV_AUTOMATION_TESTS
+
+namespace HarmonixMetasound::Nodes::MidiNoteFilter::Tests
+{
+	IMPLEMENT_SIMPLE_AUTOMATION_TEST(
+	FMidiNoteFilterNodeParityTest,
+	"Harmonix.Metasound.Nodes.MidiNoteFilter.Parity",
+	EAutomationTestFlags::EditorContext | EAutomationTestFlags::EngineFilter)
+	bool FMidiNoteFilterNodeParityTest::RunTest(const FString&)
+	{
+		const TUniquePtr<Metasound::FMetasoundGenerator> Generator =
+			Metasound::Test::FNodeTestGraphBuilder::MakeSingleNodeGraph(GetClassName(), GetCurrentMajorVersion());
+		TOptional<Metasound::TDataReadReference<FMidiStream>> NodeMidiOutput = Generator->GetOutputReadReference<FMidiStream>(Outputs::MidiStreamName);
+		UTEST_TRUE("Got node MIDI output", NodeMidiOutput.IsSet());
+		
+		Harmonix::Midi::Ops::FMidiNoteFilter FilterForComparison;
+		FMidiStream InputStream{ Generator->OperatorSettings };
+		FMidiStream OutputStream{ Generator->OperatorSettings };
+
+		FMidiStreamEvent NoteOnEvent{ static_cast<uint32>(0), FMidiMsg::CreateNoteOn(2, 43, 45) };
+		FMidiStreamEvent CCEvent{ static_cast<uint32>(0), FMidiMsg::CreateControlChange(3, 4, 5) };
+		Metasound::FAudioBuffer Buffer{ Generator->OperatorSettings };
+
+		// Default: all events pass
+		{
+			Generator->ApplyToInputValue<FMidiStream>(Inputs::MidiStreamName, [&NoteOnEvent, &CCEvent](FMidiStream& Stream)
+			{
+				Stream.AddMidiEvent(NoteOnEvent);
+				Stream.AddMidiEvent(CCEvent);
+			});
+
+			InputStream.AddMidiEvent(NoteOnEvent);
+			InputStream.AddMidiEvent(CCEvent);
+			
+			Generator->OnGenerateAudio(Buffer.GetData(), Buffer.Num());
+
+			FilterForComparison.Process(InputStream, OutputStream);
+
+			const TArray<FMidiStreamEvent>& NodeEvents = (*NodeMidiOutput)->GetEventsInBlock();
+			const TArray<FMidiStreamEvent>& FilterEvents = OutputStream.GetEventsInBlock();
+			UTEST_EQUAL("Right number of events", NodeEvents.Num(), 2);
+			UTEST_EQUAL("Same number of events", FilterEvents.Num(), NodeEvents.Num());
+		}
+
+		// Disabled: no events pass
+		{
+			// keep the same inputs from the previous block
+			Generator->SetInputValue<bool>(Inputs::EnableName, false);
+			Generator->OnGenerateAudio(Buffer.GetData(), Buffer.Num());
+			UTEST_EQUAL("No events", (*NodeMidiOutput)->GetEventsInBlock().Num(), 0);
+		}
+
+		// Filter out the note: the CC comes through, as well as a note off because the note was previously sounding
+		{
+			OutputStream.PrepareBlock();
+			
+			Generator->SetInputValue(Inputs::EnableName, true);
+			Generator->SetInputValue(Inputs::MaxVelocityName, NoteOnEvent.MidiMessage.GetStdData2() - 1);
+			Generator->OnGenerateAudio(Buffer.GetData(), Buffer.Num());
+
+			FilterForComparison.MaxVelocity = NoteOnEvent.MidiMessage.GetStdData2() - 1;
+			FilterForComparison.Process(InputStream, OutputStream);
+
+			const TArray<FMidiStreamEvent>& NodeEvents = (*NodeMidiOutput)->GetEventsInBlock();
+			const TArray<FMidiStreamEvent>& FilterEvents = OutputStream.GetEventsInBlock();
+			UTEST_EQUAL("Right number of events", NodeEvents.Num(), 2);
+			UTEST_EQUAL("Same number of events", FilterEvents.Num(), NodeEvents.Num());
+
+			bool bGotCC = false;
+			bool bGotNoteOff = false;
+
+			for (const FMidiStreamEvent& Event : NodeEvents)
+			{
+				if (Event.MidiMessage.IsControlChange())
+				{
+					UTEST_EQUAL("CC: Right channel", Event.MidiMessage.GetStdChannel(), CCEvent.MidiMessage.GetStdChannel());
+					UTEST_EQUAL("CC: Right control number", Event.MidiMessage.GetStdData1(), CCEvent.MidiMessage.GetStdData1());
+					UTEST_EQUAL("CC: Right value", Event.MidiMessage.GetStdData2(), CCEvent.MidiMessage.GetStdData2());
+					bGotCC = true;
+				}
+				else if (Event.MidiMessage.IsNoteOff())
+				{
+					UTEST_EQUAL("Note off: Right channel", Event.MidiMessage.GetStdChannel(), NoteOnEvent.MidiMessage.GetStdChannel());
+					UTEST_EQUAL("Note off: Right note number", Event.MidiMessage.GetStdData1(), NoteOnEvent.MidiMessage.GetStdData1());
+					UTEST_EQUAL("Note off: Right voice id", Event.GetVoiceId(), NoteOnEvent.GetVoiceId());
+					bGotNoteOff = true;
+				}
+				else
+				{
+					UTEST_TRUE("Unexpected event", false);
+				}
+			}
+
+			UTEST_TRUE("Got CC", bGotCC);
+			UTEST_TRUE("Got Note off", bGotNoteOff);
+		}
+		
+		return true;
+	}
+}
+
+#endif
\ No newline at end of file
