[Harmonix] Update MIDI note trigger node to conform to library design

#tests automated tests, editor

diff --git a/Engine/Plugins/Runtime/Harmonix/Source/HarmonixMetasound/Private/HarmonixMetasound/Nodes/MidiNoteTriggerNode.cpp b/Engine/Plugins/Runtime/Harmonix/Source/HarmonixMetasound/Private/HarmonixMetasound/Nodes/MidiNoteTriggerNode.cpp
--- a/Engine/Plugins/Runtime/Harmonix/Source/HarmonixMetasound/Private/HarmonixMetasound/Nodes/MidiNoteTriggerNode.cpp
+++ b/Engine/Plugins/Runtime/Harmonix/Source/HarmonixMetasound/Private/HarmonixMetasound/Nodes/MidiNoteTriggerNode.cpp
@@ -1,5 +1,7 @@
 // Copyright Epic Games, Inc. All Rights Reserved.
 
+#include "HarmonixMetasound/Nodes/MidiNoteTriggerNode.h"
+
 #include "MetasoundExecutableOperator.h"
 #include "MetasoundFacade.h"
 #include "MetasoundNodeInterface.h"
@@ -13,22 +15,228 @@
 #include "HarmonixMetasound/DataTypes/MidiStream.h"
 #include "HarmonixMetasound/DataTypes/MusicTransport.h"
 
-DEFINE_LOG_CATEGORY_STATIC(LogMidiNoteTrigger, Log, All);
-
 #define LOCTEXT_NAMESPACE "HarmonixMetaSound"
 
-namespace HarmonixMetasound
+namespace HarmonixMetasound::Nodes::MidiNoteTriggerNode
 {
 	using namespace Metasound;
 
-	class FMidiNoteTriggerOperator : public TExecutableOperator<FMidiNoteTriggerOperator>
+	const FNodeClassName& GetClassName()
+	{
+		static const FNodeClassName ClassName{ HarmonixNodeNamespace, TEXT("MidiNoteTrigger"), TEXT("")};
+		return ClassName;
+	}
+
+	int32 GetCurrentMajorVersion()
+	{
+		return 1;
+	}
+
+	namespace Inputs
+	{
+		DEFINE_METASOUND_PARAM_ALIAS(Enable, CommonPinNames::Inputs::Enable);
+		DEFINE_METASOUND_PARAM_ALIAS(MidiStream, CommonPinNames::Inputs::MidiStream);
+	}
+
+	namespace Outputs
+	{
+		DEFINE_METASOUND_PARAM_ALIAS(NoteOnTrigger, CommonPinNames::Outputs::NoteOn);
+		DEFINE_METASOUND_PARAM_ALIAS(NoteOffTrigger, CommonPinNames::Outputs::NoteOff);
+		DEFINE_METASOUND_PARAM_ALIAS(MidiNoteNumber, CommonPinNames::Outputs::MidiNoteNumber);
+		DEFINE_METASOUND_PARAM_ALIAS(MidiVelocity, CommonPinNames::Outputs::MidiVelocity);
+	}
+
+	class FMidiNoteTriggerOperator_V1 final : public TExecutableOperator<FMidiNoteTriggerOperator_V1>
+	{
+	public:
+		static const FNodeClassMetadata& GetNodeInfo()
+		{
+			auto InitNodeInfo = []() -> FNodeClassMetadata
+			{
+				FNodeClassMetadata Info;
+				Info.ClassName        = GetClassName();
+				Info.MajorVersion     = 1;
+				Info.MinorVersion     = 0;
+				Info.DisplayName      = METASOUND_LOCTEXT("MidiNoteTriggerNode_DisplayName", "MIDI Note Trigger");
+				Info.Description      = METASOUND_LOCTEXT("MidiNoteTriggerNode_Description", "Outputs triggers and info for incoming MIDI notes.");
+				Info.Author           = PluginAuthor;
+				Info.PromptIfMissing  = PluginNodeMissingPrompt;
+				Info.DefaultInterface = GetVertexInterface();
+				Info.CategoryHierarchy.Emplace(NodeCategories::Music);
+				return Info;
+			};
+
+			static const FNodeClassMetadata Info = InitNodeInfo();
+
+			return Info;
+		}
+
+		static const FVertexInterface& GetVertexInterface()
+		{
+			static const FVertexInterface Interface(
+				FInputVertexInterface(
+					TInputDataVertex<bool>(METASOUND_GET_PARAM_NAME_AND_METADATA(Inputs::Enable), true),
+					TInputDataVertex<FMidiStream>(METASOUND_GET_PARAM_NAME_AND_METADATA(Inputs::MidiStream))
+					),
+				FOutputVertexInterface(
+					TOutputDataVertex<FTrigger>(METASOUND_GET_PARAM_NAME_AND_METADATA(Outputs::NoteOnTrigger)),
+					TOutputDataVertex<FTrigger>(METASOUND_GET_PARAM_NAME_AND_METADATA(Outputs::NoteOffTrigger)),
+					TOutputDataVertex<int32>(METASOUND_GET_PARAM_NAME_AND_METADATA(Outputs::MidiNoteNumber)),
+					TOutputDataVertex<int32>(METASOUND_GET_PARAM_NAME_AND_METADATA(Outputs::MidiVelocity))
+					)
+			);
+
+			return Interface;
+		}
+
+		static TUniquePtr<IOperator> CreateOperator(const FBuildOperatorParams& InParams, FBuildResults&)
+		{
+			using namespace CommonPinNames;
+
+			const FInputVertexInterfaceData& InputData = InParams.InputData;
+			FBoolReadRef InEnabled          = InputData.GetOrCreateDefaultDataReadReference<bool>(METASOUND_GET_PARAM_NAME(Inputs::Enable), InParams.OperatorSettings);
+			FMidiStreamReadRef InMidiStream = InputData.GetOrConstructDataReadReference<FMidiStream>(METASOUND_GET_PARAM_NAME(Inputs::MidiStream), InParams.OperatorSettings);
+
+			return MakeUnique<FMidiNoteTriggerOperator_V1>(InParams, MoveTemp(InEnabled), MoveTemp(InMidiStream));
+		}
+
+		FMidiNoteTriggerOperator_V1(const FBuildOperatorParams& InParams, FBoolReadRef&& InEnabled, FMidiStreamReadRef&& InMidiStream)
+			: EnableInPin(MoveTemp(InEnabled))
+			, MidiStreamInPin(MoveTemp(InMidiStream))
+			, NoteOnOutPin(FTriggerWriteRef::CreateNew(InParams.OperatorSettings))
+			, NoteOffOutPin(FTriggerWriteRef::CreateNew(InParams.OperatorSettings))
+			, NoteNumOutPin(FInt32WriteRef::CreateNew(0))
+			, VelOutPin(FInt32WriteRef::CreateNew(0))
+		{
+			Reset(InParams);
+		}
+
+		virtual void BindInputs(FInputVertexInterfaceData& InVertexData) override
+		{
+			using namespace CommonPinNames;
+
+			InVertexData.BindReadVertex(METASOUND_GET_PARAM_NAME(Inputs::Enable), EnableInPin);
+			InVertexData.BindReadVertex(METASOUND_GET_PARAM_NAME(Inputs::MidiStream),   MidiStreamInPin);
+		}
+
+		virtual void BindOutputs(FOutputVertexInterfaceData& InVertexData) override
+		{
+			using namespace CommonPinNames;
+
+			InVertexData.BindReadVertex(METASOUND_GET_PARAM_NAME(CommonPinNames::Outputs::NoteOn), NoteOnOutPin);
+			InVertexData.BindReadVertex(METASOUND_GET_PARAM_NAME(CommonPinNames::Outputs::NoteOff), NoteOffOutPin);
+			InVertexData.BindReadVertex(METASOUND_GET_PARAM_NAME(Outputs::MidiNoteNumber), NoteNumOutPin);
+			InVertexData.BindReadVertex(METASOUND_GET_PARAM_NAME(Outputs::MidiVelocity), VelOutPin);
+		}
+
+		void Reset(const FResetParams&)
+		{
+			NoteOnOutPin->Reset();
+			NoteOffOutPin->Reset();
+			*NoteNumOutPin = 0;
+			*VelOutPin = 0;
+			SoundingNote = -1;
+			PlayingId = FMidiVoiceId::None();
+		}
+
+		void Execute()
+		{
+			NoteOnOutPin->AdvanceBlock();
+			NoteOffOutPin->AdvanceBlock();
+
+			if (!*EnableInPin)
+			{
+				if (SoundingNote >= 0)
+				{
+					TriggerNoteOff(0, SoundingNote);
+				}
+				return;
+			}
+
+			const TArray<FMidiTimestampTransportState> TransportChanges = MidiStreamInPin->GetTransportChangesInBlock();
+			if (!TransportChanges.IsEmpty())
+			{
+				if (TransportChanges.Last().TransportState != EMusicPlayerTransportState::Playing)
+				{
+					if (SoundingNote >= 0)
+					{
+						TriggerNoteOff(0, SoundingNote);
+					}
+					return;
+				}
+			}
+
+			int32 NoteOffTriggerFrame = -1;
+			
+			for (const FMidiStreamEvent& Event : MidiStreamInPin->GetEventsInBlock())
+			{
+				if (Event.MidiMessage.IsNoteOn()) 
+				{
+					if (SoundingNote > -1)
+					{
+						// Stop sounding note
+						TriggerNoteOff(Event.BlockSampleFrameIndex, SoundingNote);
+					}
+
+					// Play new note
+					PlayingId      = Event.GetVoiceId();
+					*VelOutPin     = Event.MidiMessage.GetStdData2();
+					*NoteNumOutPin = Event.MidiMessage.GetStdData1();
+					NoteOnOutPin->TriggerFrame(NoteOffTriggerFrame == Event.BlockSampleFrameIndex ? NoteOffTriggerFrame + 1 : Event.BlockSampleFrameIndex);
+					SoundingNote   = Event.MidiMessage.GetStdData1();;
+				}
+				else if (Event.MidiMessage.IsNoteOff())
+				{
+					if (Event.GetVoiceId() == PlayingId)
+					{
+						TriggerNoteOff(Event.BlockSampleFrameIndex, Event.MidiMessage.GetStdData1());
+						NoteOffTriggerFrame = Event.BlockSampleFrameIndex;
+					}
+				}
+			}
+		}
+		
+	private:
+		void TriggerNoteOff(int32 BlockSampleFrameIndex, int32 NoteNumber)
+		{
+			*VelOutPin = 0;
+			*NoteNumOutPin = NoteNumber;
+			NoteOffOutPin->TriggerFrame(BlockSampleFrameIndex);
+			SoundingNote = -1;
+			PlayingId = FMidiVoiceId::None();
+		}
+		
+		FBoolReadRef       EnableInPin;
+		FMidiStreamReadRef MidiStreamInPin;
+
+		FTriggerWriteRef NoteOnOutPin;
+		FTriggerWriteRef NoteOffOutPin;
+		FInt32WriteRef   NoteNumOutPin;
+		FInt32WriteRef   VelOutPin;
+
+		int8 SoundingNote = -1;
+		FMidiVoiceId PlayingId = FMidiVoiceId::None();
+	};
+
+	class FMidiNoteTriggerNode_V1 final : public FNodeFacade
+	{
+	public:
+		explicit FMidiNoteTriggerNode_V1(const FNodeInitData& InInitData)
+			: FNodeFacade(InInitData.InstanceName, InInitData.InstanceID, TFacadeOperatorClass<FMidiNoteTriggerOperator_V1>())
+		{}
+		virtual ~FMidiNoteTriggerNode_V1() override = default;
+	};
+
+	METASOUND_REGISTER_NODE(FMidiNoteTriggerNode_V1)
+
+	class FMidiNoteTriggerOperator_V0 : public TExecutableOperator<FMidiNoteTriggerOperator_V0>
 	{
 	public:
 		static const FNodeClassMetadata& GetNodeInfo();
 		static const FVertexInterface& GetVertexInterface();
 		static TUniquePtr<IOperator> CreateOperator(const FBuildOperatorParams& InParams, FBuildResults& OutResults);
 
-		FMidiNoteTriggerOperator(const FBuildOperatorParams& InParams,
+		FMidiNoteTriggerOperator_V0(const FBuildOperatorParams& InParams,
 								 const FBoolReadRef&       InEnabled,
 								 const FMidiStreamReadRef& InMidiStream,
 								 const FInt32ReadRef&      InTrackNumber,
@@ -73,23 +281,23 @@ namespace HarmonixMetasound
 		FMidiVoiceId PlayingId;
 	};
 
-	class FMidiNoteTriggerNode : public FNodeFacade
+	class FMidiNoteTriggerNode_V0 : public FNodeFacade
 	{
 	public:
-		FMidiNoteTriggerNode(const FNodeInitData& InInitData)
-			: FNodeFacade(InInitData.InstanceName, InInitData.InstanceID, TFacadeOperatorClass<FMidiNoteTriggerOperator>())
+		FMidiNoteTriggerNode_V0(const FNodeInitData& InInitData)
+			: FNodeFacade(InInitData.InstanceName, InInitData.InstanceID, TFacadeOperatorClass<FMidiNoteTriggerOperator_V0>())
 		{}
-		virtual ~FMidiNoteTriggerNode() = default;
+		virtual ~FMidiNoteTriggerNode_V0() = default;
 	};
 
-	METASOUND_REGISTER_NODE(FMidiNoteTriggerNode)
+	METASOUND_REGISTER_NODE(FMidiNoteTriggerNode_V0)
 
-	const FNodeClassMetadata& FMidiNoteTriggerOperator::GetNodeInfo()
+	const FNodeClassMetadata& FMidiNoteTriggerOperator_V0::GetNodeInfo()
 	{
 		auto InitNodeInfo = []() -> FNodeClassMetadata
 		{
 			FNodeClassMetadata Info;
-			Info.ClassName        = { HarmonixNodeNamespace, TEXT("MidiNoteTrigger"), TEXT("")};
+			Info.ClassName        = GetClassName();
 			Info.MajorVersion     = 0;
 			Info.MinorVersion     = 1;
 			Info.DisplayName      = METASOUND_LOCTEXT("MidiNoteTriggerNode_DisplayName", "Midi Note Trigger");
@@ -98,6 +306,7 @@ namespace HarmonixMetasound
 			Info.PromptIfMissing  = PluginNodeMissingPrompt;
 			Info.DefaultInterface = GetVertexInterface();
 			Info.CategoryHierarchy.Emplace(NodeCategories::Music);
+			Info.bDeprecated = true;
 			return Info;
 		};
 
@@ -106,56 +315,52 @@ namespace HarmonixMetasound
 		return Info;
 	}
 
-	const FVertexInterface& FMidiNoteTriggerOperator::GetVertexInterface()
+	const FVertexInterface& FMidiNoteTriggerOperator_V0::GetVertexInterface()
 	{
-		using namespace CommonPinNames;
-
 		static const FVertexInterface Interface(
 			FInputVertexInterface(
 				TInputDataVertex<bool>(METASOUND_GET_PARAM_NAME_AND_METADATA(Inputs::Enable), true),
 				TInputDataVertex<FMidiStream>(METASOUND_GET_PARAM_NAME_AND_METADATA(Inputs::MidiStream)),
-				TInputDataVertex<int32>(METASOUND_GET_PARAM_NAME_AND_METADATA(Inputs::MidiTrackNumber),1),
-				TInputDataVertex<int32>(METASOUND_GET_PARAM_NAME_AND_METADATA(Inputs::MidiChannelNumber), 1),
-				TInputDataVertex<int32>(METASOUND_GET_PARAM_NAME_AND_METADATA(Inputs::MinMidiNote), 0),
-				TInputDataVertex<int32>(METASOUND_GET_PARAM_NAME_AND_METADATA(Inputs::MaxMidiNote), 127),
-				TInputDataVertex<int32>(METASOUND_GET_PARAM_NAME_AND_METADATA(Inputs::MinMidiVelocity), 0),
-				TInputDataVertex<int32>(METASOUND_GET_PARAM_NAME_AND_METADATA(Inputs::MaxMidiVelocity), 127),
-				TInputDataVertex<bool>(METASOUND_GET_PARAM_NAME_AND_METADATA(Inputs::ClockSpeedToFrequency), false)
+				TInputDataVertex<int32>(METASOUND_GET_PARAM_NAME_AND_METADATA(CommonPinNames::Inputs::MidiTrackNumber),1),
+				TInputDataVertex<int32>(METASOUND_GET_PARAM_NAME_AND_METADATA(CommonPinNames::Inputs::MidiChannelNumber), 1),
+				TInputDataVertex<int32>(METASOUND_GET_PARAM_NAME_AND_METADATA(CommonPinNames::Inputs::MinMidiNote), 0),
+				TInputDataVertex<int32>(METASOUND_GET_PARAM_NAME_AND_METADATA(CommonPinNames::Inputs::MaxMidiNote), 127),
+				TInputDataVertex<int32>(METASOUND_GET_PARAM_NAME_AND_METADATA(CommonPinNames::Inputs::MinMidiVelocity), 0),
+				TInputDataVertex<int32>(METASOUND_GET_PARAM_NAME_AND_METADATA(CommonPinNames::Inputs::MaxMidiVelocity), 127),
+				TInputDataVertex<bool>(METASOUND_GET_PARAM_NAME_AND_METADATA(CommonPinNames::Inputs::ClockSpeedToFrequency), false)
 				),
 			FOutputVertexInterface(
-				TOutputDataVertex<FTrigger>(METASOUND_GET_PARAM_NAME_AND_METADATA(Outputs::NoteOn)),
-				TOutputDataVertex<FTrigger>(METASOUND_GET_PARAM_NAME_AND_METADATA(Outputs::NoteOff)),
+				TOutputDataVertex<FTrigger>(METASOUND_GET_PARAM_NAME_AND_METADATA(Outputs::NoteOnTrigger)),
+				TOutputDataVertex<FTrigger>(METASOUND_GET_PARAM_NAME_AND_METADATA(Outputs::NoteOffTrigger)),
 				TOutputDataVertex<int32>(METASOUND_GET_PARAM_NAME_AND_METADATA(Outputs::MidiNoteNumber)),
-				TOutputDataVertex<float>(METASOUND_GET_PARAM_NAME_AND_METADATA(Outputs::Frequency)),
+				TOutputDataVertex<float>(METASOUND_GET_PARAM_NAME_AND_METADATA(CommonPinNames::Outputs::Frequency)),
 				TOutputDataVertex<int32>(METASOUND_GET_PARAM_NAME_AND_METADATA(Outputs::MidiVelocity)),
-				TOutputDataVertex<float>(METASOUND_GET_PARAM_NAME_AND_METADATA(Outputs::NormalizedVelocity))
+				TOutputDataVertex<float>(METASOUND_GET_PARAM_NAME_AND_METADATA(CommonPinNames::Outputs::NormalizedVelocity))
 				)
 		);
 
 		return Interface;
 	}
 
-	TUniquePtr<IOperator> FMidiNoteTriggerOperator::CreateOperator(const FBuildOperatorParams& InParams, FBuildResults& OutResults)
+	TUniquePtr<IOperator> FMidiNoteTriggerOperator_V0::CreateOperator(const FBuildOperatorParams& InParams, FBuildResults& OutResults)
 	{
 		using namespace CommonPinNames;
 
-		const FMidiNoteTriggerNode& LoggerNode = static_cast<const FMidiNoteTriggerNode&>(InParams.Node);
-
 		const FInputVertexInterfaceData& InputData = InParams.InputData;
 		FBoolReadRef InEnabled          = InputData.GetOrCreateDefaultDataReadReference<bool>(METASOUND_GET_PARAM_NAME(Inputs::Enable), InParams.OperatorSettings);
 		FMidiStreamReadRef InMidiStream = InputData.GetOrConstructDataReadReference<FMidiStream>(METASOUND_GET_PARAM_NAME(Inputs::MidiStream), InParams.OperatorSettings);
-		FInt32ReadRef InTrackNumber     = InputData.GetOrCreateDefaultDataReadReference<int32>(METASOUND_GET_PARAM_NAME(Inputs::MidiTrackNumber), InParams.OperatorSettings);
-		FInt32ReadRef InMidiChannel     = InputData.GetOrCreateDefaultDataReadReference<int32>(METASOUND_GET_PARAM_NAME(Inputs::MidiChannelNumber), InParams.OperatorSettings);
-		FInt32ReadRef InMinNote         = InputData.GetOrCreateDefaultDataReadReference<int32>(METASOUND_GET_PARAM_NAME(Inputs::MinMidiNote), InParams.OperatorSettings);
-		FInt32ReadRef InMaxNote         = InputData.GetOrCreateDefaultDataReadReference<int32>(METASOUND_GET_PARAM_NAME(Inputs::MaxMidiNote), InParams.OperatorSettings);
-		FInt32ReadRef InMinVelocity     = InputData.GetOrCreateDefaultDataReadReference<int32>(METASOUND_GET_PARAM_NAME(Inputs::MinMidiVelocity), InParams.OperatorSettings);
-		FInt32ReadRef InMaxVelocity     = InputData.GetOrCreateDefaultDataReadReference<int32>(METASOUND_GET_PARAM_NAME(Inputs::MaxMidiVelocity), InParams.OperatorSettings);
-		FBoolReadRef  InClockToFreq     = InputData.GetOrCreateDefaultDataReadReference<bool>(METASOUND_GET_PARAM_NAME(Inputs::ClockSpeedToFrequency), InParams.OperatorSettings);
-
-		return MakeUnique<FMidiNoteTriggerOperator>(InParams, InEnabled, InMidiStream, InTrackNumber, InMidiChannel, InMinNote, InMaxNote, InMinVelocity, InMaxVelocity, InClockToFreq);
+		FInt32ReadRef InTrackNumber     = InputData.GetOrCreateDefaultDataReadReference<int32>(METASOUND_GET_PARAM_NAME(CommonPinNames::Inputs::MidiTrackNumber), InParams.OperatorSettings);
+		FInt32ReadRef InMidiChannel     = InputData.GetOrCreateDefaultDataReadReference<int32>(METASOUND_GET_PARAM_NAME(CommonPinNames::Inputs::MidiChannelNumber), InParams.OperatorSettings);
+		FInt32ReadRef InMinNote         = InputData.GetOrCreateDefaultDataReadReference<int32>(METASOUND_GET_PARAM_NAME(CommonPinNames::Inputs::MinMidiNote), InParams.OperatorSettings);
+		FInt32ReadRef InMaxNote         = InputData.GetOrCreateDefaultDataReadReference<int32>(METASOUND_GET_PARAM_NAME(CommonPinNames::Inputs::MaxMidiNote), InParams.OperatorSettings);
+		FInt32ReadRef InMinVelocity     = InputData.GetOrCreateDefaultDataReadReference<int32>(METASOUND_GET_PARAM_NAME(CommonPinNames::Inputs::MinMidiVelocity), InParams.OperatorSettings);
+		FInt32ReadRef InMaxVelocity     = InputData.GetOrCreateDefaultDataReadReference<int32>(METASOUND_GET_PARAM_NAME(CommonPinNames::Inputs::MaxMidiVelocity), InParams.OperatorSettings);
+		FBoolReadRef  InClockToFreq     = InputData.GetOrCreateDefaultDataReadReference<bool>(METASOUND_GET_PARAM_NAME(CommonPinNames::Inputs::ClockSpeedToFrequency), InParams.OperatorSettings);
+
+		return MakeUnique<FMidiNoteTriggerOperator_V0>(InParams, InEnabled, InMidiStream, InTrackNumber, InMidiChannel, InMinNote, InMaxNote, InMinVelocity, InMaxVelocity, InClockToFreq);
 	}
 
-	FMidiNoteTriggerOperator::FMidiNoteTriggerOperator(const FBuildOperatorParams& InParams,
+	FMidiNoteTriggerOperator_V0::FMidiNoteTriggerOperator_V0(const FBuildOperatorParams& InParams,
 													   const FBoolReadRef&       InEnabled,
 													   const FMidiStreamReadRef& InMidiStream,
 													   const FInt32ReadRef&      InTrackNumber,
@@ -184,34 +389,34 @@ namespace HarmonixMetasound
 		Reset(InParams);
 	}
 
-	void FMidiNoteTriggerOperator::BindInputs(FInputVertexInterfaceData& InVertexData)
+	void FMidiNoteTriggerOperator_V0::BindInputs(FInputVertexInterfaceData& InVertexData)
 	{
 		using namespace CommonPinNames;
 
 		InVertexData.BindReadVertex(METASOUND_GET_PARAM_NAME(Inputs::Enable), EnableInPin);
 		InVertexData.BindReadVertex(METASOUND_GET_PARAM_NAME(Inputs::MidiStream),   MidiStreamInPin);
-		InVertexData.BindReadVertex(METASOUND_GET_PARAM_NAME(Inputs::MidiTrackNumber),  TrackNumberInPin);
-		InVertexData.BindReadVertex(METASOUND_GET_PARAM_NAME(Inputs::MidiChannelNumber), ChannelNumInPin);
-		InVertexData.BindReadVertex(METASOUND_GET_PARAM_NAME(Inputs::MinMidiNote), MinNoteInPin);
-		InVertexData.BindReadVertex(METASOUND_GET_PARAM_NAME(Inputs::MaxMidiNote), MaxNoteInPin);
-		InVertexData.BindReadVertex(METASOUND_GET_PARAM_NAME(Inputs::MinMidiVelocity), MinVelInPin);
-		InVertexData.BindReadVertex(METASOUND_GET_PARAM_NAME(Inputs::MaxMidiVelocity), MaxVelInPin);
-		InVertexData.BindReadVertex(METASOUND_GET_PARAM_NAME(Inputs::ClockSpeedToFrequency), ClockAffectsFrequencyInPin);
+		InVertexData.BindReadVertex(METASOUND_GET_PARAM_NAME(CommonPinNames::Inputs::MidiTrackNumber),  TrackNumberInPin);
+		InVertexData.BindReadVertex(METASOUND_GET_PARAM_NAME(CommonPinNames::Inputs::MidiChannelNumber), ChannelNumInPin);
+		InVertexData.BindReadVertex(METASOUND_GET_PARAM_NAME(CommonPinNames::Inputs::MinMidiNote), MinNoteInPin);
+		InVertexData.BindReadVertex(METASOUND_GET_PARAM_NAME(CommonPinNames::Inputs::MaxMidiNote), MaxNoteInPin);
+		InVertexData.BindReadVertex(METASOUND_GET_PARAM_NAME(CommonPinNames::Inputs::MinMidiVelocity), MinVelInPin);
+		InVertexData.BindReadVertex(METASOUND_GET_PARAM_NAME(CommonPinNames::Inputs::MaxMidiVelocity), MaxVelInPin);
+		InVertexData.BindReadVertex(METASOUND_GET_PARAM_NAME(CommonPinNames::Inputs::ClockSpeedToFrequency), ClockAffectsFrequencyInPin);
 	}
 
-	void FMidiNoteTriggerOperator::BindOutputs(FOutputVertexInterfaceData& InVertexData)
+	void FMidiNoteTriggerOperator_V0::BindOutputs(FOutputVertexInterfaceData& InVertexData)
 	{
 		using namespace CommonPinNames;
 
-		InVertexData.BindReadVertex(METASOUND_GET_PARAM_NAME(Outputs::NoteOn), NoteOnOutPin);
-		InVertexData.BindReadVertex(METASOUND_GET_PARAM_NAME(Outputs::NoteOff), NoteOffOutPin);
+		InVertexData.BindReadVertex(METASOUND_GET_PARAM_NAME(CommonPinNames::Outputs::NoteOn), NoteOnOutPin);
+		InVertexData.BindReadVertex(METASOUND_GET_PARAM_NAME(CommonPinNames::Outputs::NoteOff), NoteOffOutPin);
 		InVertexData.BindReadVertex(METASOUND_GET_PARAM_NAME(Outputs::MidiNoteNumber), NoteNumOutPin);
-		InVertexData.BindReadVertex(METASOUND_GET_PARAM_NAME(Outputs::Frequency), FreqOutPin);
+		InVertexData.BindReadVertex(METASOUND_GET_PARAM_NAME(CommonPinNames::Outputs::Frequency), FreqOutPin);
 		InVertexData.BindReadVertex(METASOUND_GET_PARAM_NAME(Outputs::MidiVelocity), VelOutPin);
-		InVertexData.BindReadVertex(METASOUND_GET_PARAM_NAME(Outputs::NormalizedVelocity), NormVelOutPin);
+		InVertexData.BindReadVertex(METASOUND_GET_PARAM_NAME(CommonPinNames::Outputs::NormalizedVelocity), NormVelOutPin);
 	}
 	
-	FDataReferenceCollection FMidiNoteTriggerOperator::GetInputs() const
+	FDataReferenceCollection FMidiNoteTriggerOperator_V0::GetInputs() const
 	{
 		// This should never be called. Bind(...) is called instead. This method
 		// exists as a stop-gap until the API can be deprecated and removed.
@@ -219,7 +424,7 @@ namespace HarmonixMetasound
 		return {};
 	}
 
-	FDataReferenceCollection FMidiNoteTriggerOperator::GetOutputs() const
+	FDataReferenceCollection FMidiNoteTriggerOperator_V0::GetOutputs() const
 	{
 		// This should never be called. Bind(...) is called instead. This method
 		// exists as a stop-gap until the API can be deprecated and removed.
@@ -227,7 +432,7 @@ namespace HarmonixMetasound
 		return {};
 	}
 
-	void FMidiNoteTriggerOperator::Reset(const FResetParams& ResetParams)
+	void FMidiNoteTriggerOperator_V0::Reset(const FResetParams& ResetParams)
 	{
 		NoteOnOutPin->Reset();
 		NoteOffOutPin->Reset();
@@ -239,7 +444,7 @@ namespace HarmonixMetasound
 		PlayingId = FMidiVoiceId::None();
 	}
 
-	void FMidiNoteTriggerOperator::Execute()
+	void FMidiNoteTriggerOperator_V0::Execute()
 	{
 		NoteOnOutPin->AdvanceBlock();
 		NoteOffOutPin->AdvanceBlock();
diff --git a/Engine/Plugins/Runtime/Harmonix/Source/HarmonixMetasound/Public/HarmonixMetasound/Nodes/MidiNoteTriggerNode.h b/Engine/Plugins/Runtime/Harmonix/Source/HarmonixMetasound/Public/HarmonixMetasound/Nodes/MidiNoteTriggerNode.h
new file mode 100644
--- /dev/null
+++ b/Engine/Plugins/Runtime/Harmonix/Source/HarmonixMetasound/Public/HarmonixMetasound/Nodes/MidiNoteTriggerNode.h
@@ -0,0 +1,27 @@
+// Copyright Epic Games, Inc. All Rights Reserved.
+
+#pragma once
+
+#include "MetasoundNodeInterface.h"
+
+#include "HarmonixMetasound/Common.h"
+
+namespace HarmonixMetasound::Nodes::MidiNoteTriggerNode
+{
+	const HARMONIXMETASOUND_API Metasound::FNodeClassName& GetClassName();
+	HARMONIXMETASOUND_API int32 GetCurrentMajorVersion();
+
+	namespace Inputs
+	{
+		DECLARE_METASOUND_PARAM_EXTERN(Enable);
+		DECLARE_METASOUND_PARAM_EXTERN(MidiStream);
+	}
+
+	namespace Outputs
+	{
+		DECLARE_METASOUND_PARAM_EXTERN(NoteOnTrigger);
+		DECLARE_METASOUND_PARAM_EXTERN(NoteOffTrigger);
+		DECLARE_METASOUND_PARAM_EXTERN(MidiNoteNumber);
+		DECLARE_METASOUND_PARAM_EXTERN(MidiVelocity);
+	}
+}
\ No newline at end of file
diff --git a/Engine/Plugins/Runtime/Harmonix/Source/HarmonixMetasoundTests/Private/Nodes/MidiNoteTriggerNodeTests.cpp b/Engine/Plugins/Runtime/Harmonix/Source/HarmonixMetasoundTests/Private/Nodes/MidiNoteTriggerNodeTests.cpp
--- a/Engine/Plugins/Runtime/Harmonix/Source/HarmonixMetasoundTests/Private/Nodes/MidiNoteTriggerNodeTests.cpp
+++ b/Engine/Plugins/Runtime/Harmonix/Source/HarmonixMetasoundTests/Private/Nodes/MidiNoteTriggerNodeTests.cpp
@@ -4,229 +4,29 @@
 #include "HarmonixDsp/AudioBuffer.h"
 #include "HarmonixMetasound/Common.h"
 #include "HarmonixMetasound/DataTypes/MidiStream.h"
+#include "HarmonixMetasound/Nodes/MidiNoteTriggerNode.h"
 #include "Misc/AutomationTest.h"
-#include "HarmonixMetasound/DataTypes/MidiAsset.h"
-#include "HarmonixMetasound/DataTypes/MusicTransport.h"
-#include "HarmonixMidi/MidiMsg.h"
-#include "HarmonixMidi/MidiFile.h"
 
 #if WITH_DEV_AUTOMATION_TESTS
 
 namespace HarmonixMetasoundTests::MidiNoteTriggerNode
 {
-	//testing Note Numbers
-	constexpr int32 TestNoteNumberShouldPass = 60; //C4, frequency = 261.63 Hz
-	constexpr int32 TestNoteNumberShouldPass2 = 62; //D4, frequency = 293.66 Hz
-
-	constexpr int32 TestNoteNumberShouldNotPass = 72; //C5, frequency = 523.25 Hz
-	//Min and Max Note Numbers for input 
-	constexpr int32 TestMaxNoteNumberDefault = 100;
-	constexpr int32 TestMinNoteNumberDefault = 20;
-	//Note Numbers that are outside of the Min to Max range
-	constexpr int32 TestNoteNumberLessThanMin = 10;
-	constexpr int32 TestNoteNumberGreaterThanMax = 127;
-
-	//testing velocity values
-	constexpr int32 TestNoteVelocityShouldPass = 90;
-	constexpr float TestNoteNormalizedVelocityShouldPass = 90 / 127.f;
-	//Min and Max Velocity values
-	constexpr int32 TestMaxVelocityDefault = 100;
-	constexpr int32 TestMinVelocityDefault = 20;
-	//velocity values outside of the Min to Max range
-	constexpr int32 TestVelocityLessThanMin = 10;
-	constexpr int32 TestVelocityGreaterThanMax = 127;
-	constexpr int32 TestNoteVelocityShouldNotPass = 127;//maybe we don't need this?
-	
-	//test midi tracks
-	constexpr int32 TestMidiTrackShouldPass = 1;
-	constexpr int32 TestMidiTrackShouldNotPass = 2;
-
-	//test midi channel
-	constexpr int32 TestMidiChannelShouldPass = 0;
-	constexpr int32 TestMidiChannelShouldNotPass = 1;
-
-	//frequency outputs
-	constexpr float ExpectedFrequencyOnPass = 261.625549f; //needs verification 
-	constexpr float ExpectedFrequencyNotPass = 0.0f;
-
-	//other expected values if event does not pass through
-	constexpr int32 ExpectedNoteNumberNotPass = 0;
-	constexpr int32 ExpectedVelocityNotPass = 0;
-	constexpr float ExpectedNormalizedVelocityNotPass = 0.0f;
-
-	//expected number of triggers that pass through during each block
-	constexpr int32 ExpectedTriggerCountOnPass = 1;
-	constexpr int32 ExpectedTriggerCountNotPass = 0;
-
-	//expected trigger frames
-	constexpr int32 ExpectedTriggerFrameIndexNotPass = -1;
-
-	//Time indices (Ms) for Midi events in Testing blocks
-	constexpr float TestingMsNoteOnBlock1 = 0.0f;
-	constexpr float TestingMsNoteOffBlock2 = 13.0f;
-
-	constexpr float TestingMsNoteOnBlock2 = 18.0f;
-	constexpr float TestingMsNoteOffBlock3 = 23.0f;
-	constexpr float TestingMsNoteOnBlock3 = 28.0f;
-	constexpr float TestingMsNoteOffBlock4 = 33.0f;	
-	constexpr float TestingMsNoteOnBlock4 = 38.0f;
-	
-	//last block tests for Non-Note events
-	constexpr float TestingMsNoteOffBlock5 = 43.0f;
-	constexpr float TestingMsEventBlock5 = 48.0f;
-
-	//number of test blocks to generate during tests
-	constexpr int32 NumBlocksToTest = 5;
-
-	UMidiFile* BuildMidiFile()
-	{
-		UMidiFile* TheMidiFile = NewObject<UMidiFile>();
-
-		//add tempo info bpm = 120
-		TheMidiFile->GetSongMaps()->GetTempoMap().AddTempoInfoPoint(MidiConstants::BPMToMidiTempo(120), 0);
-
-		//Add two midi tracks for testing
-		FMidiTrack* MidiNoteTrackDefault = TheMidiFile->AddTrack("MidiNoteTestTrackDefault");//track 1
-		FMidiTrack* MidiNoteTrackDifferent = TheMidiFile->AddTrack("MidiNoteTestTrackDifferent");//track 2
-
-		/**
-		* Midi Events for Testing Block 1:
-		* Add 2 Note On events in the first block with different channel information and different note number for testing
-		*/
-		//Note On Events
-		FMidiEvent MidiChannelMatchNoteOn(TheMidiFile->GetSongMaps()->MsToTick(TestingMsNoteOnBlock1),
-			FMidiMsg::CreateNoteOn(TestMidiChannelShouldPass, TestNoteNumberShouldPass, TestNoteVelocityShouldPass));
-		//this should not pass
-		FMidiEvent MidiChannelMismatchNoteOn(TheMidiFile->GetSongMaps()->MsToTick(TestingMsNoteOnBlock1),
-			FMidiMsg::CreateNoteOn(TestMidiChannelShouldNotPass, TestNoteNumberShouldNotPass, TestNoteVelocityShouldNotPass));
-		
-		//Add 2 Note On events to default track 1
-		MidiNoteTrackDefault->AddEvent(MidiChannelMismatchNoteOn);
-		MidiNoteTrackDefault->AddEvent(MidiChannelMatchNoteOn);
-
-
-		/**
-		 * Midi Events for Testing Block 2:
-		 * Add 2 Note Off events for the 2 Note On events in the previous block
-		 * (NOTE: Putting Note On and Note Off Events in separate blocks for testing the number of output triggers
-		 * since Note Offs will reset output)
-		 * Add 2 Note On events to 2 separate midi tracks at the same tick, default is track 1
-		 */
-		 //Note Off Events pairing with Note On events in block 1
-		FMidiEvent MidiChannelMatchNoteOff(TheMidiFile->GetSongMaps()->MsToTick(TestingMsNoteOffBlock2),
-			FMidiMsg::CreateNoteOff(TestMidiChannelShouldPass, TestNoteNumberShouldPass));
-		//this should not pass
-		FMidiEvent MidiChannelMismatchNoteOff(TheMidiFile->GetSongMaps()->MsToTick(TestingMsNoteOffBlock2),
-			FMidiMsg::CreateNoteOff(TestMidiChannelShouldNotPass, TestNoteNumberShouldNotPass));
-		//Add Note Off events to track 1
-		MidiNoteTrackDefault->AddEvent(MidiChannelMismatchNoteOff);
-		MidiNoteTrackDefault->AddEvent(MidiChannelMatchNoteOff);
-
-		FMidiEvent MidiTrackMatchNoteOn(TheMidiFile->GetSongMaps()->MsToTick(TestingMsNoteOnBlock2),
-			FMidiMsg::CreateNoteOn(TestMidiChannelShouldPass, TestNoteNumberShouldPass, TestNoteVelocityShouldPass));
-		//this should not pass
-		FMidiEvent MidiTrackMismatchNoteOn(TheMidiFile->GetSongMaps()->MsToTick(TestingMsNoteOnBlock2),
-			FMidiMsg::CreateNoteOn(TestMidiChannelShouldPass, TestNoteNumberShouldNotPass, TestNoteVelocityShouldNotPass));
-
-		//Add 1 Note On event to default track 1 and the other to track 2
-		MidiNoteTrackDifferent->AddEvent(MidiTrackMismatchNoteOn);
-		MidiNoteTrackDefault->AddEvent(MidiTrackMatchNoteOn);
-
-		/**
-		* Midi Events for Testing Block 3:
-		* Add 2 Note Off events for the 2 Note On events in the previous block
-		* Add 2 pairs of Note On events that have note numbers less than min note number or greater than max note number
-		* Neither should pass through to output 
-		*/
-		//Not Off Events pairing with Note On events from the previous block
-		FMidiEvent MidiTrackMatchNoteOff(TheMidiFile->GetSongMaps()->MsToTick(TestingMsNoteOffBlock3),
-			FMidiMsg::CreateNoteOff(TestMidiChannelShouldPass, TestNoteNumberShouldPass));
-		//this should not pass
-		FMidiEvent MidiTrackMismatchNoteOff(TheMidiFile->GetSongMaps()->MsToTick(TestingMsNoteOffBlock3),
-			FMidiMsg::CreateNoteOff(TestMidiChannelShouldPass, TestNoteNumberShouldNotPass));
-		//Add Note Off events to their corresponding tracks
-		MidiNoteTrackDifferent->AddEvent(MidiTrackMismatchNoteOff);
-		MidiNoteTrackDefault->AddEvent(MidiTrackMatchNoteOff);
-
-		//Note On Events that have Note Numbers greater than the max or less than min Note Number
-		FMidiEvent NoteNumberLessThanMinNoteOn(TheMidiFile->GetSongMaps()->MsToTick(TestingMsNoteOnBlock3),
-			FMidiMsg::CreateNoteOn(TestMidiChannelShouldPass, TestNoteNumberLessThanMin, TestNoteVelocityShouldPass));
-		FMidiEvent NoteNumberGreaterThanMaxNoteOn(TheMidiFile->GetSongMaps()->MsToTick(TestingMsNoteOnBlock3),
-			FMidiMsg::CreateNoteOn(TestMidiChannelShouldPass, TestNoteNumberGreaterThanMax, TestNoteVelocityShouldPass));
-		//Add 2 Note On events to default track 1, neither should pass
-		MidiNoteTrackDefault->AddEvent(NoteNumberLessThanMinNoteOn);
-		MidiNoteTrackDefault->AddEvent(NoteNumberGreaterThanMaxNoteOn);
-
-		/**
-		* Midi Events for Testing Block 4:
-		* Add 2 Note Off events for the 2 Note On events in the previous block
-		* Add 1 event that has a different type (should NOT pass)
-		*/
-		//Not Off events pairing with Note On events from the previous block
-		FMidiEvent NoteNumberLessThanMinNoteOff(TheMidiFile->GetSongMaps()->MsToTick(TestingMsNoteOffBlock4),
-			FMidiMsg::CreateNoteOff(TestMidiChannelShouldPass, TestNoteNumberLessThanMin));
-		FMidiEvent NoteNumberGreaterThanMaxNoteOff(TheMidiFile->GetSongMaps()->MsToTick(TestingMsNoteOffBlock4),
-			FMidiMsg::CreateNoteOff(TestMidiChannelShouldPass, TestNoteNumberGreaterThanMax));
-		//Add 2 Note Off events to default track 1, neither should pass
-		MidiNoteTrackDefault->AddEvent(NoteNumberLessThanMinNoteOff);
-		MidiNoteTrackDefault->AddEvent(NoteNumberGreaterThanMaxNoteOff);
-		
-		//Note On Events that have velocity values greater than the max or less than min velocity value
-		FMidiEvent NoteVelocityLessThanMinNoteOn(TheMidiFile->GetSongMaps()->MsToTick(TestingMsNoteOnBlock4),
-			FMidiMsg::CreateNoteOn(TestMidiChannelShouldPass, TestNoteNumberShouldPass, TestVelocityLessThanMin));
-		FMidiEvent NoteVelocityGreaterThanMaxNoteOn(TheMidiFile->GetSongMaps()->MsToTick(TestingMsNoteOnBlock4),
-			FMidiMsg::CreateNoteOn(TestMidiChannelShouldPass, TestNoteNumberShouldPass, TestVelocityGreaterThanMax));
-		//Add 2 Note On events to default track 1, neither should pass
-		MidiNoteTrackDefault->AddEvent(NoteVelocityLessThanMinNoteOn);
-		MidiNoteTrackDefault->AddEvent(NoteVelocityGreaterThanMaxNoteOn);
-		
-
-
-		/**
-		* Midi Events for Testing Block 5:
-		* Add 2 Note Off events for the 2 Note On events in the previous block
-		* Add 1 event that has a different type (should NOT pass), this one is a Pitch Bend event
-		*/
-		//Not Off events pairing with Note On events from the previous block
-		FMidiEvent NoteVelocityLessThanMinNoteOff(TheMidiFile->GetSongMaps()->MsToTick(TestingMsNoteOffBlock5),
-			FMidiMsg::CreateNoteOff(TestMidiChannelShouldPass, TestNoteNumberShouldPass));
-		FMidiEvent NoteVelocityGreaterThanMaxNoteOff(TheMidiFile->GetSongMaps()->MsToTick(TestingMsNoteOffBlock5),
-			FMidiMsg::CreateNoteOff(TestMidiChannelShouldPass, TestNoteNumberShouldPass));
-
-		//Add Note Off events to default track 1, neither should pass
-		MidiNoteTrackDefault->AddEvent(NoteVelocityLessThanMinNoteOff);
-		MidiNoteTrackDefault->AddEvent(NoteVelocityGreaterThanMaxNoteOff);
-
-		//Add a Pitch Bend event to default track 1, should not pass
-		FMidiEvent EventTypeMismatch(TheMidiFile->GetSongMaps()->MsToTick(TestingMsEventBlock5),
-			FMidiMsg(MidiConstants::kPitch, TestNoteNumberShouldPass2, TestNoteVelocityShouldPass));
-		MidiNoteTrackDefault->AddEvent(EventTypeMismatch);
-
-		//call TracksChanged() to update SongLengthData
-		TheMidiFile->TracksChanged();
-
-		//conform midi file length to an integer bar length
-		TheMidiFile->ConformMidiFileLength(EMidiFileLengthConformOption::RoundUp);
-
-		return TheMidiFile;
-	}
-
-	using GraphBuilder = Metasound::Test::FNodeTestGraphBuilder;
+	using FGraphBuilder = Metasound::Test::FNodeTestGraphBuilder;
 	using namespace Metasound;
 	using namespace Metasound::Frontend;
 	using namespace HarmonixMetasound;
 
 	IMPLEMENT_SIMPLE_AUTOMATION_TEST(
 		FMidiNoteTriggerCreateNodeTest,
-		"Harmonix.Metasound.Nodes.MidiNoteTriggerNode.CreateNode",
+		"Harmonix.Metasound.Nodes.MidiNoteTriggerNode.Defaults",
 		EAutomationTestFlags::EditorContext | EAutomationTestFlags::EngineFilter)
 		bool FMidiNoteTriggerCreateNodeTest::RunTest(const FString&)
 	{
 		// Build the graph.
 		constexpr int32 NumSamplesPerBlock = 256;
-		const TUniquePtr<FMetasoundGenerator> Generator = GraphBuilder::MakeSingleNodeGraph(
-			{ HarmonixMetasound::HarmonixNodeNamespace, "MidiNoteTrigger", "" },
-			0,
+		const TUniquePtr<FMetasoundGenerator> Generator = FGraphBuilder::MakeSingleNodeGraph(
+			Nodes::MidiNoteTriggerNode::GetClassName(),
+			Nodes::MidiNoteTriggerNode::GetCurrentMajorVersion(),
 			48000,
 			NumSamplesPerBlock);
 		UTEST_TRUE("Graph successfully built", Generator.IsValid());
@@ -237,33 +37,25 @@ namespace HarmonixMetasoundTests::MidiNoteTriggerNode
 			Generator->OnGenerateAudio(Buffer.GetRawChannelData(0), Buffer.GetNumTotalValidSamples());
 		}
 
-		// Validate output.
-		TOptional<FTriggerReadRef> OutputNoteOn = Generator->GetOutputReadReference<FTrigger>(CommonPinNames::Outputs::NoteOnName);
+		// Validate defaults.
+		TOptional<FTriggerReadRef> OutputNoteOn = Generator->GetOutputReadReference<FTrigger>(Nodes::MidiNoteTriggerNode::Outputs::NoteOnTriggerName);
 		UTEST_TRUE("Output exists", OutputNoteOn.IsSet());
 		UTEST_EQUAL("Number of Note On triggers", (*OutputNoteOn)->NumTriggeredInBlock(), 0);
 		UTEST_EQUAL("Note On Trigger sample", (*OutputNoteOn)->First(), -1);
 
-		TOptional<FTriggerReadRef> OutputNoteOff = Generator->GetOutputReadReference<FTrigger>(CommonPinNames::Outputs::NoteOffName);
+		TOptional<FTriggerReadRef> OutputNoteOff = Generator->GetOutputReadReference<FTrigger>(Nodes::MidiNoteTriggerNode::Outputs::NoteOffTriggerName);
 		UTEST_TRUE("Output exists", OutputNoteOff.IsSet());
 		UTEST_EQUAL("Number of Note Off triggers", (*OutputNoteOff)->NumTriggeredInBlock(), 0);
 		UTEST_EQUAL("Note Off Trigger sample", (*OutputNoteOff)->First(), -1);
 
-		TOptional<FInt32ReadRef> OutputMidiNoteNum = Generator->GetOutputReadReference<int32>("Midi Note #");
+		TOptional<FInt32ReadRef> OutputMidiNoteNum = Generator->GetOutputReadReference<int32>(Nodes::MidiNoteTriggerNode::Outputs::MidiNoteNumberName);
 		UTEST_TRUE("Output exists", OutputMidiNoteNum.IsSet());
 		UTEST_EQUAL("Midi Note Num check", **OutputMidiNoteNum, 0);
-
-		TOptional<FFloatReadRef> OutputFrequency = Generator->GetOutputReadReference<float>("Frequency");
-		UTEST_TRUE("Output exists", OutputFrequency.IsSet());
-		UTEST_EQUAL("Frequency check", **OutputFrequency, 0.0f);
-
-		TOptional<FInt32ReadRef> OutputVelocity = Generator->GetOutputReadReference<int32>("Velocity");
+		
+		TOptional<FInt32ReadRef> OutputVelocity = Generator->GetOutputReadReference<int32>(Nodes::MidiNoteTriggerNode::Outputs::MidiVelocityName);
 		UTEST_TRUE("Output exists", OutputVelocity.IsSet());
 		UTEST_EQUAL("Velocity check", **OutputVelocity, 0);
 
-		TOptional<FFloatReadRef> OutputNormalizedVelocity = Generator->GetOutputReadReference<float>("Normalized Velocity");
-		UTEST_TRUE("Output exists", OutputNormalizedVelocity.IsSet());
-		UTEST_EQUAL("Normalized Velocity check", **OutputNormalizedVelocity, 0.0f);
-
 		return true;
 	}
 
@@ -273,207 +65,137 @@ namespace HarmonixMetasoundTests::MidiNoteTriggerNode
 		EAutomationTestFlags::EditorContext | EAutomationTestFlags::EngineFilter)
 		bool FMidiNoteTriggerNodeTestBasic::RunTest(const FString&)
 	{
-		//************************************************************************************************************************
 		// Build the graph.
-		//************************************************************************************************************************
-		GraphBuilder Builder;
-		Builder.AddOutput("AudioOut", GetMetasoundDataTypeName<FAudioBuffer>());
-		//transport
-		Frontend::FNodeHandle TransportNode = Builder.AddNode({ HarmonixMetasound::HarmonixNodeNamespace, "TriggerToTransport", "" }, 0);
-		UTEST_TRUE("Transport Node Created", TransportNode.Get().IsValid());
-		// hoist up play trigger input
-		Builder.AddAndConnectDataReferenceInput(TransportNode, CommonPinNames::Inputs::TransportPlayName, GetMetasoundDataTypeName<FTrigger>());
-		// midi player
-		Frontend::FNodeHandle MidiPlayerNode = Builder.AddNode({ HarmonixMetasound::HarmonixNodeNamespace, "MidiPlayer", "" }, 0);
-		UTEST_TRUE("MidiPlayer Node Created", MidiPlayerNode.Get().IsValid());
-		// hoist up midi asset input
-		Builder.AddAndConnectDataReferenceInput(MidiPlayerNode, CommonPinNames::Inputs::MidiFileAssetName, GetMetasoundDataTypeName<FMidiAsset>());
-		//connect transport to midi player
-		bool ConnectionSuccess = Builder.ConnectNodes(TransportNode, CommonPinNames::Outputs::TransportName, MidiPlayerNode, CommonPinNames::Inputs::TransportName);
-		UTEST_TRUE("Transport Out connected to MidiPlayer node", ConnectionSuccess);
-		// Midi Note Trigger node to test
-		FNodeHandle NoteTriggerNode = Builder.AddNode({ HarmonixMetasound::HarmonixNodeNamespace, "MidiNoteTrigger", "" }, 0);
-		UTEST_TRUE("MidiNoteTrigger Node Created", NoteTriggerNode.Get().IsValid());
-
-		//Midi Input Channel Number 
-		Builder.AddAndConnectDataReferenceInput(NoteTriggerNode, CommonPinNames::Inputs::MidiChannelNumberName, GetMetasoundDataTypeName<int32>());
-		//Midi Input Track Number
-		Builder.AddAndConnectDataReferenceInput(NoteTriggerNode, CommonPinNames::Inputs::MidiTrackNumberName, GetMetasoundDataTypeName<int32>());
-		//min Input Midi Note Number
-		Builder.AddAndConnectDataReferenceInput(NoteTriggerNode, CommonPinNames::Inputs::MinMidiNoteName, GetMetasoundDataTypeName<int32>());
-		//Max Input Midi Note Number
-		Builder.AddAndConnectDataReferenceInput(NoteTriggerNode, CommonPinNames::Inputs::MaxMidiNoteName, GetMetasoundDataTypeName<int32>());
-		//Input Min Velocity
-		Builder.AddAndConnectDataReferenceInput(NoteTriggerNode, CommonPinNames::Inputs::MinMidiVelocityName, GetMetasoundDataTypeName<int32>());
-		//Input Max Velocity
-		Builder.AddAndConnectDataReferenceInput(NoteTriggerNode, CommonPinNames::Inputs::MaxMidiVelocityName, GetMetasoundDataTypeName<int32>());
-		
-
-		//output Note On trigger
-		Builder.AddAndConnectDataReferenceOutput(NoteTriggerNode, CommonPinNames::Outputs::NoteOnName, GetMetasoundDataTypeName<FTrigger>(), CommonPinNames::Outputs::NoteOnName);
-		//output Note Off trigger
-		Builder.AddAndConnectDataReferenceOutput(NoteTriggerNode, CommonPinNames::Outputs::NoteOffName, GetMetasoundDataTypeName<FTrigger>(), CommonPinNames::Outputs::NoteOffName);
-		
-		//output Midi Note number
-		Builder.AddAndConnectDataReferenceOutput(NoteTriggerNode, CommonPinNames::Outputs::MidiNoteNumberName, GetMetasoundDataTypeName<int32>(), CommonPinNames::Outputs::MidiNoteNumberName);
-
-		//output Midi Note Velocity (Int32)
-		Builder.AddAndConnectDataReferenceOutput(NoteTriggerNode, CommonPinNames::Outputs::MidiVelocityName, GetMetasoundDataTypeName<int32>(), CommonPinNames::Outputs::MidiVelocityName);
-		//output normalized Midi Note Velocity (float)
-		Builder.AddAndConnectDataReferenceOutput(NoteTriggerNode, CommonPinNames::Outputs::NormalizedVelocityName, GetMetasoundDataTypeName<float>(), CommonPinNames::Outputs::NormalizedVelocityName);
-
-		//output frequency 
-		Builder.AddAndConnectDataReferenceOutput(NoteTriggerNode, CommonPinNames::Outputs::FrequencyName, GetMetasoundDataTypeName<float>(), CommonPinNames::Outputs::FrequencyName);
-
-		//connect midi player node to Note Trigger node
-		ConnectionSuccess = Builder.ConnectNodes(MidiPlayerNode, CommonPinNames::Outputs::MidiStreamName, NoteTriggerNode, CommonPinNames::Inputs::MidiStreamName);
-		UTEST_TRUE("MidiPlayer Out successfully connected to Midi Note Trigger node", ConnectionSuccess);
-
-		//************************************************************************************************************************
-		// Make the generator.
-		//************************************************************************************************************************
-		constexpr int32 NumSamplesPerBlock = 480; // 10ms per block for sample rate = 48000
-		const TUniquePtr<FMetasoundGenerator> Generator = Builder.BuildGenerator(48000, NumSamplesPerBlock);
+		constexpr int32 NumSamplesPerBlock = 256;
+		const TUniquePtr<FMetasoundGenerator> Generator = FGraphBuilder::MakeSingleNodeGraph(
+			Nodes::MidiNoteTriggerNode::GetClassName(),
+			Nodes::MidiNoteTriggerNode::GetCurrentMajorVersion(),
+			48000,
+			NumSamplesPerBlock);
 		UTEST_TRUE("Graph successfully built", Generator.IsValid());
 
-		//************************************************************************************************************************
-		// Get reference to outputs we want to check.
-		//************************************************************************************************************************
-		
-		TOptional<FTriggerReadRef> OutputNoteOnTrigger = Generator->GetOutputReadReference<FTrigger>(CommonPinNames::Outputs::NoteOnName);
-		UTEST_TRUE("Output exists", OutputNoteOnTrigger.IsSet());
-		UTEST_EQUAL("Number of Note Triggers at start", (*OutputNoteOnTrigger)->NumTriggeredInBlock(), 0);
-		UTEST_EQUAL("Note On Output Trigger sample", (*OutputNoteOnTrigger)->First(), -1);
+		// Get the outputs
+		TOptional<FTriggerReadRef> OutputNoteOn = Generator->GetOutputReadReference<FTrigger>(Nodes::MidiNoteTriggerNode::Outputs::NoteOnTriggerName);
+		UTEST_TRUE("Output exists", OutputNoteOn.IsSet());
+
+		TOptional<FTriggerReadRef> OutputNoteOff = Generator->GetOutputReadReference<FTrigger>(Nodes::MidiNoteTriggerNode::Outputs::NoteOffTriggerName);
+		UTEST_TRUE("Output exists", OutputNoteOff.IsSet());
+
+		TOptional<FInt32ReadRef> OutputMidiNoteNum = Generator->GetOutputReadReference<int32>(Nodes::MidiNoteTriggerNode::Outputs::MidiNoteNumberName);
+		UTEST_TRUE("Output exists", OutputMidiNoteNum.IsSet());
 		
-		TOptional<FTriggerReadRef> OutputNoteOffTrigger = Generator->GetOutputReadReference<FTrigger>(CommonPinNames::Outputs::NoteOffName);
-		UTEST_TRUE("Output exists", OutputNoteOffTrigger.IsSet());
-		UTEST_EQUAL("Number of Note Triggers at start", (*OutputNoteOffTrigger)->NumTriggeredInBlock(), 0);
-		UTEST_EQUAL("Note off Output Trigger sample", (*OutputNoteOffTrigger)->First(), -1);
-
-		TOptional<FInt32ReadRef> OutputMidiNoteNumber = Generator->GetOutputReadReference<int32>(CommonPinNames::Outputs::MidiNoteNumberName);
-		UTEST_TRUE("Output exists", OutputMidiNoteNumber.IsSet());
-		UTEST_EQUAL("Midi Note Number at start", **OutputMidiNoteNumber, 0);
-
-		TOptional<FInt32ReadRef> OutputMidiVelocity = Generator->GetOutputReadReference<int32>(CommonPinNames::Outputs::MidiVelocityName);
-		UTEST_TRUE("Output exists", OutputMidiVelocity.IsSet());
-		UTEST_EQUAL("Midi Velocity (int32) at start", **OutputMidiVelocity, 0);
-
-		TOptional<FFloatReadRef> OutputNormalizedVelocity = Generator->GetOutputReadReference<float>(CommonPinNames::Outputs::NormalizedVelocityName);
-		UTEST_TRUE("Output exists", OutputNormalizedVelocity.IsSet());
-		UTEST_EQUAL("Normalized Midi Velocity (float) at start", **OutputNormalizedVelocity, 0.0f);
-
-		TOptional<FFloatReadRef> OutputFrequency= Generator->GetOutputReadReference<float>(CommonPinNames::Outputs::FrequencyName);
-		UTEST_TRUE("Output exists", OutputFrequency.IsSet());
-		UTEST_EQUAL("Normalized Midi Velocity (float) at start", **OutputFrequency, 0.0f);
-
-		//************************************************************************************************************************
-		// Make initial inputs
-		//************************************************************************************************************************
-		// Make a Midi File to play by Midi Player
-		UMidiFile* MidiFile = BuildMidiFile();
-		TSharedPtr<Audio::IProxyData> MidiFileProxy = MidiFile->CreateProxyData({ "MidiNoteTriggerTest" });
-		Generator->SetInputValue<FMidiAsset>(CommonPinNames::Inputs::MidiFileAssetName, MidiFileProxy);
-
-		// Set the testing Midi Channel (default = 1)
-		Generator->SetInputValue<int32>(CommonPinNames::Inputs::MidiChannelNumberName, TestMidiChannelShouldPass + 1);
-		// Set the testing Midi Track (default = 1)
-		Generator->SetInputValue<int32>(CommonPinNames::Inputs::MidiTrackNumberName, TestMidiTrackShouldPass);
-		// Set the testing Min Midi Note
-		Generator->SetInputValue<int32>(CommonPinNames::Inputs::MinMidiNoteName, TestMinNoteNumberDefault);
-		// Set the testing Max Midi Note
-		Generator->SetInputValue<int32>(CommonPinNames::Inputs::MaxMidiNoteName, TestMaxNoteNumberDefault);
-		// Set the testing Min Velocity
-		Generator->SetInputValue<int32>(CommonPinNames::Inputs::MinMidiVelocityName, TestMinVelocityDefault);
-		// Set the testing Max Velocity
-		Generator->SetInputValue<int32>(CommonPinNames::Inputs::MaxMidiVelocityName, TestMaxVelocityDefault);
-
-		// Trigger transport
-		Generator->ApplyToInputValue<FTrigger>(CommonPinNames::Inputs::TransportPlayName, [](FTrigger& TriggerIn) { TriggerIn.TriggerFrame(0); });
-
-		//************************************************************************************************************************
-		// Generate 5 10ms blocks for testing various Midi Events in file/stream
-		//************************************************************************************************************************
-
-		TAudioBuffer<float> Buffer{ Generator->GetNumChannels(), NumSamplesPerBlock, EAudioBufferCleanupMode::Delete};
-			
-			//Block 1:
-			Generator->OnGenerateAudio(Buffer.GetRawChannelData(0), Buffer.GetNumTotalValidSamples());
-			//Expected: 1 Note On Trigger, 0 Note Off Trigger, 
-			// Note Number/Velocity/Frequency outputs contain values equal to inputs
-				
-			//check output triggers
-			UTEST_EQUAL("Number of Note On Output Triggers", (*OutputNoteOnTrigger)->NumTriggeredInBlock(), ExpectedTriggerCountOnPass);
-			UTEST_EQUAL("Number of Note Off Output Triggers", (*OutputNoteOffTrigger)->NumTriggeredInBlock(), ExpectedTriggerCountNotPass);
-
-			//check output values
-			UTEST_EQUAL("Midi Note Number check", **OutputMidiNoteNumber, TestNoteNumberShouldPass);
-			UTEST_EQUAL("Midi Velocity (int32) check", **OutputMidiVelocity, TestNoteVelocityShouldPass);
-			UTEST_EQUAL("Midi Normalized Velocity (float) check", **OutputNormalizedVelocity, TestNoteNormalizedVelocityShouldPass);
-			UTEST_EQUAL("Output frequency check", **OutputFrequency, ExpectedFrequencyOnPass);
-
-			//Block 2:
-			Generator->OnGenerateAudio(Buffer.GetRawChannelData(0), Buffer.GetNumTotalValidSamples());
-			//Expected: 1 Note Off Trigger (for previous Note On that also should pass), 1 Note On Trigger,
-			//Note Number/Velocity/Frequency outputs contain values equal to inputs
-
-			//check output triggers
-			UTEST_EQUAL("Number of Note Off Output Triggers", (*OutputNoteOffTrigger)->NumTriggeredInBlock(), ExpectedTriggerCountOnPass);
-			UTEST_EQUAL("Number of Note On Output Triggers", (*OutputNoteOnTrigger)->NumTriggeredInBlock(), ExpectedTriggerCountOnPass);
-				
-			//check output values
-			UTEST_EQUAL("Midi Note Number check", **OutputMidiNoteNumber, TestNoteNumberShouldPass);
-			UTEST_EQUAL("Midi Velocity (int32) check", **OutputMidiVelocity, TestNoteVelocityShouldPass);
-			UTEST_EQUAL("Midi Normalized Velocity (float) check", **OutputNormalizedVelocity, TestNoteNormalizedVelocityShouldPass);
-			UTEST_EQUAL("Output frequency check", **OutputFrequency, ExpectedFrequencyOnPass);
-
-			//Block 3:
-			Generator->OnGenerateAudio(Buffer.GetRawChannelData(0), Buffer.GetNumTotalValidSamples());
-			//Expected: 1 Note Off Trigger (for previous Note On that also should pass), 0 Note On Trigger
-			//Note Number/Frequency should be the same as the previous block
-			//Velocity should be 0
-				
-			//check output triggers
-			UTEST_EQUAL("Number of Note Off Output Triggers", (*OutputNoteOffTrigger)->NumTriggeredInBlock(), ExpectedTriggerCountOnPass);
-			UTEST_EQUAL("Number of Note On Output Triggers", (*OutputNoteOnTrigger)->NumTriggeredInBlock(), ExpectedTriggerCountNotPass);
-
-			//check output values
-			UTEST_EQUAL("Midi Note Number check", **OutputMidiNoteNumber, TestNoteNumberShouldPass);
-			UTEST_EQUAL("Midi Velocity (int32) check", **OutputMidiVelocity, ExpectedVelocityNotPass);
-			UTEST_EQUAL("Midi Normalized Velocity (float) check", **OutputNormalizedVelocity, ExpectedNormalizedVelocityNotPass);
-			UTEST_EQUAL("Output frequency check", **OutputFrequency, ExpectedFrequencyOnPass);
-
-
-			//Block 4:
-			Generator->OnGenerateAudio(Buffer.GetRawChannelData(0), Buffer.GetNumTotalValidSamples());
-			//Expected: 0 Note On/Note Off trigger,
-			//Note Number/Frequency should be the same as the previous block
-			//Velocity should be 0
+		TOptional<FInt32ReadRef> OutputVelocity = Generator->GetOutputReadReference<int32>(Nodes::MidiNoteTriggerNode::Outputs::MidiVelocityName);
+		UTEST_TRUE("Output exists", OutputVelocity.IsSet());
 
-			//check output triggers (should all be 0)
-			UTEST_EQUAL("Number of Note Off Output Triggers", (*OutputNoteOffTrigger)->NumTriggeredInBlock(), ExpectedTriggerCountNotPass);
-			UTEST_EQUAL("Number of Note On Output Triggers", (*OutputNoteOnTrigger)->NumTriggeredInBlock(), ExpectedTriggerCountNotPass);
+		// Send a note on and expect the right values
+		constexpr uint8 Note1Channel = 3;
+		constexpr uint8 Note1NoteNumber = 45;
+		constexpr uint8 Note1Velocity = 79;
+		{
+			// Add a note on
+			Generator->ApplyToInputValue<FMidiStream>(
+				Nodes::MidiNoteTriggerNode::Inputs::MidiStreamName,
+				[Note1Channel, Note1NoteNumber, Note1Velocity](FMidiStream& MidiStream)
+				{
+					MidiStream.AddMidiEvent({ static_cast<uint32>(0), FMidiMsg::CreateNoteOn(Note1Channel, Note1NoteNumber, Note1Velocity)});
+				});
+
+			// Render a block
+			FAudioBuffer Buffer{ Generator->OperatorSettings };
+			Generator->OnGenerateAudio(Buffer.GetData(), Buffer.Num());
+
+			// Make sure we got the right stuff out the other side
+			UTEST_EQUAL("Number of Note On triggers", (*OutputNoteOn)->NumTriggeredInBlock(), 1);
+			UTEST_EQUAL("Note On Trigger sample", (*OutputNoteOn)->First(), 0);
+
+			UTEST_EQUAL("Number of Note Off triggers", (*OutputNoteOff)->NumTriggeredInBlock(), 0);
+			UTEST_EQUAL("Note Off Trigger sample", (*OutputNoteOff)->First(), -1);
+
+			UTEST_EQUAL("Midi Note Num check", **OutputMidiNoteNum, Note1NoteNumber);
+
+			UTEST_EQUAL("Velocity check", **OutputVelocity, Note1Velocity);
+
+			// Advance the MIDI stream
+			Generator->ApplyToInputValue<FMidiStream>(
+				Nodes::MidiNoteTriggerNode::Inputs::MidiStreamName,
+				[](FMidiStream& MidiStream)
+				{
+					MidiStream.PrepareBlock();
+				});
+		}
 
-			//check output values
-			UTEST_EQUAL("Midi Note Number check", **OutputMidiNoteNumber, TestNoteNumberShouldPass);
-			UTEST_EQUAL("Midi Velocity (int32) check", **OutputMidiVelocity, ExpectedVelocityNotPass);
-			UTEST_EQUAL("Midi Normalized Velocity (float) check", **OutputNormalizedVelocity, ExpectedNormalizedVelocityNotPass);
-			UTEST_EQUAL("Output frequency check", **OutputFrequency, ExpectedFrequencyOnPass);
+		// Send a note off and expect the right values
+		{
+			// Add a note on
+			Generator->ApplyToInputValue<FMidiStream>(
+				Nodes::MidiNoteTriggerNode::Inputs::MidiStreamName,
+				[Note1Channel, Note1NoteNumber](FMidiStream& MidiStream)
+				{
+					MidiStream.AddMidiEvent({ static_cast<uint32>(0), FMidiMsg::CreateNoteOff(Note1Channel, Note1NoteNumber)});
+				});
+
+			// Render a block
+			FAudioBuffer Buffer{ Generator->OperatorSettings };
+			Generator->OnGenerateAudio(Buffer.GetData(), Buffer.Num());
+
+			// Make sure we got the right stuff out the other side
+			UTEST_EQUAL("Number of Note On triggers", (*OutputNoteOn)->NumTriggeredInBlock(), 0);
+			UTEST_EQUAL("Note On Trigger sample", (*OutputNoteOn)->First(), -1);
+
+			UTEST_EQUAL("Number of Note Off triggers", (*OutputNoteOff)->NumTriggeredInBlock(), 1);
+			UTEST_EQUAL("Note Off Trigger sample", (*OutputNoteOff)->First(), 0);
+
+			UTEST_EQUAL("Midi Note Num check", **OutputMidiNoteNum, Note1NoteNumber);
+
+			UTEST_EQUAL("Velocity check", **OutputVelocity, 0);
+
+			// Advance the MIDI stream
+			Generator->ApplyToInputValue<FMidiStream>(
+				Nodes::MidiNoteTriggerNode::Inputs::MidiStreamName,
+				[](FMidiStream& MidiStream)
+				{
+					MidiStream.PrepareBlock();
+				});
+		}
 
-			//Block 5:
-			Generator->OnGenerateAudio(Buffer.GetRawChannelData(0), Buffer.GetNumTotalValidSamples());
-			//Expected: 0 Note On/Note Off trigger,
-			//Note Number/Frequency should be the same as the previous block
-			//Velocity should be 0
-
-			//check output triggers (should all be 0)
-			UTEST_EQUAL("Number of Note Off Output Triggers", (*OutputNoteOffTrigger)->NumTriggeredInBlock(), ExpectedTriggerCountNotPass);
-			UTEST_EQUAL("Number of Note On Output Triggers", (*OutputNoteOnTrigger)->NumTriggeredInBlock(), ExpectedTriggerCountNotPass);
-
-			//check output values
-			UTEST_EQUAL("Midi Note Number check", **OutputMidiNoteNumber, TestNoteNumberShouldPass);
-			UTEST_EQUAL("Midi Velocity (int32) check", **OutputMidiVelocity, ExpectedVelocityNotPass);
-			UTEST_EQUAL("Midi Normalized Velocity (float) check", **OutputNormalizedVelocity, ExpectedNormalizedVelocityNotPass);
-			UTEST_EQUAL("Output frequency check", **OutputFrequency, ExpectedFrequencyOnPass);
+		// Send a note on, but turn off the node, expect the same values as before
+		constexpr uint8 Note2NoteNumber = 98;
+		constexpr uint8 Note2Velocity = 123;
+		{
+			// Turn off the node
+			Generator->SetInputValue<bool>(Nodes::MidiNoteTriggerNode::Inputs::EnableName, false);
+			
+			// Add a note on
+			Generator->ApplyToInputValue<FMidiStream>(
+				Nodes::MidiNoteTriggerNode::Inputs::MidiStreamName,
+				[Note1Channel, Note2NoteNumber, Note2Velocity](FMidiStream& MidiStream)
+				{
+					MidiStream.AddMidiEvent({ static_cast<uint32>(0), FMidiMsg::CreateNoteOn(Note1Channel, Note2NoteNumber, Note2Velocity)});
+				});
+
+			// Render a block
+			FAudioBuffer Buffer{ Generator->OperatorSettings };
+			Generator->OnGenerateAudio(Buffer.GetData(), Buffer.Num());
+
+			// Make sure we got the right stuff out the other side
+			UTEST_EQUAL("Number of Note On triggers", (*OutputNoteOn)->NumTriggeredInBlock(), 0);
+			UTEST_EQUAL("Note On Trigger sample", (*OutputNoteOn)->First(), -1);
+
+			UTEST_EQUAL("Number of Note Off triggers", (*OutputNoteOff)->NumTriggeredInBlock(), 0);
+			UTEST_EQUAL("Note Off Trigger sample", (*OutputNoteOff)->First(), -1);
+
+			UTEST_EQUAL("Midi Note Num check", **OutputMidiNoteNum, Note1NoteNumber);
+
+			UTEST_EQUAL("Velocity check", **OutputVelocity, 0);
+
+			// Advance the MIDI stream
+			Generator->ApplyToInputValue<FMidiStream>(
+				Nodes::MidiNoteTriggerNode::Inputs::MidiStreamName,
+				[](FMidiStream& MidiStream)
+				{
+					MidiStream.PrepareBlock();
+				});
+		}
 
 		return true;
 	}
