Fix issue with jittery music clock when tempo slowed down using a Midi Player node
#tests PIE

diff --git a/Engine/Plugins/Runtime/Harmonix/Source/HarmonixMetasound/Private/HarmonixMetasound/Nodes/MidiPlayerNode.cpp b/Engine/Plugins/Runtime/Harmonix/Source/HarmonixMetasound/Private/HarmonixMetasound/Nodes/MidiPlayerNode.cpp
--- a/Engine/Plugins/Runtime/Harmonix/Source/HarmonixMetasound/Private/HarmonixMetasound/Nodes/MidiPlayerNode.cpp
+++ b/Engine/Plugins/Runtime/Harmonix/Source/HarmonixMetasound/Private/HarmonixMetasound/Nodes/MidiPlayerNode.cpp
@@ -440,8 +440,9 @@ namespace HarmonixMetasound
 						float Ms = MidiClockOut->GetSongMaps().TickToMs(Tick);
 
 						float ClockInSpeed = MidiClockIn->GetSpeedAtBlockSampleFrame(StartFrameIndex);
+						float AdvanceRatio = MidiClockIn->GetSongMaps().GetTempoAtTick(Event.StartTick) / MidiClockOut->GetSongMaps().GetTempoAtTick(Tick);
 						// midi clock needs to know how fast its advancing based on their authority
-						MidiClockOut->InformOfCurrentAdvanceRate(ClockInSpeed * *SpeedMultInPin);
+						MidiClockOut->InformOfCurrentAdvanceRate(ClockInSpeed * *SpeedMultInPin * AdvanceRatio);
 						MidiClockOut->AdvanceHiResToMs(Event.BlockFrameIndex, Ms, true);
 
 						// we have to update our loop offset index _after_ the advance
