[Harmonix] Fix loc name collisions for MIDI note trigger nodes

#tests automated tests

diff --git a/Engine/Plugins/Runtime/Harmonix/Source/HarmonixMetasound/Private/HarmonixMetasound/Nodes/MidiNoteTriggerNode.cpp b/Engine/Plugins/Runtime/Harmonix/Source/HarmonixMetasound/Private/HarmonixMetasound/Nodes/MidiNoteTriggerNode.cpp
--- a/Engine/Plugins/Runtime/Harmonix/Source/HarmonixMetasound/Private/HarmonixMetasound/Nodes/MidiNoteTriggerNode.cpp
+++ b/Engine/Plugins/Runtime/Harmonix/Source/HarmonixMetasound/Private/HarmonixMetasound/Nodes/MidiNoteTriggerNode.cpp
@@ -57,8 +57,8 @@ namespace HarmonixMetasound::Nodes::MidiNoteTriggerNode
 				Info.ClassName        = GetClassName();
 				Info.MajorVersion     = 1;
 				Info.MinorVersion     = 0;
-				Info.DisplayName      = METASOUND_LOCTEXT("MidiNoteTriggerNode_DisplayName", "MIDI Note Trigger");
-				Info.Description      = METASOUND_LOCTEXT("MidiNoteTriggerNode_Description", "Outputs triggers and info for incoming MIDI notes.");
+				Info.DisplayName      = METASOUND_LOCTEXT("MidiNoteTriggerNodeV1_DisplayName", "MIDI Note Trigger");
+				Info.Description      = METASOUND_LOCTEXT("MidiNoteTriggerNodeV1_Description", "Outputs triggers and info for incoming MIDI notes.");
 				Info.Author           = PluginAuthor;
 				Info.PromptIfMissing  = PluginNodeMissingPrompt;
 				Info.DefaultInterface = GetVertexInterface();
@@ -300,8 +300,8 @@ namespace HarmonixMetasound::Nodes::MidiNoteTriggerNode
 			Info.ClassName        = GetClassName();
 			Info.MajorVersion     = 0;
 			Info.MinorVersion     = 1;
-			Info.DisplayName      = METASOUND_LOCTEXT("MidiNoteTriggerNode_DisplayName", "Midi Note Trigger");
-			Info.Description      = METASOUND_LOCTEXT("MidiNoteTriggerNode_Description", "Receives a midi stream, filters for the desired messages, and outputs triggers.");
+			Info.DisplayName      = METASOUND_LOCTEXT("MidiNoteTriggerNodeV0_DisplayName", "Midi Note Trigger");
+			Info.Description      = METASOUND_LOCTEXT("MidiNoteTriggerNodeV0_Description", "Receives a midi stream, filters for the desired messages, and outputs triggers.");
 			Info.Author           = PluginAuthor;
 			Info.PromptIfMissing  = PluginNodeMissingPrompt;
 			Info.DefaultInterface = GetVertexInterface();
