Less aggressive asset validation for FusionPatches. Only warn about upgrading assets when version is no longer supported

diff --git a/Engine/Plugins/Runtime/Harmonix/Source/HarmonixDsp/Private/HarmonixDsp/FusionSampler/FusionPatch.cpp b/Engine/Plugins/Runtime/Harmonix/Source/HarmonixDsp/Private/HarmonixDsp/FusionSampler/FusionPatch.cpp
--- a/Engine/Plugins/Runtime/Harmonix/Source/HarmonixDsp/Private/HarmonixDsp/FusionSampler/FusionPatch.cpp
+++ b/Engine/Plugins/Runtime/Harmonix/Source/HarmonixDsp/Private/HarmonixDsp/FusionSampler/FusionPatch.cpp
@@ -144,9 +144,9 @@ EDataValidationResult UFusionPatch::IsDataValid(FDataValidationContext& Context)
 	EDataValidationResult Result = Super::IsDataValid(Context);
 	const int32 Version = GetLinkerCustomVersion(FFusionPatchCustomVersion::GUID);
 
-	if (Version < FFusionPatchCustomVersion::LatestVersion)
+	if (Version < FFusionPatchCustomVersion::PitchShifterNameRedirects)
 	{
-		Context.AddWarning(FText::Format(INVTEXT("Asset saved with outdated version: {0}. Resave asset to latest version: {1}."),
+		Context.AddWarning(FText::Format(INVTEXT("Asset saved with a version that is no longer supported: {0}. Resave asset to latest version: {1}."),
 			Version, FFusionPatchCustomVersion::LatestVersion));
 
 		Result = EDataValidationResult::Invalid;
