Don't name your array sizes `Num` for your C-style arrays you're using as UPROPERTYs. UhtProperty.cs will think you meant it to be an ArraySizeEnum.

#tests PIE

diff --git a/Engine/Plugins/Runtime/Harmonix/Source/HarmonixDsp/Public/HarmonixDsp/Modulators/Settings/AdsrSettings.h b/Engine/Plugins/Runtime/Harmonix/Source/HarmonixDsp/Public/HarmonixDsp/Modulators/Settings/AdsrSettings.h
--- a/Engine/Plugins/Runtime/Harmonix/Source/HarmonixDsp/Public/HarmonixDsp/Modulators/Settings/AdsrSettings.h
+++ b/Engine/Plugins/Runtime/Harmonix/Source/HarmonixDsp/Public/HarmonixDsp/Modulators/Settings/AdsrSettings.h
@@ -151,8 +151,6 @@ struct FAdsrSettingsArray
 {
 	GENERATED_BODY()
 
-	static constexpr int32 Num = 2;
-
 	UPROPERTY()
-	FAdsrSettings Array[Num];
+	FAdsrSettings Array[2];
 };
diff --git a/Engine/Plugins/Runtime/Harmonix/Source/HarmonixDsp/Public/HarmonixDsp/Modulators/Settings/LfoSettings.h b/Engine/Plugins/Runtime/Harmonix/Source/HarmonixDsp/Public/HarmonixDsp/Modulators/Settings/LfoSettings.h
--- a/Engine/Plugins/Runtime/Harmonix/Source/HarmonixDsp/Public/HarmonixDsp/Modulators/Settings/LfoSettings.h
+++ b/Engine/Plugins/Runtime/Harmonix/Source/HarmonixDsp/Public/HarmonixDsp/Modulators/Settings/LfoSettings.h
@@ -106,8 +106,6 @@ struct FLfoSettingsArray
 {
 	GENERATED_BODY()
 
-	static constexpr int32 Num = 2;
-	
 	UPROPERTY()
-	FLfoSettings Array[Num];
+	FLfoSettings Array[2];
 };
\ No newline at end of file
diff --git a/Engine/Plugins/Runtime/Harmonix/Source/HarmonixDsp/Public/HarmonixDsp/Modulators/Settings/ModulatorSettings.h b/Engine/Plugins/Runtime/Harmonix/Source/HarmonixDsp/Public/HarmonixDsp/Modulators/Settings/ModulatorSettings.h
--- a/Engine/Plugins/Runtime/Harmonix/Source/HarmonixDsp/Public/HarmonixDsp/Modulators/Settings/ModulatorSettings.h
+++ b/Engine/Plugins/Runtime/Harmonix/Source/HarmonixDsp/Public/HarmonixDsp/Modulators/Settings/ModulatorSettings.h
@@ -37,7 +37,6 @@ struct FModulatorSettingsArray
 {
 	GENERATED_BODY()
 
-	static constexpr int32 Num = 2;
 	UPROPERTY()
-	FModulatorSettings Array[Num];
+	FModulatorSettings Array[2];
 };
\ No newline at end of file
