When importing a FusionPatch, copy the fusion patch Settings and Keyzones, modify those, and then Update the FusionPatch with the imported settings only if there are no errors to avoid importing with corrupted data.

#tests PIE

diff --git a/Engine/Plugins/Runtime/Harmonix/Source/HarmonixDsp/Private/HarmonixDsp/FusionSampler/FusionPatch.cpp b/Engine/Plugins/Runtime/Harmonix/Source/HarmonixDsp/Private/HarmonixDsp/FusionSampler/FusionPatch.cpp
--- a/Engine/Plugins/Runtime/Harmonix/Source/HarmonixDsp/Private/HarmonixDsp/FusionSampler/FusionPatch.cpp
+++ b/Engine/Plugins/Runtime/Harmonix/Source/HarmonixDsp/Private/HarmonixDsp/FusionSampler/FusionPatch.cpp
@@ -52,6 +52,25 @@ void UFusionPatch::UpdateRenderableForNonTrivialChange()
 	RenderableFusionPatchData->SetNewSettings(NewSettings);
 }
 
+void UFusionPatch::UpdateSettings(const FFusionPatchSettings& InSettings)
+{
+	FusionPatchData.Settings = InSettings;
+	UpdateRenderableForNonTrivialChange();
+}
+
+void UFusionPatch::UpdateKeyzones(const TArray<FKeyzoneSettings>& NewKeyzones)
+{
+	FusionPatchData.Keyzones = NewKeyzones;
+	UpdateRenderableForNonTrivialChange();
+}
+
+
+void UFusionPatch::UpdatePatch(const FFusionPatchData& InPatchData)
+{
+	FusionPatchData = InPatchData;
+	UpdateRenderableForNonTrivialChange();
+}
+
 #if WITH_EDITORONLY_DATA
 
 void UFusionPatch::PostEditChangeProperty(FPropertyChangedEvent& PropertyChangedEvent)
diff --git a/Engine/Plugins/Runtime/Harmonix/Source/HarmonixDsp/Public/HarmonixDsp/FusionSampler/FusionPatch.h b/Engine/Plugins/Runtime/Harmonix/Source/HarmonixDsp/Public/HarmonixDsp/FusionSampler/FusionPatch.h
--- a/Engine/Plugins/Runtime/Harmonix/Source/HarmonixDsp/Public/HarmonixDsp/FusionSampler/FusionPatch.h
+++ b/Engine/Plugins/Runtime/Harmonix/Source/HarmonixDsp/Public/HarmonixDsp/FusionSampler/FusionPatch.h
@@ -20,17 +20,13 @@ struct HARMONIXDSP_API FFusionPatchData
 	GENERATED_BODY()
 
 public:
-	friend class FFusionPatchJsonImporter;
 	friend class UFusionPatch;
 
 	IMPL_AUDIORENDERABLE_PROXYABLE(FFusionPatchData)
 
 	FFusionPatchData() {};
-
 	const TArray<FKeyzoneSettings>& GetKeyzones() const { return Keyzones; };
 	const FFusionPatchSettings& GetSettings() const { return Settings; }
-	void UpdateSettings(const FFusionPatchSettings& InSettings) { Settings = InSettings; }
-
 	void InitProxyData(const Audio::FProxyDataInitParams& InitParams);
 
 	void DisconnectSampler(const FFusionSampler* sampler);
@@ -73,17 +69,19 @@ class HARMONIXDSP_API UFusionPatch : public UObject, public IAudioProxyDataFacto
 public:
 	// IAudioProxyDataFactory
 	virtual TSharedPtr<Audio::IProxyData> CreateProxyData(const Audio::FProxyDataInitParams& InitParams);
-
-public:
 	
-	static const int32 kMaxLayersPerNote = 128;
 	static const int32 kVoicePriorityNoSteal = 0;
 
 	UFusionPatch();
 
+	void UpdatePatch(const FFusionPatchData& InPatchData);
+
 	const FFusionPatchSettings& GetSettings() const { return FusionPatchData.GetSettings(); }
 	void UpdateSettings(const FFusionPatchSettings& InSettings);
+
 	const TArray<FKeyzoneSettings>& GetKeyzones() const { return FusionPatchData.GetKeyzones(); }
+	int32 GetNumKeyzones() const { return FusionPatchData.GetKeyzones().Num(); }
+	void UpdateKeyzones(const TArray<FKeyzoneSettings>& NewKeyzones);
 
 #if WITH_EDITORONLY_DATA
 
@@ -111,9 +109,6 @@ protected:
 	UPROPERTY(EditAnywhere, BlueprintReadOnly, Category="Fusion Patch")
 	FFusionPatchData FusionPatchData;
 private:
-
-	friend class FFusionPatchJsonImporter;
-
 	// Notice here that we cache a pointer the Proxy's "Queue" so we can...
 	// 1 - Supply it to all instances of Metasound nodes rendering this data. How?
 	//     CreateNewProxyData instantiates a NEW unique ptr to an FFusionPatchDataProxy
diff --git a/Engine/Plugins/Runtime/Harmonix/Source/HarmonixDspEditor/Private/FusionPatchJsonImporter.cpp b/Engine/Plugins/Runtime/Harmonix/Source/HarmonixDspEditor/Private/FusionPatchJsonImporter.cpp
--- a/Engine/Plugins/Runtime/Harmonix/Source/HarmonixDspEditor/Private/FusionPatchJsonImporter.cpp
+++ b/Engine/Plugins/Runtime/Harmonix/Source/HarmonixDspEditor/Private/FusionPatchJsonImporter.cpp
@@ -52,9 +52,9 @@ bool FFusionPatchJsonImporter::TryParseJson(TSharedPtr<FJsonObject> JsonObj, UFu
 		UE_LOG(LogHarmonixJsonImporter, Warning, TEXT("Bad version number importing Fusion Patch! \"version\" = %d, must be >= %d)"), Version, kMinVersion);
 	};
 	
-	FFusionPatchData& FusionPatchData = FusionPatch->FusionPatchData;
-
 	// import patch settings and keyzones first
+	// make a copy of the settings and update it with the imported settings
+	FFusionPatchSettings PatchSettingsImport = FusionPatch->GetSettings();
 	for (TSharedPtr<FJsonValue> ArrayValue : IterField(JsonObj, "presets"))
 	{
 		TSharedPtr<FJsonObject> PresetObject = nullptr;
@@ -75,15 +75,12 @@ bool FFusionPatchJsonImporter::TryParseJson(TSharedPtr<FJsonObject> JsonObj, UFu
 				UE_LOG(LogFusionPatchJsonImporter, Warning, TEXT("Unable to parse preset in json"));
 				continue;
 			}
-
-			FFusionPatchSettings PatchSettings;
-			if (!FSettingsJsonImporter::TryParseJson(PresetJson, PatchSettings))
+			
+			if (!FSettingsJsonImporter::TryParseJson(PresetJson, PatchSettingsImport))
 			{
 				UE_LOG(LogFusionPatchJsonImporter, Warning, TEXT("Unable to parse preset in json"));
 				continue;
 			}
-
-			FusionPatchData.UpdateSettings(PatchSettings);
 			// we only support one "preset" and that's the settings
 			break;
 		}
@@ -94,6 +91,8 @@ bool FFusionPatchJsonImporter::TryParseJson(TSharedPtr<FJsonObject> JsonObj, UFu
 	bool AllFilesExist = true;
 	TMap<int, int> KeyzoneIdxToPathIdx;
 
+	// make a copy of the keyzones and update it with the imported settings
+	TArray<FKeyzoneSettings> KeyzonesImported = FusionPatch->GetKeyzones();
 	int NumKeyzones = 0;
 	for (TSharedPtr<FJsonValue> ArrayValue : IterField(JsonObj, "keymap"))
 	{
@@ -106,20 +105,15 @@ bool FFusionPatchJsonImporter::TryParseJson(TSharedPtr<FJsonObject> JsonObj, UFu
 		if (!TryGetObjectField(KeyzoneObject, "keyzone", KeyzoneValue))
 			continue;
 
-		// append elements to the list until the number of keyzones 
-		// is the same as the number of imported keyzones
-		while (!FusionPatchData.Keyzones.IsValidIndex(KeyzoneIdx))
-		{
-			FKeyzoneSettings Settings;
-			FusionPatchData.Keyzones.Add(Settings);
-		}
+		// Update array to contain keyzones up to this index
+		KeyzonesImported.SetNum(KeyzoneIdx + 1);
 
-		while (!KeyzoneIdxToPathIdx.Contains(KeyzoneIdx))
+		if (!KeyzoneIdxToPathIdx.Contains(KeyzoneIdx))
 		{
 			KeyzoneIdxToPathIdx.Add(KeyzoneIdx, -1);
 		}
 
-		FKeyzoneSettings& Settings = FusionPatchData.Keyzones[KeyzoneIdx];
+		FKeyzoneSettings& Settings = KeyzonesImported[KeyzoneIdx];
 		NumKeyzones++;
 
 		if (!FSettingsJsonImporter::TryParseJson(KeyzoneValue, Settings))
@@ -149,12 +143,8 @@ bool FFusionPatchJsonImporter::TryParseJson(TSharedPtr<FJsonObject> JsonObj, UFu
 	
 	// prune the list to match the size of the newly imported asset
 	// KeyzoneIdx, at this point, will be equal to the number of _imported_ keyzones
-	while (FusionPatchData.Keyzones.Num() > NumKeyzones)
-	{
-		FusionPatchData.Keyzones.RemoveAt(FusionPatchData.Keyzones.Num() - 1);
-	}
-
-	check(KeyzoneIdxToPathIdx.Num() == FusionPatchData.Keyzones.Num());
+	KeyzonesImported.SetNum(NumKeyzones);
+	check(KeyzoneIdxToPathIdx.Num() == KeyzonesImported.Num());
 
 	if (!AllFilesExist)
 	{
@@ -183,13 +173,13 @@ bool FFusionPatchJsonImporter::TryParseJson(TSharedPtr<FJsonObject> JsonObj, UFu
 		}
 	}
 
-	for (int KeyzoneIdx = 0; KeyzoneIdx < FusionPatchData.Keyzones.Num(); ++KeyzoneIdx)
+	for (int KeyzoneIdx = 0; KeyzoneIdx < KeyzonesImported.Num(); ++KeyzoneIdx)
 	{
 		int32 PathIndex = KeyzoneIdxToPathIdx.Contains(KeyzoneIdx) ? KeyzoneIdxToPathIdx[KeyzoneIdx] : INDEX_NONE;
 
 		if (PathIndex == INDEX_NONE)
 		{
-			FString SamplePath = FusionPatchData.Keyzones[KeyzoneIdx].SamplePath;
+			FString SamplePath = KeyzonesImported[KeyzoneIdx].SamplePath;
 			FString AssetName = ImportedAssets[KeyzoneIdx]->GetName();
 			OutErrors.Add(FString::Printf(TEXT("Imported asset (Name: %s) failed to map asset (%s) to keyzone: %d"), *AssetName, *SamplePath, KeyzoneIdx));
 			UE_LOG(LogFusionPatchJsonImporter, Error, TEXT("Failed to import FusionPatch. Imported asset (Name: %s) failed to map asset (%s) to keyzone: %d"), *AssetName, *SamplePath, KeyzoneIdx);
@@ -198,7 +188,7 @@ bool FFusionPatchJsonImporter::TryParseJson(TSharedPtr<FJsonObject> JsonObj, UFu
 
 		if (!AudioSampleFiles.IsValidIndex(PathIndex))
 		{
-			FString SamplePath = FusionPatchData.Keyzones[KeyzoneIdx].SamplePath;
+			FString SamplePath = KeyzonesImported[KeyzoneIdx].SamplePath;
 			FString AssetName = AudioSampleFiles[KeyzoneIdx];
 			OutErrors.Add(FString::Printf(TEXT("Imported asset (Name: %s) failed to map file (%s) to keyzone: %d"), *AssetName, *SamplePath, KeyzoneIdx));
 			UE_LOG(LogFusionPatchJsonImporter, Error, TEXT("Failed to import FusionPatch. Imported asset (Name: %s) failed to map asset (%s) to keyzone: %d"), *AssetName, *SamplePath, KeyzoneIdx);
@@ -229,12 +219,16 @@ bool FFusionPatchJsonImporter::TryParseJson(TSharedPtr<FJsonObject> JsonObj, UFu
 
 		if (SoundWave)
 		{
-			FusionPatchData.Keyzones[KeyzoneIdx].SoundWave = SoundWave;
+			KeyzonesImported[KeyzoneIdx].SoundWave = SoundWave;
 		}
 	}
 
-	FusionPatch->UpdateRenderableForNonTrivialChange();
-
+	// no errors, update the FusionPatchData with the imported data
+	if (OutErrors.Num() == 0)
+	{
+		FusionPatch->UpdateSettings(PatchSettingsImport);
+		FusionPatch->UpdateKeyzones(KeyzonesImported);
+	}
 	return true;
 }
 
