Fusion Patch Settings detail customization.

#tests PIE

diff --git a/Engine/Plugins/Runtime/Harmonix/Source/HarmonixDspEditor/Private/Customization/FusionPatchSettingsDetailCustomization.cpp b/Engine/Plugins/Runtime/Harmonix/Source/HarmonixDspEditor/Private/Customization/FusionPatchSettingsDetailCustomization.cpp
new file mode 100644
--- /dev/null
+++ b/Engine/Plugins/Runtime/Harmonix/Source/HarmonixDspEditor/Private/Customization/FusionPatchSettingsDetailCustomization.cpp
@@ -0,0 +1,105 @@
+﻿// Copyright Epic Games, Inc. All Rights Reserved.
+
+#include "FusionPatchSettingsDetailCustomization.h"
+#include "HarmonixDsp/FusionSampler/Settings/FusionPatchSettings.h"
+#include "PropertyCustomizationHelpers.h"
+#include "DetailLayoutBuilder.h"
+#include "IDetailChildrenBuilder.h"
+#include "IPropertyUtilities.h"
+#include "Widgets/Layout/SUniformGridPanel.h"
+#include "Widgets/Text/STextBlock.h"
+
+void FFusionPatchSettingsDetailCustomization::CustomizeHeader(TSharedRef<IPropertyHandle> PropertyHandle, FDetailWidgetRow& HeaderRow, IPropertyTypeCustomizationUtils& StructCustomizationUtils)
+{
+	HeaderRow
+	.NameContent()
+	[
+		PropertyHandle->CreatePropertyNameWidget()
+	]
+	.ValueContent()
+	[
+		PropertyHandle->CreatePropertyValueWidget()
+	];
+}
+
+void FFusionPatchSettingsDetailCustomization::CustomizeChildren(TSharedRef<IPropertyHandle> PropertyHandle, IDetailChildrenBuilder& StructBuilder, IPropertyTypeCustomizationUtils& StructCustomizationUtils)
+{
+	uint32 NumChildren;
+	PropertyHandle->GetNumChildren(NumChildren);
+	for (uint32 ChildIdx = 0; ChildIdx < NumChildren; ++ChildIdx)
+	{
+		TSharedPtr<IPropertyHandle> ChildHandle = PropertyHandle->GetChildHandle(ChildIdx);
+		if (ChildHandle->GetProperty()->GetFName() == GET_MEMBER_NAME_CHECKED(FFusionPatchSettings, Adsr))
+		{
+			CustomizeArrayProperty(ChildHandle.ToSharedRef(), StructBuilder, 2);
+		}
+		else if (ChildHandle->GetProperty()->GetFName() == GET_MEMBER_NAME_CHECKED(FFusionPatchSettings, Lfo))
+		{
+			CustomizeArrayProperty(ChildHandle.ToSharedRef(), StructBuilder, 3);
+		}
+		else if (ChildHandle->GetProperty()->GetFName() == GET_MEMBER_NAME_CHECKED(FFusionPatchSettings, Randomizer))
+		{
+			CustomizeArrayProperty(ChildHandle.ToSharedRef(), StructBuilder, 3);
+		}
+		else if (ChildHandle->GetProperty()->GetFName() == GET_MEMBER_NAME_CHECKED(FFusionPatchSettings, VelocityModulator))
+		{
+			CustomizeArrayProperty(ChildHandle.ToSharedRef(), StructBuilder, 3);
+		}
+		else
+		{
+			StructBuilder.AddProperty(ChildHandle.ToSharedRef());
+		}
+	}
+}
+
+void FFusionPatchSettingsDetailCustomization::CustomizeArrayProperty(TSharedRef<IPropertyHandle> PropertyHandle, IDetailChildrenBuilder& StructBuilder, uint32 MaxPreviewNum)
+{
+	TSharedPtr<IPropertyHandleArray> ArrayHandle = PropertyHandle->AsArray();
+	uint32 NumElements;
+	ArrayHandle->GetNumElements(NumElements);
+	for (uint32 ElemIdx = 0; ElemIdx < NumElements; ++ElemIdx)
+	{
+		TSharedRef<IPropertyHandle> Element = ArrayHandle->GetElement(ElemIdx);
+		TSharedPtr<SUniformGridPanel> GridPanel;
+		StructBuilder.AddProperty(Element).CustomWidget(true)
+		.NameContent()
+		[
+			SNew(STextBlock)
+			.Font(IDetailLayoutBuilder::GetDetailFont())
+			.Text_Lambda([PropertyHandle, ElemIdx]()
+			{
+				return FText::FromString(FString::Printf(TEXT("%s: %d"),
+					*PropertyHandle->GetPropertyDisplayName().ToString(), ElemIdx));
+			})
+		]
+		.ValueContent()
+		[
+			CreateStructPropertiesPreviewValueWidget(Element, MaxPreviewNum)
+		];
+	}
+}
+
+TSharedRef<SWidget> FFusionPatchSettingsDetailCustomization::CreateStructPropertiesPreviewValueWidget(TSharedRef<IPropertyHandle> PropertyHandle, uint32 MaxPreviewNum)
+{
+	uint32 NumChildren;
+	PropertyHandle->GetNumChildren(NumChildren);
+	uint32 PreviewNum = FMath::Min(NumChildren, MaxPreviewNum);
+	TSharedPtr<SUniformGridPanel> GridPanel = SNew(SUniformGridPanel).MinDesiredSlotWidth(100);
+	for (uint32 ChildIdx = 0; ChildIdx < PreviewNum; ++ChildIdx)
+	{
+		TSharedPtr<IPropertyHandle> ChildHandle = PropertyHandle->GetChildHandle(ChildIdx);
+
+		GridPanel->AddSlot(ChildIdx, 0)
+		[
+			SNew(STextBlock)
+			.Font(IDetailLayoutBuilder::GetDetailFont())
+			.Text_Lambda([ChildHandle]()
+			{
+				FText ValueText;
+				ChildHandle->GetValueAsFormattedText(ValueText);
+				return FText::FromString(FString::Printf(TEXT("%s: %s"), *ChildHandle->GetPropertyDisplayName().ToString(), *ValueText.ToString()));
+			})
+		];
+	}
+	return GridPanel.ToSharedRef();
+}
diff --git a/Engine/Plugins/Runtime/Harmonix/Source/HarmonixDspEditor/Private/Customization/FusionPatchSettingsDetailCustomization.h b/Engine/Plugins/Runtime/Harmonix/Source/HarmonixDspEditor/Private/Customization/FusionPatchSettingsDetailCustomization.h
new file mode 100644
--- /dev/null
+++ b/Engine/Plugins/Runtime/Harmonix/Source/HarmonixDspEditor/Private/Customization/FusionPatchSettingsDetailCustomization.h
@@ -0,0 +1,27 @@
+﻿// Copyright Epic Games, Inc. All Rights Reserved.
+
+#pragma once
+
+#include "IPropertyTypeCustomization.h"
+
+class SWidget;
+
+class FFusionPatchSettingsDetailCustomization : public IPropertyTypeCustomization
+{
+public:
+	static TSharedRef<IPropertyTypeCustomization> MakeInstance()
+	{
+		return MakeShareable(new FFusionPatchSettingsDetailCustomization);
+	}
+
+	/** IPropertyTypeCustomization interface */
+	virtual void CustomizeHeader(TSharedRef<IPropertyHandle> InStructPropertyHandle, class FDetailWidgetRow& HeaderRow, IPropertyTypeCustomizationUtils& StructCustomizationUtils) override;
+
+	virtual void CustomizeChildren(TSharedRef<IPropertyHandle> InStructPropertyHandle, class IDetailChildrenBuilder& StructBuilder, IPropertyTypeCustomizationUtils& StructCustomizationUtils) override;
+
+private:
+
+	void CustomizeArrayProperty(TSharedRef<IPropertyHandle> PropertyHandle, IDetailChildrenBuilder& StructBuilder, uint32 MaxPreviewNum);
+
+	TSharedRef<SWidget> CreateStructPropertiesPreviewValueWidget(TSharedRef<IPropertyHandle> PropertyHandle, uint32 MaxPreviewNum);
+};
\ No newline at end of file
diff --git a/Engine/Plugins/Runtime/Harmonix/Source/HarmonixDspEditor/Private/HarmonixDspEditorModule.cpp b/Engine/Plugins/Runtime/Harmonix/Source/HarmonixDspEditor/Private/HarmonixDspEditorModule.cpp
--- a/Engine/Plugins/Runtime/Harmonix/Source/HarmonixDspEditor/Private/HarmonixDspEditorModule.cpp
+++ b/Engine/Plugins/Runtime/Harmonix/Source/HarmonixDspEditor/Private/HarmonixDspEditorModule.cpp
@@ -13,6 +13,7 @@
 #include "Customization/PitchShifterConfigCustomization.h"
 #include "Customization/PannerDetailsCustomization.h"
 #include "Customization/FusionPatchDetailCustomization.h"
+#include "Customization/FusionPatchSettingsDetailCustomization.h"
 #include "HarmonixDsp/FusionSampler/FusionPatch.h"
 
 #define LOCTEXT_NAMESPACE "HarmonixDspEditor"
@@ -26,8 +27,10 @@ void FHarmonixDspEditorModule::StartupModule()
 	PropertyEditorModule.RegisterCustomPropertyTypeLayout("TypedParameter", FOnGetPropertyTypeCustomizationInstance::CreateStatic(&FTypedParameterCustomization::MakeInstance));
 	PropertyEditorModule.RegisterCustomPropertyTypeLayout("PitchShifterName", FOnGetPropertyTypeCustomizationInstance::CreateStatic(&FPitchShifterNameCustomization::MakeInstance));
 	PropertyEditorModule.RegisterCustomPropertyTypeLayout("PannerDetails", FOnGetPropertyTypeCustomizationInstance::CreateStatic(&FPannerDetailsCustomization::MakeInstance));
+	PropertyEditorModule.RegisterCustomPropertyTypeLayout("FusionPatchSettings", FOnGetPropertyTypeCustomizationInstance::CreateStatic(&FFusionPatchSettingsDetailCustomization::MakeInstance));
 	PropertyEditorModule.RegisterCustomClassLayout("StretcherAndPitchShifterFactoryConfig", FOnGetDetailCustomizationInstance::CreateStatic(&FPitchShifterConfigCustomization::MakeInstance));
 	PropertyEditorModule.RegisterCustomClassLayout(UFusionPatch::StaticClass()->GetFName(), FOnGetDetailCustomizationInstance::CreateStatic(&FFusionPatchDetailCustomization::MakeInstance));
+	
 }
 
 void FHarmonixDspEditorModule::ShutdownModule()
@@ -36,6 +39,7 @@ void FHarmonixDspEditorModule::ShutdownModule()
 	FPropertyEditorModule& PropertyEditorModule = FModuleManager::LoadModuleChecked<FPropertyEditorModule>("PropertyEditor");
 	PropertyEditorModule.UnregisterCustomClassLayout("FusionPatch");
 	PropertyEditorModule.UnregisterCustomClassLayout("StretcherAndPitchShifterFactoryConfig");
+	PropertyEditorModule.UnregisterCustomPropertyTypeLayout("FusionPatchSettings");
 	PropertyEditorModule.UnregisterCustomPropertyTypeLayout("PannerDetails");
 	PropertyEditorModule.UnregisterCustomPropertyTypeLayout("PitchShifterName");
 	PropertyEditorModule.UnregisterCustomPropertyTypeLayout("TypedParameter");
