[HarmonixMidi] Protect against 0 denominator misuses

diff --git a/Engine/Plugins/Runtime/Harmonix/Source/HarmonixMidi/Private/HarmonixMidi/BarMap.cpp b/Engine/Plugins/Runtime/Harmonix/Source/HarmonixMidi/Private/HarmonixMidi/BarMap.cpp
--- a/Engine/Plugins/Runtime/Harmonix/Source/HarmonixMidi/Private/HarmonixMidi/BarMap.cpp
+++ b/Engine/Plugins/Runtime/Harmonix/Source/HarmonixMidi/Private/HarmonixMidi/BarMap.cpp
@@ -56,6 +56,11 @@ int32 FBarMap::MusicTimestampToTick(const FMusicTimestamp& Timestamp) const
 		--Index;
 	}
 
+	if (!Points.IsValidIndex(Index))
+	{
+		Index = 0;
+	}
+
 	int32 BarDelta = BarIndex - Points[Index].BarIndex;
 	return (int32)(Points[Index].StartTick +
 		BarDelta * GetTicksInBarAfterPoint(Index) +
@@ -90,6 +95,11 @@ int32 FBarMap::BarIncludingCountInToTick(int32 BarIndex, int32 * OutBeatsPerBar,
 		--TimeSigIndex;
 	}
 
+	if (!Points.IsValidIndex(TimeSigIndex))
+	{
+		TimeSigIndex = 0;
+	}
+
 	if (OutBeatsPerBar)
 	{
 		*OutBeatsPerBar = Points[TimeSigIndex].TimeSignature.Numerator;
@@ -135,6 +145,11 @@ int32 FBarMap::BarBeatTickIncludingCountInToTick(int32 BarIndex, int32 BeatInBar
 		--TimeSigIndex;
 	}
 
+	if (!Points.IsValidIndex(TimeSigIndex))
+	{
+		TimeSigIndex = 0;
+	}
+
 	int32 BarDelta = BarIndex - Points[TimeSigIndex].BarIndex;
 	int32 TicksInBeat = GetTicksInBeatAfterPoint(TimeSigIndex);
 	if (BeatInBar >= Points[TimeSigIndex].TimeSignature.Numerator)
@@ -172,7 +187,7 @@ FMusicTimestamp FBarMap::TickToMusicTimestamp(int32 Tick, int32* OutBeatsPerBar)
 	}
 
 	int32 TimeSigIndex = FMusicMapUtl::GetPointIndexForTick(Points, Tick);
-	if (TimeSigIndex < 0)
+	if (!Points.IsValidIndex(TimeSigIndex))
 	{
 		TimeSigIndex = 0;
 	}
@@ -292,7 +307,7 @@ float FBarMap::TickToFractionalBarIncludingCountIn(int32 Tick) const
 	}
 
 	int32 TimeSigIndex = FMusicMapUtl::GetPointIndexForTick(Points, Tick);
-	if (TimeSigIndex < 0)
+	if (!Points.IsValidIndex(TimeSigIndex))
 	{
 		TimeSigIndex = 0;
 	}
@@ -310,7 +325,7 @@ int32 FBarMap::TickToBarIncludingCountIn(int32 Tick) const
 		return FMath::FloorToInt32((float)Tick / (float)(TicksPerQuarterNote * 4));
 	}
 	int32 TimeSigIndex = FMusicMapUtl::GetPointIndexForTick(Points, Tick);
-	if (TimeSigIndex < 0)
+	if (!Points.IsValidIndex(TimeSigIndex))
 	{
 		TimeSigIndex = 0;
 	}
@@ -336,6 +351,11 @@ int32 FBarMap::FractionalBarIncludingCountInToTick(float Bar) const
 		--TimeSigIndex;
 	}
 
+	if (!Points.IsValidIndex(TimeSigIndex))
+	{
+		TimeSigIndex = 0;
+	}
+
 	float BarsPassed = Bar - (float)Points[TimeSigIndex].BarIndex;
 	int32 TicksPerBar = GetTicksInBarAfterPoint(TimeSigIndex);
 	return FMath::FloorToInt32((float)Points[TimeSigIndex].StartTick + (BarsPassed * (float)TicksPerBar));
@@ -399,21 +419,35 @@ int32 FBarMap::GetPointIndexForTick(int32 Tick) const
 
 const FTimeSignaturePoint& FBarMap::GetTimeSignaturePoint(int32 Index) const
 {
+	check(!Points.IsEmpty());
+
+	if (!Points.IsValidIndex(Index))
+	{
+		Index = 0;
+	}
+
 	return Points[Index];
 }
 
 FTimeSignaturePoint& FBarMap::GetTimeSignaturePoint(int32 Index)
 {
+	check(!Points.IsEmpty());
+
+	if (!Points.IsValidIndex(Index))
+	{
+		Index = 0;
+	}
+
 	return Points[Index];
 }
 
 const FTimeSignature& FBarMap::GetTimeSignatureAtTick(int32 Tick) const
 {
 	check(!Points.IsEmpty());
-		
+	
 	int32 Index = FMusicMapUtl::GetPointIndexForTick(Points, Tick);
 
-	if (Index < 0)
+	if (!Points.IsValidIndex(Index))
 	{
 		Index = 0;
 	}
@@ -434,6 +468,11 @@ const FTimeSignature& FBarMap::GetTimeSignatureAtBar(int32 Bar) const
 		--Index;
 	}
 
+	if (!Points.IsValidIndex(Index))
+	{
+		Index = 0;
+	}
+
 	return Points[Index].TimeSignature;
 }
 
diff --git a/Engine/Plugins/Runtime/Harmonix/Source/HarmonixMidi/Public/HarmonixMidi/BarMap.h b/Engine/Plugins/Runtime/Harmonix/Source/HarmonixMidi/Public/HarmonixMidi/BarMap.h
--- a/Engine/Plugins/Runtime/Harmonix/Source/HarmonixMidi/Public/HarmonixMidi/BarMap.h
+++ b/Engine/Plugins/Runtime/Harmonix/Source/HarmonixMidi/Public/HarmonixMidi/BarMap.h
@@ -305,7 +305,7 @@ public:
 
 	FORCEINLINE int32 GetTicksInBeatAfterPoint(int32 Index) const
 	{
-		if (Points.IsEmpty())
+		if (!Points.IsValidIndex(Index) || Points[Index].TimeSignature.Denominator == 0)
 		{
 			return TicksPerQuarterNote;
 		}
@@ -314,7 +314,7 @@ public:
 
 	FORCEINLINE int32 GetTicksInBarAfterPoint(int32 Index) const
 	{
-		if (Points.IsEmpty())
+		if (!Points.IsValidIndex(Index))
 		{
 			return 4 * TicksPerQuarterNote;
 		}
