Fix time quantization to midi tick in HarmonixMidi

#tests Played Main Stage on Windows and platform clients

diff --git a/Engine/Plugins/Runtime/Harmonix/Source/HarmonixMidi/Private/HarmonixMidi/BarMap.cpp b/Engine/Plugins/Runtime/Harmonix/Source/HarmonixMidi/Private/HarmonixMidi/BarMap.cpp
--- a/Engine/Plugins/Runtime/Harmonix/Source/HarmonixMidi/Private/HarmonixMidi/BarMap.cpp
+++ b/Engine/Plugins/Runtime/Harmonix/Source/HarmonixMidi/Private/HarmonixMidi/BarMap.cpp
@@ -24,7 +24,7 @@ bool FBarMap::IsEmpty() const
 	return Points.IsEmpty();
 }
 
-int32 FBarMap::MusicTimestampToTick(const FMusicTimestamp& Timestamp) const
+float FBarMap::MusicTimestampToTick(const FMusicTimestamp& Timestamp) const
 {
 	float Beat = Timestamp.Beat;
 	if (Beat < 1.0f)
@@ -45,8 +45,8 @@ int32 FBarMap::MusicTimestampToTick(const FMusicTimestamp& Timestamp) const
 	Beat -= 1.0f;
 	if (Points.IsEmpty())
 	{
-		return	BarIndex * 4 * TicksPerQuarterNote +
-			FMath::FloorToInt(Beat * TicksPerQuarterNote);
+		return BarIndex * 4 * TicksPerQuarterNote +
+			Beat * TicksPerQuarterNote;
 	}
 	
 	int32 Index = Algo::UpperBound(Points, BarIndex, FTimeSignaturePoint::BarLessThan());
@@ -62,9 +62,9 @@ int32 FBarMap::MusicTimestampToTick(const FMusicTimestamp& Timestamp) const
 	}
 
 	int32 BarDelta = BarIndex - Points[Index].BarIndex;
-	return (int32)(Points[Index].StartTick +
+	return Points[Index].StartTick +
 		BarDelta * GetTicksInBarAfterPoint(Index) +
-		(int32)(Beat * GetTicksInBeatAfterPoint(Index)));
+		Beat * GetTicksInBeatAfterPoint(Index);
 }
 
 int32 FBarMap::MusicTimestampBarToTick(int32 InBar, int32* OutBeatsPerBar, int32* OutTicksPerBeat) const
@@ -73,7 +73,7 @@ int32 FBarMap::MusicTimestampBarToTick(int32 InBar, int32* OutBeatsPerBar, int32
 	return BarIncludingCountInToTick(BarIndex, OutBeatsPerBar, OutTicksPerBeat);
 }
 
-int32 FBarMap::BarIncludingCountInToTick(int32 BarIndex, int32 * OutBeatsPerBar, int32 * OutTicksPerBeat) const
+int32 FBarMap::BarIncludingCountInToTick(int32 BarIndex, int32* OutBeatsPerBar, int32* OutTicksPerBeat) const
 {
 	if (Points.IsEmpty())
 	{
@@ -85,7 +85,7 @@ int32 FBarMap::BarIncludingCountInToTick(int32 BarIndex, int32 * OutBeatsPerBar,
 		{
 			*OutTicksPerBeat = TicksPerQuarterNote;
 		}
-		return	BarIndex * 4 * TicksPerQuarterNote;
+		return BarIndex * 4 * TicksPerQuarterNote;
 	}
 
 	int32 TimeSigIndex = Algo::UpperBound(Points, BarIndex, FTimeSignaturePoint::BarLessThan());
@@ -166,55 +166,55 @@ int32 FBarMap::BarBeatTickIncludingCountInToTick(int32 BarIndex, int32 BeatInBar
 			TickInBeat;
 }
 
-FMusicTimestamp FBarMap::TickToMusicTimestamp(int32 Tick, int32* OutBeatsPerBar) const
+FMusicTimestamp FBarMap::TickToMusicTimestamp(float Tick, int32* OutBeatsPerBar) const
 {
 	FMusicTimestamp Result;
 	if (Points.IsEmpty())
 	{
 		// Assume 4/4 time.
-		Result.Bar = Tick / (TicksPerQuarterNote * 4);
+		Result.Bar = int32(Tick) / (TicksPerQuarterNote * 4);
 		Tick -= (Result.Bar * TicksPerQuarterNote * 4);
 		if (Tick < 0)
 		{
 			Result.Bar--;
 			Tick += TicksPerQuarterNote * 4;
 		}
-		Result.Beat = (float)Tick / (float)TicksPerQuarterNote;
+		Result.Beat = Tick / TicksPerQuarterNote;
 		Result.Bar += StartBar;
 		Result.Beat += 1.0f; // 1 based
 		if (OutBeatsPerBar) *OutBeatsPerBar = 4;
 		return Result;
 	}
 
-	int32 TimeSigIndex = FMusicMapUtl::GetPointIndexForTick(Points, Tick);
+	int32 TimeSigIndex = FMusicMapUtl::GetPointIndexForTick(Points, int32(Tick));
 	if (!Points.IsValidIndex(TimeSigIndex))
 	{
 		TimeSigIndex = 0;
 	}
 
 	int32 TicksPerBar = GetTicksInBarAfterPoint(TimeSigIndex);
-	int32 TicksPassed = Tick - Points[TimeSigIndex].StartTick;
+	float TicksPassed = Tick - Points[TimeSigIndex].StartTick;
 	if (TicksPassed < 0)
 	{
-		int32 BarsPassed = (- TicksPassed / TicksPerBar) + 1;
+		int32 BarsPassed = -int32(TicksPassed) / TicksPerBar + 1;
 		TicksPassed += BarsPassed * TicksPerBar;
-		float BeatsPassed = (float)TicksPassed / (float)GetTicksInBeatAfterPoint(TimeSigIndex);
+		float BeatsPassed = TicksPassed / GetTicksInBeatAfterPoint(TimeSigIndex);
 		Result.Bar = -BarsPassed + StartBar; // already - 1 based
 		Result.Beat = BeatsPassed + 1.0f; // 1 based
 	}
 	else
 	{
-		int32 BarsPassed = TicksPassed / TicksPerBar;
+		int32 BarsPassed = int32(TicksPassed) / TicksPerBar;
 		TicksPassed -= BarsPassed * TicksPerBar;
-		float BeatsPassed = (float)TicksPassed / (float)GetTicksInBeatAfterPoint(TimeSigIndex);
+		float BeatsPassed = TicksPassed / GetTicksInBeatAfterPoint(TimeSigIndex);
 		Result.Bar = Points[TimeSigIndex].BarIndex + BarsPassed + StartBar; // 1 based
-		Result.Beat = BeatsPassed + 1.0; // 1 based
+		Result.Beat = BeatsPassed + 1.0f; // 1 based
 	}
 	if (OutBeatsPerBar) *OutBeatsPerBar = Points[TimeSigIndex].TimeSignature.Numerator;
 	return Result;
 }
 
-FMusicTimestamp FBarMap::TickFromBarOneToMusicTimestamp(int32 InTickFromBarOne, int32* OutBeatsPerBar) const
+FMusicTimestamp FBarMap::TickFromBarOneToMusicTimestamp(float InTickFromBarOne, int32* OutBeatsPerBar) const
 {
 	return TickToMusicTimestamp(InTickFromBarOne + GetTickOfBarOne(), OutBeatsPerBar);
 }
@@ -298,23 +298,23 @@ float FBarMap::TickToFractionalMusicalTimestampBar(int32 Tick) const
 }
 */
 
-float FBarMap::TickToFractionalBarIncludingCountIn(int32 Tick) const
+float FBarMap::TickToFractionalBarIncludingCountIn(float Tick) const
 {
 	if (Points.IsEmpty())
 	{
 		// assume 4/4
-		return (float)Tick / (float)(TicksPerQuarterNote * 4);
+		return Tick / (TicksPerQuarterNote * 4);
 	}
 
-	int32 TimeSigIndex = FMusicMapUtl::GetPointIndexForTick(Points, Tick);
+	int32 TimeSigIndex = FMusicMapUtl::GetPointIndexForTick(Points, int32(Tick));
 	if (!Points.IsValidIndex(TimeSigIndex))
 	{
 		TimeSigIndex = 0;
 	}
-	int32 TicksPast = Tick - Points[TimeSigIndex].StartTick;
+	float TicksPast = Tick - Points[TimeSigIndex].StartTick;
 	int32 TicksPerBar = GetTicksInBarAfterPoint(TimeSigIndex);
-	float BarsPassed = (float)TicksPast / (float)TicksPerBar;
-	return (float)Points[TimeSigIndex].BarIndex + BarsPassed;
+	float BarsPassed = TicksPast / TicksPerBar;
+	return Points[TimeSigIndex].BarIndex + BarsPassed;
 }
 
 int32 FBarMap::TickToBarIncludingCountIn(int32 Tick) const
@@ -322,7 +322,7 @@ int32 FBarMap::TickToBarIncludingCountIn(int32 Tick) const
 	if (Points.IsEmpty())
 	{
 		// assume 4/4
-		return FMath::FloorToInt32((float)Tick / (float)(TicksPerQuarterNote * 4));
+		return Tick / (TicksPerQuarterNote * 4);
 	}
 	int32 TimeSigIndex = FMusicMapUtl::GetPointIndexForTick(Points, Tick);
 	if (!Points.IsValidIndex(TimeSigIndex))
@@ -335,12 +335,12 @@ int32 FBarMap::TickToBarIncludingCountIn(int32 Tick) const
 	return Points[TimeSigIndex].BarIndex + BarsPassed;
 }
 
-int32 FBarMap::FractionalBarIncludingCountInToTick(float Bar) const
+float FBarMap::FractionalBarIncludingCountInToTick(float Bar) const
 {
 	if (Points.IsEmpty())
 	{
 		// assume 4/4
-		return FMath::FloorToInt32(Bar * 4.0f * TicksPerQuarterNote);
+		return Bar * 4.0f * TicksPerQuarterNote;
 	}
 
 	int32 BarIndex = FMath::FloorToInt32(Bar);
@@ -356,9 +356,9 @@ int32 FBarMap::FractionalBarIncludingCountInToTick(float Bar) const
 		TimeSigIndex = 0;
 	}
 
-	float BarsPassed = Bar - (float)Points[TimeSigIndex].BarIndex;
+	float BarsPassed = Bar - Points[TimeSigIndex].BarIndex;
 	int32 TicksPerBar = GetTicksInBarAfterPoint(TimeSigIndex);
-	return FMath::FloorToInt32((float)Points[TimeSigIndex].StartTick + (BarsPassed * (float)TicksPerBar));
+	return Points[TimeSigIndex].StartTick + BarsPassed * TicksPerBar;
 }
 
 bool FBarMap::AddTimeSignatureAtMusicTimestampBar(int32 BarNumber, int32 Numerator, int32 Denominator, bool SortNow, bool FailOnError)
diff --git a/Engine/Plugins/Runtime/Harmonix/Source/HarmonixMidi/Private/HarmonixMidi/BeatMap.cpp b/Engine/Plugins/Runtime/Harmonix/Source/HarmonixMidi/Private/HarmonixMidi/BeatMap.cpp
--- a/Engine/Plugins/Runtime/Harmonix/Source/HarmonixMidi/Private/HarmonixMidi/BeatMap.cpp
+++ b/Engine/Plugins/Runtime/Harmonix/Source/HarmonixMidi/Private/HarmonixMidi/BeatMap.cpp
@@ -54,7 +54,7 @@ float FBeatMap::GetFractionalBeatAtTick(float Tick) const
 		return WholeBeats + FractionalBeats + 1.0f; // +1 because position is always 1 based.
 	}
 
-	return FMusicMapUtl::GetFrationalPointForTick(Points, Tick);
+	return FMusicMapUtl::GetFractionalPointForTick(Points, Tick);
 }
 
 float FBeatMap::GetFractionalTickAtBeat(float Beat) const
@@ -67,12 +67,12 @@ float FBeatMap::GetFractionalTickAtBeat(float Beat) const
 	return FMusicMapUtl::GetFractionalTickForFractionalPoint(Points, Beat);
 }
 
-EMusicalBeatType FBeatMap::GetBeatTypeAtTick(float Tick) const
+EMusicalBeatType FBeatMap::GetBeatTypeAtTick(int32 Tick) const
 {
 	if (Points.IsEmpty())
 	{
 		// assume 4/4
-		int32 Beat = FMath::FloorToInt(Tick / (float)TicksPerQuarterNote);
+		int32 Beat = Tick / TicksPerQuarterNote;
 		if (Beat % 4 == 0)
 		{
 			return EMusicalBeatType::Downbeat;
@@ -80,7 +80,7 @@ EMusicalBeatType FBeatMap::GetBeatTypeAtTick(float Tick) const
 		return EMusicalBeatType::Normal;
 	}
 
-	const FBeatMapPoint* Point = FMusicMapUtl::GetPointInfoForTick(Points, FMath::FloorToInt(Tick));
+	const FBeatMapPoint* Point = FMusicMapUtl::GetPointInfoForTick(Points, Tick);
 	if (!Point)
 	{
 		return EMusicalBeatType::Normal;
@@ -95,26 +95,26 @@ float FBeatMap::GetBeatInPulseBarAtTick(float Tick) const
 	{
 		return 1.0f; // 1 because beats relative to a bar are always 1 based.
 	}
-	float ProgressInBeat = Points[PointIndex].Progress(FMath::FloorToInt(Tick));
+	float ProgressInBeat = Points[PointIndex].Progress(Tick);
 	if (Bars.IsEmpty())
 	{
-		return (float)PointIndex + ProgressInBeat + 1.0f; // +1 because beats relative to a bar are always 1 based.
+		return PointIndex + ProgressInBeat + 1.0f; // +1 because beats relative to a bar are always 1 based.
 	}
 
 	int32 BarIndex = Points[PointIndex].PulseBar;
 	const FPulseBar& PulseBar = Bars[BarIndex];
-	float BeatInBar = PointIndex - PulseBar.FirstIncludedBeatIndex;
+	int32 BeatInBar = PointIndex - PulseBar.FirstIncludedBeatIndex;
 	return BeatInBar + ProgressInBeat + 1.0f; // +1 because beats relative to a bar are always 1 based.
 }
 
-float FBeatMap::GetNumBeatsInPulseBarAt(float Tick) const
+int32 FBeatMap::GetNumBeatsInPulseBarAt(int32 Tick) const
 {
 	if (Bars.IsEmpty())
 	{
 		return Points.Num();
 	}
 
-	int32 BeatIndex = FMusicMapUtl::GetPointIndexForTick(Points, FMath::Floor(Tick));
+	int32 BeatIndex = FMusicMapUtl::GetPointIndexForTick(Points, Tick);
 	if (BeatIndex == -1)
 	{
 		return 0;
@@ -150,7 +150,7 @@ int32 FBeatMap::FindDownbeatIndexAtOrAfterBeat(float Beat) const
 	int32 BeatIndex = FMath::FloorToInt32(Beat);
 	if (Points.IsEmpty())
 	{
-		return (((Beat + 3) / 4) * 4);
+		return (((BeatIndex + 3) / 4) * 4);
 	}
 
 	for (int32 TestIndex = BeatIndex; TestIndex < Points.Num(); ++TestIndex)
diff --git a/Engine/Plugins/Runtime/Harmonix/Source/HarmonixMidi/Private/HarmonixMidi/MidiSongPos.cpp b/Engine/Plugins/Runtime/Harmonix/Source/HarmonixMidi/Private/HarmonixMidi/MidiSongPos.cpp
--- a/Engine/Plugins/Runtime/Harmonix/Source/HarmonixMidi/Private/HarmonixMidi/MidiSongPos.cpp
+++ b/Engine/Plugins/Runtime/Harmonix/Source/HarmonixMidi/Private/HarmonixMidi/MidiSongPos.cpp
@@ -86,45 +86,50 @@ void FMidiSongPos::SetByTime(float InElapsedMs, float InBpm, int32 InTimeSigNume
 void FMidiSongPos::SetByTime(float InMs, const FSongMaps& Maps)
 {
 	float Tick = Maps.GetTempoMap().MsToTick(InMs);
-	SetByTick(Tick, Maps);
+	SetByTimeAndTick(InMs, Tick, Maps);
 }
 
 void FMidiSongPos::SetByTick(float InTick, const FSongMaps& Maps)
 {
-	const FTimeSignature* TimeSig = Maps.GetTimeSignatureAtTick(InTick);
-	float TickMs			= Maps.TickToMs(InTick);
-	SecondsIncludingCountIn = TickMs / 1000.0f;
+	float Ms = Maps.GetTempoMap().TickToMs(InTick);
+	SetByTimeAndTick(Ms, InTick, Maps);
+}
+
+void FMidiSongPos::SetByTimeAndTick(float InMs, float InTick, const FSongMaps& Maps)
+{
+	const FTimeSignature* TimeSig = Maps.GetTimeSignatureAtTick(int32(InTick));
+	SecondsIncludingCountIn = InMs / 1000.0f;
 	SecondsFromBarOne       = SecondsIncludingCountIn - Maps.GetCountInSeconds();
 	TimeSigNumerator        = TimeSig ? TimeSig->Numerator : 4;
 	TimeSigDenominator      = TimeSig ? TimeSig->Denominator : 4;
-	Tempo                   = Maps.GetTempoAtTick(InTick);
+	Tempo                   = Maps.GetTempoAtTick(int32(InTick));
 	BarsIncludingCountIn    = Maps.GetBarIncludingCountInAtTick(InTick);
 
-	const FSongSection* SongSection = !Maps.GetSectionMap().IsEmpty() ? Maps.GetSectionAtMs(TickMs) : nullptr;
+	const FSongSection* SongSection = !Maps.GetSectionMap().IsEmpty() ? Maps.GetSectionAtTick(int32(InTick)) : nullptr;
 	CurrentSongSection		= SongSection ? FSongSection(SongSection->Name, SongSection->StartTick, SongSection->LengthTicks) : FSongSection();
 	
 	const FBeatMap& BeatMap = Maps.GetBeatMap();
-	int32 BeatPointIndex     = BeatMap.GetPointIndexForTick(FMath::FloorToInt32(InTick));
+	int32 BeatPointIndex     = BeatMap.GetPointIndexForTick(int32(InTick));
 	if (BeatPointIndex >= 0)
 	{
 		const FBeatMapPoint& BeatPoint = BeatMap.GetBeatPointInfo(BeatPointIndex);
-		float TickInBeat     = InTick - (float)BeatPoint.StartTick;
-		float FractionalPart = TickInBeat / (float)BeatPoint.LengthTicks;
-		BeatsIncludingCountIn     = (float)BeatPointIndex + FractionalPart;
+		float TickInBeat     = InTick - BeatPoint.StartTick;
+		float FractionalPart = TickInBeat / BeatPoint.LengthTicks;
+		BeatsIncludingCountIn     = BeatPointIndex + FractionalPart;
 		BeatType             = BeatPoint.Type;
 	}
 	else
 	{
 		// use basic time signature based info from bar map...
 		const FBarMap& BarMap = Maps.GetBarMap();
-		int32 BarPointIndex = BarMap.GetPointIndexForTick(FMath::FloorToInt32(InTick));
+		int32 BarPointIndex = BarMap.GetPointIndexForTick(int32(InTick));
 		if (BarPointIndex == INDEX_NONE)
 		{
 			ensure(InTick < 0.0f);
 			BarPointIndex = 0;
 		}
 		const FTimeSignaturePoint& TimeSigPoint = BarMap.GetTimeSignaturePoint(BarPointIndex);
-		float BarAtTimeSig = BarsIncludingCountIn - (float)TimeSigPoint.BarIndex;
+		float BarAtTimeSig = BarsIncludingCountIn - TimeSigPoint.BarIndex;
 		float FractPart = FMath::Fractional(BarAtTimeSig);
 		BeatsIncludingCountIn = TimeSigPoint.BeatIndex + (BarAtTimeSig * TimeSigPoint.TimeSignature.Numerator);
 		BeatType = FMath::IsNearlyEqual(Timestamp.Beat, 1.0f) ? EMusicalBeatType::Downbeat : EMusicalBeatType::Normal;
diff --git a/Engine/Plugins/Runtime/Harmonix/Source/HarmonixMidi/Private/HarmonixMidi/SmoothedMidiPlayCursor.cpp b/Engine/Plugins/Runtime/Harmonix/Source/HarmonixMidi/Private/HarmonixMidi/SmoothedMidiPlayCursor.cpp
--- a/Engine/Plugins/Runtime/Harmonix/Source/HarmonixMidi/Private/HarmonixMidi/SmoothedMidiPlayCursor.cpp
+++ b/Engine/Plugins/Runtime/Harmonix/Source/HarmonixMidi/Private/HarmonixMidi/SmoothedMidiPlayCursor.cpp
@@ -16,9 +16,9 @@ namespace SmoothedPlayCursor
 FSmoothedMidiPlayCursor::FSmoothedMidiPlayCursor()
 {
 	SetupMsLookahead(-SmoothingLatencyMs, FMidiPlayCursor::ESyncOptions::NoBroadcastNoPreRoll);
-	SmoothingTimer.SetSpeed(1.0f);
+	SmoothingTimer.SetSpeed(1.0);
 	SmoothingTimer.Start();
-	SmoothingTimer.Reset(0.0f);
+	SmoothingTimer.Reset(0.0);
 }
 
 void FSmoothedMidiPlayCursor::SetSpeed(float NewSpeed)
@@ -91,11 +91,11 @@ void FSmoothedMidiPlayCursor::SyncSmoothingTimer(bool bEnableErrorCorrection)
 	SetSpeed(Owner->GetCurrentAdvanceRate());
 
 	float RawMs = CurrentMs;
-	float SmoothMs = SmoothingTimer.Ms();
+	double SmoothMs = SmoothingTimer.Ms();
 
 	if (bEnableErrorCorrection)
 	{
-		float Error = (SmoothMs - RawMs) * (1.0f / (float)SmoothingTimer.GetSpeed());
+		float Error = float((SmoothMs - RawMs) * (1.0 / SmoothingTimer.GetSpeed()));
 		ErrorTracker.Push(Error);
 		float MinRecentError = ErrorTracker.Min();
 		float AbsError = FMath::Abs(MinRecentError);
@@ -108,7 +108,7 @@ void FSmoothedMidiPlayCursor::SyncSmoothingTimer(bool bEnableErrorCorrection)
 			SmoothingTimer.Reset(RawMs);
 			SmoothedMs = RawMs;
 			SmoothedMsDelta = 0.0f;
-			SmoothedTick = Owner->GetTempoMap().MsToTick(SmoothedMs);
+			SmoothedTick = Owner->GetTempoMap().MsToTick(RawMs);
 			ErrorTracker.Reset();
 			return;
 		}
@@ -143,7 +143,7 @@ void FSmoothedMidiPlayCursor::SyncSmoothingTimer(bool bEnableErrorCorrection)
 
 	// keep track of the difference in Ms between the current time and the smoothed time
 	// update the MsDelta before applying any looping so that it's 
-	SmoothedMsDelta = SmoothMs - RawMs;
+	SmoothedMsDelta = float(SmoothMs - RawMs);
 
 	// we need to see if the smoothed position is a position that we have never played!
 	float LoopStartMs = Owner->GetLoopStartMs();
@@ -166,7 +166,7 @@ void FSmoothedMidiPlayCursor::SyncSmoothingTimer(bool bEnableErrorCorrection)
 			SmoothMs = LoopStartMs + (SmoothMs - LoopEndMs);
 		}
 	}
-	SmoothedMs = SmoothMs;
+	SmoothedMs = float(SmoothMs);
 	SmoothedTick = Owner->GetTempoMap().MsToTick(SmoothedMs);
 	LoopedThisPass = false;
 }
diff --git a/Engine/Plugins/Runtime/Harmonix/Source/HarmonixMidi/Private/HarmonixMidi/SongMaps.cpp b/Engine/Plugins/Runtime/Harmonix/Source/HarmonixMidi/Private/HarmonixMidi/SongMaps.cpp
--- a/Engine/Plugins/Runtime/Harmonix/Source/HarmonixMidi/Private/HarmonixMidi/SongMaps.cpp
+++ b/Engine/Plugins/Runtime/Harmonix/Source/HarmonixMidi/Private/HarmonixMidi/SongMaps.cpp
@@ -42,11 +42,11 @@ void FSongMaps::Copy(const FSongMaps& Other, int32 StartTick, int32 EndTick)
 	LengthData.LastTick = LastTick;
 	LengthData.LengthTicks = LengthData.LastTick + 1;
 	FMusicTimestamp Timestamp;
-	Timestamp = BarMap.TickToMusicTimestamp(LengthData.LastTick);
+	Timestamp = BarMap.TickToMusicTimestamp(float(LengthData.LastTick));
 	LengthData.LengthBars = FMath::IsNearlyZero(Timestamp.Beat-1.0f,1.0E-4) ? Timestamp.Bar - 1 : Timestamp.Bar;
 	if (LengthData.LengthBars == Timestamp.Bar)
 	{
-		LengthData.LengthTicks = BarMap.MusicTimestampToTick({LengthData.LengthBars + BarMap.GetStartBar(), 1.0f}); 
+		LengthData.LengthTicks = int32(BarMap.MusicTimestampToTick({LengthData.LengthBars + BarMap.GetStartBar(), 1.0f}));
 		LengthData.LastTick = LengthData.LengthTicks - 1;
 	}
 }
@@ -72,7 +72,7 @@ bool FSongMaps::LoadFromStdMidiFile(TSharedPtr<FArchive> Archive, const FString&
 	return ReadWithReader(Reader);
 }
 
-float FSongMaps::TickToMs(int32 Tick) const
+float FSongMaps::TickToMs(float Tick) const
 {
 	if (TempoMap.IsEmpty())
 	{
@@ -82,20 +82,20 @@ float FSongMaps::TickToMs(int32 Tick) const
 	return TempoMap.TickToMs(Tick);
 }
 
-int32 FSongMaps::MsToTick(float Ms) const
+float FSongMaps::MsToTick(float Ms) const
 {
 	if (TempoMap.IsEmpty())
 	{
 		UE_LOG(LogMidi, Log, TEXT("SongMaps does not contain a Tempo Map."));
 		return 0;
 	}
-	return FMath::FloorToInt(TempoMap.MsToTick(Ms));
+	return TempoMap.MsToTick(Ms);
 }
 
 float FSongMaps::GetCountInSeconds() const
 {
-	int32 BarOneBeatOneTick = BarMap.MusicTimestampToTick({1,1.0f});
-	return TempoMap.TickToSeconds(BarOneBeatOneTick);
+	float BarOneBeatOneTick = BarMap.MusicTimestampToTick({ 1, 1.0f });
+	return TempoMap.TickToMs(BarOneBeatOneTick) / 1000.f;
 }
 
 bool FSongMaps::FinalizeRead(IMidiReader* Reader)
@@ -118,11 +118,11 @@ bool FSongMaps::FinalizeRead(IMidiReader* Reader)
 	LengthData.LastTick = LastTick;
 	LengthData.LengthTicks = LengthData.LastTick + 1;
 	FMusicTimestamp Timestamp;
-	Timestamp = BarMap.TickToMusicTimestamp(LengthData.LastTick);
+	Timestamp = BarMap.TickToMusicTimestamp(float(LengthData.LastTick));
 	LengthData.LengthBars = FMath::IsNearlyZero(Timestamp.Beat - 1.0f, 1.0E-4) ? Timestamp.Bar - 1 : Timestamp.Bar;
 	if (LengthData.LengthBars == Timestamp.Bar)
 	{
-		LengthData.LengthTicks = BarMap.MusicTimestampToTick({ LengthData.LengthBars + BarMap.GetStartBar(), 1.0f });
+		LengthData.LengthTicks = int32(BarMap.MusicTimestampToTick({ LengthData.LengthBars + BarMap.GetStartBar(), 1.0f }));
 		LengthData.LastTick = LengthData.LengthTicks - 1;
 	}
 	return true;
@@ -154,7 +154,7 @@ void FSongMaps::EmptyAllMaps()
 
 float FSongMaps::GetSongLengthMs() const
 {
-	return TickToMs(GetSongLengthData().LengthTicks);
+	return TickToMs(float(GetSongLengthData().LengthTicks));
 }
 
 int32 FSongMaps::GetSongLengthBeats() const
@@ -249,7 +249,7 @@ const FTempoInfoPoint* FSongMaps::GetTempoInfoForMs(float Ms) const
 		UE_LOG(LogMidi, Log, TEXT("SongMaps does not contain a Tempo Map."));
 		return nullptr;
 	}
-	int32 Tick = MsToTick(Ms);
+	int32 Tick = int32(MsToTick(Ms));
 	return TempoMap.GetTempoPointAtTick(Tick);
 }
 
@@ -270,7 +270,7 @@ float FSongMaps::GetTempoAtMs(float Ms) const
 		UE_LOG(LogMidi, Log, TEXT("SongMaps does not contain a Tempo Map."));
 		return 0.0f;
 	}
-	int32 Tick = MsToTick(Ms);
+	int32 Tick = int32(MsToTick(Ms));
 	return TempoMap.GetTempoAtTick(Tick);
 }
 
@@ -293,7 +293,7 @@ const FBeatMapPoint* FSongMaps::GetBeatAtMs(float Ms) const
 		UE_LOG(LogMidi, Log, TEXT("SongMaps does not contain a Beat Map."));
 		return nullptr;
 	}
-	int32 Tick = MsToTick(Ms);
+	int32 Tick = int32(MsToTick(Ms));
 	return BeatMap.GetPointInfoForTick(Tick);
 }
 
@@ -345,7 +345,7 @@ const FBeatMapPoint* FSongMaps::GetBeatAtTick(int32 Tick) const
 
 float FSongMaps::GetMsPerBeatAtMs(float Ms) const
 {
-	int32 Tick = MsToTick(Ms);
+	int32 Tick = int32(MsToTick(Ms));
 	return GetMsPerBeatAtTick(Tick);
 }
 
@@ -362,18 +362,18 @@ float FSongMaps::GetMsPerBeatAtTick(int32 Tick) const
 		float BeatsPerMinute = QuarterNotesPerMinute / ((float)TimeSignature->Denominator / 4.0f /* quarter note / 4 == 1 */);
 		return 60000.0f / BeatsPerMinute;
 	}
-	return  TickToMs(BeatInfo->StartTick + BeatInfo->LengthTicks) - TickToMs(BeatInfo->StartTick);
+	return TickToMs(float(BeatInfo->StartTick + BeatInfo->LengthTicks)) - TickToMs(float(BeatInfo->StartTick));
 }
 
 float FSongMaps::GetFractionalBeatAtMs(float Ms) const
 {
-	int32 Tick = MsToTick(Ms);
+	float Tick = MsToTick(Ms);
 	return GetFractionalBeatAtTick(Tick);
 }
 
-float FSongMaps::GetFractionalBeatAtTick(int32 Tick) const
+float FSongMaps::GetFractionalBeatAtTick(float Tick) const
 {
-	int32 BeatIndex = GetBeatIndexAtTick(Tick);
+	int32 BeatIndex = GetBeatIndexAtTick(int32(Tick));
 	if (BeatIndex < 0 || BeatMap.IsEmpty())
 	{
 		return 1.0f; // 1 based position!
@@ -381,14 +381,14 @@ float FSongMaps::GetFractionalBeatAtTick(int32 Tick) const
 
 	const FBeatMapPoint* BeatInfo = &BeatMap.GetBeatPointInfo(BeatIndex);
 	check(BeatInfo);
-	int32 TickInBeat = Tick - BeatInfo->StartTick;
-	float FractionalPart = (float)TickInBeat / (float)BeatInfo->LengthTicks;
-	return (float)BeatIndex + FractionalPart + 1.0f; // +1 for musical position
+	float TickInBeat = Tick - BeatInfo->StartTick;
+	float FractionalPart = TickInBeat / BeatInfo->LengthTicks;
+	return BeatIndex + FractionalPart + 1.0f; // +1 for musical position
 }
 
 int32 FSongMaps::GetBeatIndexAtMs(float Ms) const
 {
-	int32 Tick = MsToTick(Ms);
+	int32 Tick = int32(MsToTick(Ms));
 	return GetBeatIndexAtTick(Tick);
 }
 
@@ -404,7 +404,7 @@ int32 FSongMaps::GetBeatIndexAtTick(int32 Tick) const
 
 EMusicalBeatType FSongMaps::GetBeatTypeAtMs(float Ms) const
 {
-	int32 Tick = MsToTick(Ms);
+	int32 Tick = int32(MsToTick(Ms));
 	return GetBeatTypeAtTick(Tick);
 }
 
@@ -420,11 +420,11 @@ EMusicalBeatType FSongMaps::GetBeatTypeAtTick(int32 Tick) const
 
 float FSongMaps::GetBeatInPulseBarAtMs(float Ms) const
 {
-	int32 Tick = MsToTick(Ms);
+	float Tick = MsToTick(Ms);
 	return GetBeatInPulseBarAtTick(Tick);
 }
 
-float FSongMaps::GetBeatInPulseBarAtTick(int32 Tick) const
+float FSongMaps::GetBeatInPulseBarAtTick(float Tick) const
 {
 	if (BeatMap.IsEmpty())
 	{
@@ -434,18 +434,18 @@ float FSongMaps::GetBeatInPulseBarAtTick(int32 Tick) const
 	return BeatMap.GetBeatInPulseBarAtTick(Tick);
 }
 
-float FSongMaps::GetNumBeatsInPulseBarAtMs(float Ms) const
+int32 FSongMaps::GetNumBeatsInPulseBarAtMs(float Ms) const
 {
-	int32 Tick = MsToTick(Ms);
+	int32 Tick = int32(MsToTick(Ms));
 	return GetNumBeatsInPulseBarAtTick(Tick);
 }
 
-float FSongMaps::GetNumBeatsInPulseBarAtTick(int32 Tick) const
+int32 FSongMaps::GetNumBeatsInPulseBarAtTick(int32 Tick) const
 {
 	if (BeatMap.IsEmpty())
 	{
 		UE_LOG(LogMidi, Log, TEXT("SongMaps does not contain a Beat Map."));
-		return 0.0f;
+		return 0;
 	}
 	return BeatMap.GetNumBeatsInPulseBarAt(Tick);
 }
@@ -455,7 +455,7 @@ float FSongMaps::GetNumBeatsInPulseBarAtTick(int32 Tick) const
 
 const FTimeSignature* FSongMaps::GetTimeSignatureAtMs(float Ms) const
 {
-	int32 Tick = MsToTick(Ms);
+	int32 Tick = int32(MsToTick(Ms));
 	return GetTimeSignatureAtTick(Tick);
 }
 
@@ -486,11 +486,11 @@ const FTimeSignature* FSongMaps::GetTimeSignatureAtBar(int32 Bar) const
 
 float FSongMaps::GetBarIncludingCountInAtMs(float Ms) const
 {
-	int32 Tick = MsToTick(Ms);
+	float Tick = MsToTick(Ms);
 	return GetBarIncludingCountInAtTick(Tick);
 }
 
-float FSongMaps::GetBarIncludingCountInAtTick(int32 Tick) const
+float FSongMaps::GetBarIncludingCountInAtTick(float Tick) const
 {
 	if (BarMap.IsEmpty())
 	{
@@ -502,11 +502,11 @@ float FSongMaps::GetBarIncludingCountInAtTick(int32 Tick) const
 
 float FSongMaps::GetMsPerBarAtMs(float Ms) const
 {
-	int32 Tick = MsToTick(Ms);
+	float Tick = MsToTick(Ms);
 	return GetMsPerBarAtTick(Tick);
 }
 
-float FSongMaps::GetMsPerBarAtTick(int32 Tick) const
+float FSongMaps::GetMsPerBarAtTick(float Tick) const
 {
 	if (TempoMap.IsEmpty())
 	{
@@ -518,8 +518,8 @@ float FSongMaps::GetMsPerBarAtTick(int32 Tick) const
 		UE_LOG(LogMidi, Log, TEXT("SongMaps does not contain a Bar Map."));
 		return 0.0f;
 	}
-	float Bpm = TempoMap.GetTempoAtTick(Tick); // quarter noted per minute
-	const FTimeSignature* TimeSignature = GetTimeSignatureAtTick(Tick);
+	float Bpm = TempoMap.GetTempoAtTick(int32(Tick)); // quarter notes per minute
+	const FTimeSignature* TimeSignature = GetTimeSignatureAtTick(int32(Tick));
 	float QuarterNotesInBar;
 	if (TimeSignature)
 	{
@@ -547,7 +547,7 @@ void FSongMaps::SetLengthTotalBars(int32 Bars)
 
 const FSongSection* FSongMaps::GetSectionAtMs(float Ms) const
 {
-	int32 Tick = MsToTick(Ms);
+	int32 Tick = int32(MsToTick(Ms));
 	return GetSectionAtTick(Tick);
 }
 
@@ -573,7 +573,7 @@ const FSongSection* FSongMaps::GetSectionWithName(const FString& Name) const
 
 FString FSongMaps::GetSectionNameAtMs(float Ms) const
 {
-	int32 Tick = MsToTick(Ms);
+	int32 Tick = int32(MsToTick(Ms));
 	return GetSectionNameAtTick(Tick);
 }
 
@@ -589,32 +589,32 @@ FString FSongMaps::GetSectionNameAtTick(int32 Tick) const
 
 float FSongMaps::GetSectionLengthMsAtMs(float Ms) const
 {
-	int32 Tick = MsToTick(Ms);
+	int32 Tick = int32(MsToTick(Ms));
 	return GetSectionLengthMsAtTick(Tick);
 }
 
 float FSongMaps::GetSectionStartMsAtMs(float Ms) const
 {
-	const int32 Tick = MsToTick(Ms);
+	const int32 Tick = int32(MsToTick(Ms));
 	const FSongSection* SectionAtTick = GetSectionAtTick(Tick);
 	if (!SectionAtTick)
 	{
 		return 0.0f;
 	}
 
-	return TickToMs(SectionAtTick->StartTick);
+	return TickToMs(float(SectionAtTick->StartTick));
 }
 
 float FSongMaps::GetSectionEndMsAtMs(float Ms) const
 {
-	const int32 Tick = MsToTick(Ms);
+	const int32 Tick = int32(MsToTick(Ms));
 	const FSongSection* SectionAtTick = GetSectionAtTick(Tick);
 	if (!SectionAtTick)
 	{
 		return 0.0f;
 	}
 
-	return TickToMs(SectionAtTick->EndTick());
+	return TickToMs(float(SectionAtTick->EndTick()));
 }
 
 float FSongMaps::GetSectionLengthMsAtTick(int32 Tick) const
@@ -625,8 +625,8 @@ float FSongMaps::GetSectionLengthMsAtTick(int32 Tick) const
 		return 0.0f;
 	}
 
-	float StartMs = TickToMs(SectionAtTick->StartTick);
-	float EndMs = TickToMs(SectionAtTick->EndTick());
+	float StartMs = TickToMs(float(SectionAtTick->StartTick));
+	float EndMs = TickToMs(float(SectionAtTick->EndTick()));
 	return EndMs - StartMs;
 }
 
@@ -635,11 +635,11 @@ float FSongMaps::GetSectionLengthMsAtTick(int32 Tick) const
 
 const FChordMapPoint* FSongMaps::GetChordAtMs(float Ms) const
 {
-	int32 Tick = MsToTick(Ms);
+	int32 Tick = int32(MsToTick(Ms));
 	return GetChordAtTick(Tick);
 }
 
-const FChordMapPoint* FSongMaps::GetChordAtTick(float Tick) const
+const FChordMapPoint* FSongMaps::GetChordAtTick(int32 Tick) const
 {
 	if (ChordMap.IsEmpty())
 	{
@@ -651,7 +651,7 @@ const FChordMapPoint* FSongMaps::GetChordAtTick(float Tick) const
 
 FName FSongMaps::GetChordNameAtMs(float Ms) const
 {
-	int32 Tick = MsToTick(Ms);
+	int32 Tick = int32(MsToTick(Ms));
 	return GetChordNameAtTick(Tick);
 }
 
@@ -667,7 +667,7 @@ FName FSongMaps::GetChordNameAtTick(int32 Tick) const
 
 float FSongMaps::GetChordLengthMsAtMs(float Ms) const
 {
-	int32 Tick = MsToTick(Ms);
+	int32 Tick = int32(MsToTick(Ms));
 	return GetChordLengthMsAtTick(Tick);
 }
 
@@ -678,8 +678,8 @@ float FSongMaps::GetChordLengthMsAtTick(int32 Tick) const
 	{
 		return 0.0f;
 	}
-	float ChordStartMs = TickToMs(ChordInfo->StartTick);
-	float ChordEndMs = TickToMs(ChordInfo->EndTick());
+	float ChordStartMs = TickToMs(float(ChordInfo->StartTick));
+	float ChordEndMs = TickToMs(float(ChordInfo->EndTick()));
 	return ChordEndMs - ChordStartMs;
 }
 
diff --git a/Engine/Plugins/Runtime/Harmonix/Source/HarmonixMidi/Public/HarmonixMidi/BarMap.h b/Engine/Plugins/Runtime/Harmonix/Source/HarmonixMidi/Public/HarmonixMidi/BarMap.h
--- a/Engine/Plugins/Runtime/Harmonix/Source/HarmonixMidi/Public/HarmonixMidi/BarMap.h
+++ b/Engine/Plugins/Runtime/Harmonix/Source/HarmonixMidi/Public/HarmonixMidi/BarMap.h
@@ -158,7 +158,7 @@ public:
 	 * @param Timestamp
 	 * @return The 0-based absolute tick in the midi data.
 	 */
-	int32 MusicTimestampToTick(const FMusicTimestamp& Timestamp) const;
+	float MusicTimestampToTick(const FMusicTimestamp& Timestamp) const;
 
 	/**
 	 * Takes an absolute tick and calculates a music timestamp. Uses the collection of time signatures,
@@ -169,7 +169,7 @@ public:
 	 * @param OutBeatPerBar Is not nullptr it will be filled in with the number of total beats in the described bar.
 	 * @return The FMusicTimestamp where Bar 1 Beat 1 is the "beginning of the song" after count-in and pickup bars.
 	 */
-	FMusicTimestamp TickToMusicTimestamp(int32 Tick, int32* OutBeatsPerBar = nullptr) const;
+	FMusicTimestamp TickToMusicTimestamp(float Tick, int32* OutBeatsPerBar = nullptr) const;
 
 	/**
 	 * Takes an tick relative to "bar 1 beat 1" and calculates a music timestamp. Uses the collection of 
@@ -181,7 +181,7 @@ public:
 	 * @param OutBeatPerBar Is not nullptr it will be filled in with the number of total beats in the described bar.
 	 * @return The FMusicTimestamp where Bar 1 Beat 1 is the "beginning of the song" after count-in and pickup bars.
 	 */
-	FMusicTimestamp TickFromBarOneToMusicTimestamp(int32 Tick, int32* OutBeatsPerBar = nullptr) const;
+	FMusicTimestamp TickFromBarOneToMusicTimestamp(float Tick, int32* OutBeatsPerBar = nullptr) const;
 
 
 	/**
@@ -245,7 +245,7 @@ public:
 	 * @param Tick 0-based 'raw total ticks'
 	 * @return The 0-based fractional bar index
 	 */
-	float TickToFractionalBarIncludingCountIn(int32 Tick) const;
+	float TickToFractionalBarIncludingCountIn(float Tick) const;
 
 	/**
 	 * Takes an absolute tick (from the beginning of the midi data) and calculates the bar index,
@@ -268,7 +268,7 @@ public:
 	* @param FractionalBarIndex 0-based fractional bar
 	* @return The 0-based raw tick
 	*/
-	int32 FractionalBarIncludingCountInToTick(float FractionalBarIndex) const;
+	float FractionalBarIncludingCountInToTick(float FractionalBarIndex) const;
 
 	/**
 	 * Adds a time signature at the specified position, in this case a bar number as one would find
@@ -323,7 +323,7 @@ public:
 
 	void Finalize(int32 InLastTick);
 
-	int32 GetTickOfBarOne() const { return MusicTimestampBarBeatTickToTick(1,1,0); }
+	int32 GetTickOfBarOne() const { return MusicTimestampBarBeatTickToTick(1, 1, 0); }
 
 	void SetStartBar(int32 InStartBar) { StartBar = InStartBar; }
 	int32 GetStartBar() const { return StartBar; }
diff --git a/Engine/Plugins/Runtime/Harmonix/Source/HarmonixMidi/Public/HarmonixMidi/BeatMap.h b/Engine/Plugins/Runtime/Harmonix/Source/HarmonixMidi/Public/HarmonixMidi/BeatMap.h
--- a/Engine/Plugins/Runtime/Harmonix/Source/HarmonixMidi/Public/HarmonixMidi/BeatMap.h
+++ b/Engine/Plugins/Runtime/Harmonix/Source/HarmonixMidi/Public/HarmonixMidi/BeatMap.h
@@ -113,11 +113,11 @@ public:
 
 	float GetFractionalTickAtBeat(float InBeat) const;
 
-	EMusicalBeatType GetBeatTypeAtTick(float Tick) const;
+	EMusicalBeatType GetBeatTypeAtTick(int32 Tick) const;
 
-	float GetBeatInPulseBarAtTick(float tick) const;
+	float GetBeatInPulseBarAtTick(float Tick) const;
 
-	float GetNumBeatsInPulseBarAt(float Tick) const;
+	int32 GetNumBeatsInPulseBarAt(int32 Tick) const;
 
 	/** Is the specified "from count in" beat the first beat in its bar? */
 	bool IsDownbeat(float Beat) const;
diff --git a/Engine/Plugins/Runtime/Harmonix/Source/HarmonixMidi/Public/HarmonixMidi/MidiSongPos.h b/Engine/Plugins/Runtime/Harmonix/Source/HarmonixMidi/Public/HarmonixMidi/MidiSongPos.h
--- a/Engine/Plugins/Runtime/Harmonix/Source/HarmonixMidi/Public/HarmonixMidi/MidiSongPos.h
+++ b/Engine/Plugins/Runtime/Harmonix/Source/HarmonixMidi/Public/HarmonixMidi/MidiSongPos.h
@@ -83,4 +83,7 @@ public:
 	bool operator>(const FMidiSongPos& rhs) const;
 	bool operator>=(const FMidiSongPos& rhs) const;
 	bool operator==(const FMidiSongPos& rhs) const;
+
+private:
+	void SetByTimeAndTick(float Ms, float Tick, const FSongMaps& Maps);
 };
diff --git a/Engine/Plugins/Runtime/Harmonix/Source/HarmonixMidi/Public/HarmonixMidi/MusicMapBase.h b/Engine/Plugins/Runtime/Harmonix/Source/HarmonixMidi/Public/HarmonixMidi/MusicMapBase.h
--- a/Engine/Plugins/Runtime/Harmonix/Source/HarmonixMidi/Public/HarmonixMidi/MusicMapBase.h
+++ b/Engine/Plugins/Runtime/Harmonix/Source/HarmonixMidi/Public/HarmonixMidi/MusicMapBase.h
@@ -28,14 +28,14 @@ public:
 	int32 LengthTicks;
 
 	FORCEINLINE int32 EndTick() const { return StartTick + LengthTicks; };
-	FORCEINLINE int32 ContainsTick(int32 Tick) const { return Tick >= StartTick && Tick < StartTick + LengthTicks; }
-	FORCEINLINE int32 TickInPoint(int32 Tick) const { return FMath::Clamp(Tick - StartTick, 0, LengthTicks - 1); }
-	FORCEINLINE float Progress(int32 Tick) const { return LengthTicks < 2 ? 1.0f : (float)TickInPoint(Tick) / (float)(LengthTicks - 1); }
+	FORCEINLINE bool ContainsTick(float Tick) const { return Tick >= StartTick && Tick < StartTick + LengthTicks; }
+	FORCEINLINE float TickInPoint(float Tick) const { return FMath::Clamp(Tick - StartTick, 0, LengthTicks - 1); }
+	FORCEINLINE float Progress(float Tick) const { return LengthTicks < 2 ? 1.0f : TickInPoint(Tick) / (LengthTicks - 1); }
 	// friend FArchive& operator<<(FArchive& Archive, FMusicMapTimespanBase& Info);
 	struct LessThan
 	{
-		bool operator()(const FMusicMapTimespanBase& Point, int32 Tick) const { return Point.StartTick < Tick; }
-		bool operator()(int32 Tick, const FMusicMapTimespanBase& Point) const { return Tick < Point.StartTick; }
+		bool operator()(const FMusicMapTimespanBase& Point, float Tick) const { return Point.StartTick < Tick; }
+		bool operator()(float Tick, const FMusicMapTimespanBase& Point) const { return Tick < Point.StartTick; }
 		bool operator()(const FMusicMapTimespanBase& Point1, const FMusicMapTimespanBase& Point2) const { return Point1.StartTick < Point2.StartTick; }
 	};
 };
@@ -47,7 +47,7 @@ struct HARMONIXMIDI_API FMusicMapBaseInterface
 	GENERATED_BODY()
 public:
 	virtual int32 GetPointIndexForTick(int32 Tick) const = 0;
-	virtual float GetFrationalPointForTick(int32 Tick) const = 0;
+	virtual float GetFractionalPointForTick(int32 Tick) const = 0;
 	virtual float GetFractionalTickForFractionalPoint(float Point) const = 0;
 	virtual void Finalize(int32 LastTick = 0) = 0;
 	virtual void RemovePointsOnAndAfterTick(int32 Tick) = 0;
@@ -93,14 +93,14 @@ public:
 	}
 
 	template<typename T>
-	static float GetFrationalPointForTick(const TArray<T>& Points, int32 Tick)
+	static float GetFractionalPointForTick(const TArray<T>& Points, float Tick)
 	{
-		int32 Index = GetPointIndexForTick(Points, Tick);
+		int32 Index = GetPointIndexForTick(Points, int32(Tick));
 		if (Index == -1)
 		{
 			return 0.0f;
 		}
-		return (float)Index + Points[Index].Progress(Tick);
+		return Index + Points[Index].Progress(Tick);
 	}
 
 	template<typename T>
@@ -118,7 +118,7 @@ public:
 		{
 			return 0.0f;
 		}
-		return (float)Points[Index].StartTick + ((float)Points[Index].LengthTicks * FMath::Frac(Point));
+		return Points[Index].StartTick + Points[Index].LengthTicks * FMath::Frac(Point);
 	}
 
 	template<typename T>
diff --git a/Engine/Plugins/Runtime/Harmonix/Source/HarmonixMidi/Public/HarmonixMidi/SongMaps.h b/Engine/Plugins/Runtime/Harmonix/Source/HarmonixMidi/Public/HarmonixMidi/SongMaps.h
--- a/Engine/Plugins/Runtime/Harmonix/Source/HarmonixMidi/Public/HarmonixMidi/SongMaps.h
+++ b/Engine/Plugins/Runtime/Harmonix/Source/HarmonixMidi/Public/HarmonixMidi/SongMaps.h
@@ -51,8 +51,8 @@ public:
 	bool LoadFromStdMidiFile(void* Buffer, int32 BufferSize, const FString& Filename);
 	bool LoadFromStdMidiFile(TSharedPtr<FArchive> Archive, const FString& Filename);
 
-	float TickToMs(int32 Tick) const;
-	int32 MsToTick(float Ms) const;
+	float TickToMs(float Tick) const;
+	float MsToTick(float Ms) const;
 	float GetCountInSeconds() const;
 
 	// tracks
@@ -79,15 +79,15 @@ public:
 	float                 GetMsPerBeatAtMs(float Ms) const;
 	float                 GetMsPerBeatAtTick(int32 Tick) const;
 	float                 GetFractionalBeatAtMs(float Ms) const;
-	float                 GetFractionalBeatAtTick(int32 Tick) const;
+	float                 GetFractionalBeatAtTick(float Tick) const;
 	int32                 GetBeatIndexAtMs(float Ms) const;
 	int32                 GetBeatIndexAtTick(int32 Tick) const;
 	EMusicalBeatType      GetBeatTypeAtMs(float Ms) const;
 	EMusicalBeatType      GetBeatTypeAtTick(int32 Tick) const;
 	float                 GetBeatInPulseBarAtMs(float Ms) const;
-	float                 GetBeatInPulseBarAtTick(int32 Tick) const;
-	float                 GetNumBeatsInPulseBarAtMs(float Ms) const ;
-	float                 GetNumBeatsInPulseBarAtTick(int32 Tick) const;
+	float                 GetBeatInPulseBarAtTick(float Tick) const;
+	int32                 GetNumBeatsInPulseBarAtMs(float Ms) const ;
+	int32                 GetNumBeatsInPulseBarAtTick(int32 Tick) const;
 	FBeatMap&             GetBeatMap() { return BeatMap; }
 	const FBeatMap&       GetBeatMap() const { return BeatMap; }
 	bool                  BeatMapIsEmpty() { return BeatMap.GetNumMapPoints() == 0; }
@@ -98,9 +98,9 @@ public:
 	const FTimeSignature* GetTimeSignatureAtTick(int32 Tick) const;
 	const FTimeSignature* GetTimeSignatureAtBar(int32 Bar) const;
 	float                 GetBarIncludingCountInAtMs(float Ms) const;
-	float                 GetBarIncludingCountInAtTick(int32 Tick) const;
+	float                 GetBarIncludingCountInAtTick(float Tick) const;
 	float                 GetMsPerBarAtMs(float Ms) const;
-	float                 GetMsPerBarAtTick(int32 Tick) const;
+	float                 GetMsPerBarAtTick(float Tick) const;
 	FBarMap&              GetBarMap() { return BarMap; }
 	const FBarMap&        GetBarMap() const { return BarMap; }
 	bool                  BarMapIsEmpty() { return BarMap.GetNumTimeSignaturePoints() == 0; }
@@ -124,7 +124,7 @@ public:
 
 	// chords
 	const FChordMapPoint* GetChordAtMs(float Ms) const;
-	const FChordMapPoint* GetChordAtTick(float Tick) const;
+	const FChordMapPoint* GetChordAtTick(int32 Tick) const;
 	FName                 GetChordNameAtMs(float Ms) const;
 	FName                 GetChordNameAtTick(int32 Tick) const;
 	float                 GetChordLengthMsAtMs(float Ms) const;
diff --git a/Engine/Plugins/Runtime/Harmonix/Source/HarmonixMidi/Public/HarmonixMidi/TempoMap.h b/Engine/Plugins/Runtime/Harmonix/Source/HarmonixMidi/Public/HarmonixMidi/TempoMap.h
--- a/Engine/Plugins/Runtime/Harmonix/Source/HarmonixMidi/Public/HarmonixMidi/TempoMap.h
+++ b/Engine/Plugins/Runtime/Harmonix/Source/HarmonixMidi/Public/HarmonixMidi/TempoMap.h
@@ -79,11 +79,6 @@ public:
 
 	/** Get the time (in milliseconds) at a given tick: */
 	float TickToMs(float Tick) const;
-	/** Get the time (in milliseconds) at a given tick: */
-	float TickToMs(int32 Tick) const;
-
-	float TickToSeconds(float Tick) const;
-	float TickToSeconds(int32 Tick) const;
 
 	/**
 	 * Get the time (in milliseconds) at the given tick.
@@ -114,8 +109,6 @@ public:
 		*/
 	float MsToTick(float TimeMs, int32 Index, int32* NextIndex) const;
 
-	inline float SecondsToTick(const float Seconds) const { return MsToTick(Seconds * 1000.f); }
-
 	/** Get the tempo in ms/quarter-note */
 	float GetMsPerQuarterNoteAtTick(int32 Tick) const;
 
@@ -165,18 +158,3 @@ private:
 	float MsToTickInternal(float TimeMs, const FTempoInfoPoint& prevTempoInfoPoint) const;
 
 };
-
-inline float FTempoMap::TickToMs(int32 Tick) const
-{
-	return TickToMs(static_cast<float>(Tick));
-}
-
-inline float FTempoMap::TickToSeconds(int32 Tick) const
-{
-	return TickToMs(Tick) / 1000.0f;
-}
-
-inline float FTempoMap::TickToSeconds(float Tick) const
-{
-	return TickToMs(Tick) / 1000.0f;
-}
