VariableSpeedTimer uses game time instead of platform time.

#tests Played Main Stage on Windows and platform clients

diff --git a/Engine/Plugins/Runtime/Harmonix/Source/Harmonix/Private/Harmonix/VariableSpeedTimer.cpp b/Engine/Plugins/Runtime/Harmonix/Source/Harmonix/Private/Harmonix/VariableSpeedTimer.cpp
--- a/Engine/Plugins/Runtime/Harmonix/Source/Harmonix/Private/Harmonix/VariableSpeedTimer.cpp
+++ b/Engine/Plugins/Runtime/Harmonix/Source/Harmonix/Private/Harmonix/VariableSpeedTimer.cpp
@@ -2,6 +2,7 @@
 
 #include "Harmonix/VariableSpeedTimer.h"
 #include "HAL/PlatformTime.h"
+#include "Misc/App.h"
 
 FVariableSpeedTimer::FVariableSpeedTimer()
 	: AccumulatedSeconds(0)
@@ -12,14 +13,14 @@ void FVariableSpeedTimer::Start()
 {
 	if (!Running)
 	{
-		CurrentStartTime = FPlatformTime::Seconds();
+		CurrentStartTime = GetSeconds();
 		Running = true;
 	}
 }
 
 void FVariableSpeedTimer::Stop()
 {
-	double CurrentSeconds = FPlatformTime::Seconds();
+	double CurrentSeconds = GetSeconds();
 
 	if (!Running)
 	{
@@ -37,13 +38,13 @@ void FVariableSpeedTimer::Reset(double InitialMs)
 	// Unreal "platform time" interface works in. BUT, this means that we 
 	// need to convert seconds <-> milliseconds for the user of this class.
 
-	CurrentStartTime = FPlatformTime::Seconds();
+	CurrentStartTime = GetSeconds();
 	AccumulatedSeconds = InitialMs / 1000.0;
 }
 
 void FVariableSpeedTimer::SetSpeed(double NewSpeed)
 {
-	double CurrentSeconds = FPlatformTime::Seconds();
+	double CurrentSeconds = GetSeconds();
 
 	if (Running)
 	{
@@ -65,7 +66,12 @@ double FVariableSpeedTimer::Ms()
 		return AccumulatedSeconds * 1000.0;
 	}
 
-	double CurrentSeconds = FPlatformTime::Seconds();
+	double CurrentSeconds = GetSeconds();
 
 	return (AccumulatedSeconds + (CurrentSeconds - CurrentStartTime) * Speed) * 1000.0f;
 }
+
+double FVariableSpeedTimer::GetSeconds()
+{
+	return FApp::GetGameTime();
+}
diff --git a/Engine/Plugins/Runtime/Harmonix/Source/Harmonix/Public/Harmonix/VariableSpeedTimer.h b/Engine/Plugins/Runtime/Harmonix/Source/Harmonix/Public/Harmonix/VariableSpeedTimer.h
--- a/Engine/Plugins/Runtime/Harmonix/Source/Harmonix/Public/Harmonix/VariableSpeedTimer.h
+++ b/Engine/Plugins/Runtime/Harmonix/Source/Harmonix/Public/Harmonix/VariableSpeedTimer.h
@@ -21,6 +21,8 @@ public:
 	void Reset(double InitialMs = 0.0);
 
 private:
+	double GetSeconds();
+
 	// NOTE: Internally we track time in seconds because that is what the 
 	// Unreal "platform time" interface works in. BUT, this means that we 
 	// need to convert seconds <-> milliseconds for the user of this class.
