[Harmonix] Midi Track Filter: decouple filter logic from node, implement range-based filter node, pull stuck note logic into its own utility class

#tests automated tests

#changelist validated
#virtualized

diff --git a/Engine/Plugins/Runtime/Harmonix/Source/HarmonixMetasound/Private/HarmonixMetasound/MidiOps/MidiChannelFilter.cpp b/Engine/Plugins/Runtime/Harmonix/Source/HarmonixMetasound/Private/HarmonixMetasound/MidiOps/MidiChannelFilter.cpp
--- a/Engine/Plugins/Runtime/Harmonix/Source/HarmonixMetasound/Private/HarmonixMetasound/MidiOps/MidiChannelFilter.cpp
+++ b/Engine/Plugins/Runtime/Harmonix/Source/HarmonixMetasound/Private/HarmonixMetasound/MidiOps/MidiChannelFilter.cpp
@@ -6,53 +6,27 @@ namespace Harmonix::Midi::Ops
 {
 	void FMidiChannelFilter::Process(const HarmonixMetasound::FMidiStream& InStream, HarmonixMetasound::FMidiStream& OutStream)
 	{
-		OutStream.Copy(
-			InStream,
-			[this](const HarmonixMetasound::FMidiStreamEvent& Event)
+		const auto Filter = [this](const HarmonixMetasound::FMidiStreamEvent& Event)
+		{
+			if (Event.MidiMessage.IsStd())
 			{
-				if (Event.MidiMessage.IsStd())
-				{
-					const uint8 ChannelIdx = Event.MidiMessage.GetStdChannel();
-					check(ChannelIdx < 16);
-
-					const bool ChannelIsEnabled = GetChannelEnabled(ChannelIdx + 1);
+				const uint8 ChannelIdx = Event.MidiMessage.GetStdChannel();
+				check(ChannelIdx < 16);
+				return GetChannelEnabled(ChannelIdx + 1);
+			}
 
-					// If this is a note on or off, keep track of active voices so we can unstick notes if the filter changes later
-					if (ChannelIsEnabled)
-					{
-						if (Event.MidiMessage.IsNoteOn())
-						{
-							ActiveVoices[ChannelIdx].AddUnique(Event.GetVoiceId());
-						}
-						else if (Event.MidiMessage.IsNoteOff())
-						{
-							ActiveVoices[ChannelIdx].Remove(Event.GetVoiceId());
-						}
-					}
-					
-					return ChannelIsEnabled;
-				}
+			// Pass non-std messages
+			return true;
+		};
 
-				// Pass non-std messages
-				return true;
-			},
+		// Copy events which pass the filter to the output
+		OutStream.Copy(
+			InStream,
+			Filter,
 			true /* include transport events */);
 
-		// If we turned off some channels, send note offs for any active voices in that channel
-		for (uint8 ChannelIdx = 0; ChannelIdx < MidiConstants::kNumChannels; ++ChannelIdx)
-		{
-			if (!GetChannelEnabled(ChannelIdx + 1))
-			{
-				for (const FMidiVoiceId& VoiceId : ActiveVoices[ChannelIdx])
-				{
-					HarmonixMetasound::FMidiStreamEvent NoteOffEvent{ static_cast<uint32>(0), FMidiMsg::CreateNoteOff(ChannelIdx, 0) };
-					NoteOffEvent.SetVoiceId(VoiceId);
-					OutStream.InsertMidiEvent(NoteOffEvent);
-				}
-
-				ActiveVoices[ChannelIdx].Reset();
-			}
-		}
+		// Unstick notes if necessary
+		StuckNoteGuard.Process(InStream, OutStream, Filter);
 	}
 
 	void FMidiChannelFilter::SetChannelEnabled(const uint8 Channel, const bool Enabled)
diff --git a/Engine/Plugins/Runtime/Harmonix/Source/HarmonixMetasound/Private/HarmonixMetasound/MidiOps/MidiNoteFilter.cpp b/Engine/Plugins/Runtime/Harmonix/Source/HarmonixMetasound/Private/HarmonixMetasound/MidiOps/MidiNoteFilter.cpp
--- a/Engine/Plugins/Runtime/Harmonix/Source/HarmonixMetasound/Private/HarmonixMetasound/MidiOps/MidiNoteFilter.cpp
+++ b/Engine/Plugins/Runtime/Harmonix/Source/HarmonixMetasound/Private/HarmonixMetasound/MidiOps/MidiNoteFilter.cpp
@@ -6,79 +6,42 @@ namespace Harmonix::Midi::Ops
 {
 	void FMidiNoteFilter::Process(const HarmonixMetasound::FMidiStream& InStream, HarmonixMetasound::FMidiStream& OutStream)
 	{
-		OutStream.Copy(
-			InStream,
-			[this](const HarmonixMetasound::FMidiStreamEvent& Event)
+		const auto Filter = [this](const HarmonixMetasound::FMidiStreamEvent& Event)
+		{
+			// Note ons must pass note number and velocity check
+			if (Event.MidiMessage.IsNoteOn())
 			{
-				// Note ons must pass note number and velocity check
-				if (Event.MidiMessage.IsNoteOn())
-				{
-					const FMidiVoiceId VoiceId = Event.GetVoiceId();
-					const uint8 Channel = Event.MidiMessage.GetStdChannel();
-					const uint8 NoteNumber = Event.MidiMessage.GetStdData1();
-					const uint8 Velocity = Event.MidiMessage.GetStdData2();
-					const bool bPassesFilter =
-						NoteNumber >= MinNoteNumber
+				const uint8 NoteNumber = Event.MidiMessage.GetStdData1();
+				const uint8 Velocity = Event.MidiMessage.GetStdData2();
+				return NoteNumber >= MinNoteNumber
 						&& NoteNumber <= MaxNoteNumber
 						&& Velocity >= MinVelocity
 						&& Velocity <= MaxVelocity;
+			}
 
-					if (bPassesFilter)
-					{
-						ActiveVoices.FindOrAdd(VoiceId) = { Channel, NoteNumber, Velocity };
-					}
-
-					return bPassesFilter;
-				}
-
-				// Note offs must pass note number check
-				if (Event.MidiMessage.IsNoteOff())
-				{
-					const FMidiVoiceId VoiceId = Event.GetVoiceId();
-					const uint8 NoteNumber = Event.MidiMessage.GetStdData1();
-					const bool bPassesFilter = NoteNumber >= MinNoteNumber && NoteNumber <= MaxNoteNumber;
+			// Note offs must pass note number check
+			if (Event.MidiMessage.IsNoteOff())
+			{
+				const uint8 NoteNumber = Event.MidiMessage.GetStdData1();
+				return NoteNumber >= MinNoteNumber && NoteNumber <= MaxNoteNumber;
+			}
 
-					if (bPassesFilter)
-					{
-						ActiveVoices.Remove(VoiceId);
-					}
-					
-					return bPassesFilter;
-				}
+			// Handle all notes off/kill
+			if (Event.MidiMessage.IsNoteMessage())
+			{
+				return true;
+			}
 
-				// Handle all notes off/kill
-				if (Event.MidiMessage.IsNoteMessage())
-				{
-					ActiveVoices.Reset();
-					return true;
-				}
+			return IncludeOtherEvents.Get();
+		};
 
-				return IncludeOtherEvents.Get();
-			},
+		// Copy events which pass the filter to the output
+		OutStream.Copy(
+			InStream,
+			Filter,
 			true /* include transport events */);
 
-		// If active notes we're tracking no longer pass the filter, send a note off for each
-		for (auto It = ActiveVoices.CreateIterator(); It; ++It)
-		{
-			const uint8 NoteNumber = It.Value().NoteNumber;
-			const uint8 Velocity = It.Value().Velocity;
-			const bool bPassesFilter =
-				NoteNumber >= MinNoteNumber
-				&& NoteNumber <= MaxNoteNumber
-				&& Velocity >= MinVelocity
-				&& Velocity <= MaxVelocity;
-
-			if (!bPassesFilter)
-			{
-				// Send a note off
-				const uint8 Channel = It.Value().Channel;
-				HarmonixMetasound::FMidiStreamEvent NoteOffEvent{ static_cast<uint32>(0), FMidiMsg::CreateNoteOff(Channel, NoteNumber) };
-				NoteOffEvent.SetVoiceId(It.Key());
-				OutStream.InsertMidiEvent(NoteOffEvent);
-				
-				// Stop tracking
-				It.RemoveCurrent();
-			}
-		}
+		// Unstick notes if necessary
+		StuckNoteGuard.Process(InStream, OutStream, Filter);
 	}
 }
diff --git a/Engine/Plugins/Runtime/Harmonix/Source/HarmonixMetasound/Private/HarmonixMetasound/MidiOps/MidiTrackFilter.cpp b/Engine/Plugins/Runtime/Harmonix/Source/HarmonixMetasound/Private/HarmonixMetasound/MidiOps/MidiTrackFilter.cpp
new file mode 100644
--- /dev/null
+++ b/Engine/Plugins/Runtime/Harmonix/Source/HarmonixMetasound/Private/HarmonixMetasound/MidiOps/MidiTrackFilter.cpp
@@ -0,0 +1,43 @@
+// Copyright Epic Games, Inc. All Rights Reserved.
+
+#include "HarmonixMetasound/MidiOps/MidiTrackFilter.h"
+
+namespace Harmonix::Midi::Ops
+{
+	void FMidiTrackFilter::Process(const HarmonixMetasound::FMidiStream& InStream, HarmonixMetasound::FMidiStream& OutStream)
+	{
+		const auto Filter = [this](const HarmonixMetasound::FMidiStreamEvent& Event)
+		{
+			ensure(Event.TrackIndex >= 0 && Event.TrackIndex <= UINT16_MAX);
+			const uint16 TrackIndex = static_cast<uint16>(Event.TrackIndex);
+			bool IncludeEvent = false;
+
+			if (TrackIndex == 0)
+			{
+				IncludeEvent = IncludeConductorTrack;
+			}
+			else if (TrackIndex >= MinTrackIdx && TrackIndex <= MaxTrackIdx)
+			{
+				IncludeEvent = true;
+			}
+
+			return IncludeEvent;
+		};
+		
+		// Copy events which pass the filter to the output
+		OutStream.Copy(
+			InStream,
+			Filter,
+			true /* include transport events */);
+
+		// Unstick notes if necessary
+		StuckNoteGuard.Process(InStream, OutStream, Filter);
+	}
+
+	void FMidiTrackFilter::SetTrackRange(const uint16 InMinTrackIdx, const uint16 InMaxTrackIdx, const bool InIncludeConductorTrack)
+	{
+		MinTrackIdx = FMath::Clamp(InMinTrackIdx, 0, UINT16_MAX);
+		MaxTrackIdx = FMath::Clamp(InMaxTrackIdx, MinTrackIdx, UINT16_MAX);
+		IncludeConductorTrack = InIncludeConductorTrack;
+	}
+}
diff --git a/Engine/Plugins/Runtime/Harmonix/Source/HarmonixMetasound/Private/HarmonixMetasound/MidiOps/StuckNoteGuard.cpp b/Engine/Plugins/Runtime/Harmonix/Source/HarmonixMetasound/Private/HarmonixMetasound/MidiOps/StuckNoteGuard.cpp
new file mode 100644
--- /dev/null
+++ b/Engine/Plugins/Runtime/Harmonix/Source/HarmonixMetasound/Private/HarmonixMetasound/MidiOps/StuckNoteGuard.cpp
@@ -0,0 +1,66 @@
+// Copyright Epic Games, Inc. All Rights Reserved.
+
+#include "HarmonixMetasound/MidiOps/StuckNoteGuard.h"
+
+namespace Harmonix::Midi::Ops
+{
+	void FStuckNoteGuard::Process(
+		const HarmonixMetasound::FMidiStream& InStream,
+		HarmonixMetasound::FMidiStream& OutStream,
+		const FIncludeNotePredicate& Predicate)
+	{
+		for (const HarmonixMetasound::FMidiStreamEvent& Event : InStream.GetEventsInBlock())
+		{
+			// If this is a note on which passes the filter, track it
+			if (Event.MidiMessage.IsNoteOn() && Predicate(Event))
+			{
+				const int32 TrackIndex = Event.TrackIndex;
+				const FMidiVoiceId VoiceId = Event.GetVoiceId();
+
+				if (!ActiveNotes.ContainsByPredicate([TrackIndex, VoiceId](const HarmonixMetasound::FMidiStreamEvent& TrackedEvent)
+				{
+					return TrackIndex == TrackedEvent.TrackIndex && VoiceId == TrackedEvent.GetVoiceId();
+				}))
+				{
+					ActiveNotes.Add(Event);
+				}
+			}
+			// If this is a note off which passes the filter, untrack it
+			else if (Event.MidiMessage.IsNoteOff() && Predicate(Event))
+			{
+				const int32 TrackIndex = Event.TrackIndex;
+				const FMidiVoiceId VoiceId = Event.GetVoiceId();
+				
+				ActiveNotes.RemoveAll([TrackIndex, VoiceId](const HarmonixMetasound::FMidiStreamEvent& TrackedEvent)
+				{
+					return TrackIndex == TrackedEvent.TrackIndex && VoiceId == TrackedEvent.GetVoiceId(); 
+				});
+			}
+			// If this is an all notes off/kill, untrack everything
+			// NOTE: This diverges from the MIDI spec, where "all notes off" has a track and channel associated with it.
+			// If we end up supporting the track- and channel-specific "all notes off" we will need to add support here to avoid stuck notes.
+			else if (Event.MidiMessage.IsAllNotesOff() || Event.MidiMessage.IsAllNotesKill())
+			{
+				ActiveNotes.Reset();
+			}
+		}
+		
+		// For all of the tracked notes, unstick ones that no longer pass the filter
+		for (auto It = ActiveNotes.CreateIterator(); It; ++It)
+		{
+			if (!Predicate(*It))
+			{
+				// Send a note off
+				check(It->MidiMessage.IsNoteOn()); // shouldn't happen, yell for help
+				const uint8 Channel = It->MidiMessage.GetStdChannel();
+				const uint8 NoteNumber = It->MidiMessage.GetStdData1();
+				HarmonixMetasound::FMidiStreamEvent NoteOffEvent{ static_cast<uint32>(0), FMidiMsg::CreateNoteOff(Channel, NoteNumber) };
+				NoteOffEvent.SetVoiceId(It->GetVoiceId());
+				OutStream.InsertMidiEvent(NoteOffEvent);
+
+				// Stop tracking this note on this track
+				It.RemoveCurrent();
+			}
+		}
+	}
+}
diff --git a/Engine/Plugins/Runtime/Harmonix/Source/HarmonixMetasound/Private/HarmonixMetasound/Nodes/MidiStreamTrackFilterNode.cpp b/Engine/Plugins/Runtime/Harmonix/Source/HarmonixMetasound/Private/HarmonixMetasound/Nodes/MidiStreamTrackFilterNode.cpp
--- a/Engine/Plugins/Runtime/Harmonix/Source/HarmonixMetasound/Private/HarmonixMetasound/Nodes/MidiStreamTrackFilterNode.cpp
+++ b/Engine/Plugins/Runtime/Harmonix/Source/HarmonixMetasound/Private/HarmonixMetasound/Nodes/MidiStreamTrackFilterNode.cpp
@@ -8,16 +8,17 @@
 #include "MetasoundStandardNodesCategories.h"
 #include "MetasoundVertex.h"
 
-#include "HarmonixMetasound/Common.h"
 #include "HarmonixMetasound/DataTypes/MidiStream.h"
 #include "HarmonixMetasound/DataTypes/MusicTransport.h"
 #include "HarmonixMetasound/Nodes/MidiStreamTrackFilterNode.h"
 
+#include "HarmonixMetasound/MidiOps/MidiTrackFilter.h"
+
 DEFINE_LOG_CATEGORY_STATIC(LogMidiStreamTrackFilter, Log, All);
 
 #define LOCTEXT_NAMESPACE "HarmonixMetaSound"
 
-namespace HarmonixMetasound::Nodes::MidiStreamTrackFilter
+namespace HarmonixMetasound::Nodes::MidiTrackFilter
 {
 	using namespace Metasound;
 	
@@ -32,22 +33,186 @@ namespace HarmonixMetasound::Nodes::MidiStreamTrackFilter
 		return ClassName;
 	}
 
-	class FOp : public TExecutableOperator<FOp>
+	int32 GetCurrentMajorVersion()
+	{
+		return 1;
+	}
+
+	namespace Inputs
+	{
+		DEFINE_METASOUND_PARAM_ALIAS(Enable, CommonPinNames::Inputs::Enable);
+		DEFINE_METASOUND_PARAM_ALIAS(MidiStream, CommonPinNames::Inputs::MidiStream);
+		DEFINE_INPUT_METASOUND_PARAM(MinTrackIndex, "Min. Track Index", "The first track index to include (1-based, inclusive)");
+		DEFINE_INPUT_METASOUND_PARAM(MaxTrackIndex, "Max. Track Index", "The last track index to include (1-based, inclusive)");
+		DEFINE_INPUT_METASOUND_PARAM(IncludeConductorTrack, "Include Conductor Track", "Enable to include the conductor track (AKA track 0)");
+	}
+
+	namespace Outputs
+	{
+		DEFINE_METASOUND_PARAM_ALIAS(MidiStream, CommonPinNames::Outputs::MidiStream);
+	}
+	
+	class FMidiStreamTrackFilterOperator_V1 final : public TExecutableOperator<FMidiStreamTrackFilterOperator_V1>
+	{
+	public:
+		static const FNodeClassMetadata& GetNodeInfo()
+		{
+			auto InitNodeInfo = []() -> FNodeClassMetadata
+			{
+				FNodeClassMetadata Info;
+				Info.ClassName        = GetClassName();
+				Info.MajorVersion     = 1;
+				Info.MinorVersion     = 0;
+				Info.DisplayName      = METASOUND_LOCTEXT("MidiStreamTrackFilterNodeV1_DisplayName", "Midi Stream Track Filter");
+				Info.Description      = METASOUND_LOCTEXT("MidiStreamTrackFilterNodeV1_Description", "Duplicates the incoming midi stream to its output after filtering tracks from the input Midi Stream");
+				Info.Author           = PluginAuthor;
+				Info.PromptIfMissing  = PluginNodeMissingPrompt;
+				Info.DefaultInterface = GetVertexInterface();
+				Info.CategoryHierarchy.Emplace(NodeCategories::Music);
+				return Info;
+			};
+
+			static const FNodeClassMetadata Info = InitNodeInfo();
+
+			return Info;
+		}
+		
+		static const FVertexInterface& GetVertexInterface()
+		{
+			static const FVertexInterface Interface(
+				FInputVertexInterface(
+					TInputDataVertex<bool>(METASOUND_GET_PARAM_NAME_AND_METADATA(Inputs::Enable), true),
+					TInputDataVertex<FMidiStream>(METASOUND_GET_PARAM_NAME_AND_METADATA(Inputs::MidiStream)),
+					TInputDataVertex<int32>(METASOUND_GET_PARAM_NAME_AND_METADATA(Inputs::MinTrackIndex), 0),
+					TInputDataVertex<int32>(METASOUND_GET_PARAM_NAME_AND_METADATA(Inputs::MaxTrackIndex), 0),
+					TInputDataVertex<bool>(METASOUND_GET_PARAM_NAME_AND_METADATA(Inputs::IncludeConductorTrack), false)
+				),
+				FOutputVertexInterface(
+					TOutputDataVertex<FMidiStream>(METASOUND_GET_PARAM_NAME_AND_METADATA(Outputs::MidiStream))
+				)
+			);
+
+			return Interface;
+		}
+
+		struct FInputs
+		{
+			FBoolReadRef Enabled;
+			FMidiStreamReadRef MidiStream;
+			FInt32ReadRef MinTrackIndex;
+			FInt32ReadRef MaxTrackIndex;
+			FBoolReadRef IncludeConductorTrack;
+		};
+
+		struct FOutputs
+		{
+			FMidiStreamWriteRef MidiStream;
+		};
+
+		static TUniquePtr<IOperator> CreateOperator(const FBuildOperatorParams& InParams, FBuildResults& OutResults)
+		{
+			const FInputVertexInterfaceData& InputData = InParams.InputData;
+
+			FInputs Inputs
+			{
+				InputData.GetOrCreateDefaultDataReadReference<bool>(Inputs::EnableName, InParams.OperatorSettings),
+				InputData.GetOrConstructDataReadReference<FMidiStream>(Inputs::MidiStreamName, InParams.OperatorSettings),
+				InputData.GetOrCreateDefaultDataReadReference<int32>(Inputs::MinTrackIndexName, InParams.OperatorSettings),
+				InputData.GetOrCreateDefaultDataReadReference<int32>(Inputs::MaxTrackIndexName, InParams.OperatorSettings),
+				InputData.GetOrCreateDefaultDataReadReference<bool>(Inputs::IncludeConductorTrackName, InParams.OperatorSettings)
+			};
+
+			FOutputs Outputs
+			{
+				FMidiStreamWriteRef::CreateNew(InParams.OperatorSettings)
+			};
+			
+			return MakeUnique<FMidiStreamTrackFilterOperator_V1>(InParams, MoveTemp(Inputs), MoveTemp(Outputs));
+		}
+
+		FMidiStreamTrackFilterOperator_V1(const FBuildOperatorParams& InParams, FInputs&& InInputs, FOutputs&& InOutputs)
+			: Inputs(MoveTemp(InInputs))
+			, Outputs(MoveTemp(InOutputs))
+		{
+			Reset(InParams);
+		}
+
+		
+		virtual void BindInputs(FInputVertexInterfaceData& InVertexData) override
+		{
+			InVertexData.BindReadVertex(Inputs::EnableName, Inputs.Enabled);
+			InVertexData.BindReadVertex(Inputs::MidiStreamName, Inputs.MidiStream);
+			InVertexData.BindReadVertex(Inputs::MinTrackIndexName, Inputs.MinTrackIndex);
+			InVertexData.BindReadVertex(Inputs::MaxTrackIndexName, Inputs.MaxTrackIndex);
+			InVertexData.BindReadVertex(Inputs::IncludeConductorTrackName, Inputs.IncludeConductorTrack);
+
+			ReassignOutputClock = true;
+		}
+
+		virtual void BindOutputs(FOutputVertexInterfaceData& InVertexData) override
+		{
+			InVertexData.BindReadVertex(Outputs::MidiStreamName, Outputs.MidiStream);
+			
+			ReassignOutputClock = true;
+		}
+
+		void Reset(const FResetParams& ResetParams)
+		{
+			ReassignOutputClock = true;
+		}
+
+		void Execute()
+		{
+			if (ReassignOutputClock)
+			{
+				if (auto* Clock = Inputs.MidiStream->GetMidiClockSource())
+				{
+					Outputs.MidiStream->SetClockSource(*Clock);
+				}
+				ReassignOutputClock = false;
+			}
+
+			Filter.SetTrackRange(*Inputs.MinTrackIndex, *Inputs.MaxTrackIndex, *Inputs.IncludeConductorTrack);
+
+			Outputs.MidiStream->PrepareBlock();
+
+			if (*Inputs.Enabled)
+			{
+				Filter.Process(*Inputs.MidiStream, *Outputs.MidiStream);
+			}
+		}
+	private:
+		FInputs Inputs;
+		FOutputs Outputs;
+		Harmonix::Midi::Ops::FMidiTrackFilter Filter;
+		bool ReassignOutputClock = true;
+	};
+
+	class FMidiStreamTrackFilterNode_V1 final : public FNodeFacade
+	{
+	public:
+		explicit FMidiStreamTrackFilterNode_V1(const FNodeInitData& InInitData)
+			: FNodeFacade(InInitData.InstanceName, InInitData.InstanceID, TFacadeOperatorClass<FMidiStreamTrackFilterOperator_V1>())
+		{}
+		virtual ~FMidiStreamTrackFilterNode_V1() override = default;
+	};
+
+	METASOUND_REGISTER_NODE(FMidiStreamTrackFilterNode_V1)
+	
+	class FMidiStreamTrackFilterOperator_V0 : public TExecutableOperator<FMidiStreamTrackFilterOperator_V0>
 	{
 	public:
 		static const FNodeClassMetadata& GetNodeInfo();
 		static const FVertexInterface& GetVertexInterface();
 		static TUniquePtr<IOperator> CreateOperator(const FBuildOperatorParams& InParams, FBuildResults& OutResults);
 
-		FOp(const FBuildOperatorParams& InParams,
+		FMidiStreamTrackFilterOperator_V0(const FBuildOperatorParams& InParams,
 								  const FBoolReadRef&       InEnabled,
 								  const FStringReadRef&     InTrackSelect,
 								  const FMidiStreamReadRef& InMidiStream);
 
 		virtual void BindInputs(FInputVertexInterfaceData& InVertexData) override;
 		virtual void BindOutputs(FOutputVertexInterfaceData& InVertexData) override;
-		virtual FDataReferenceCollection GetInputs() const override;
-		virtual FDataReferenceCollection GetOutputs() const override;
 
 		void Reset(const FResetParams& ResetParams);
 
@@ -67,18 +232,18 @@ namespace HarmonixMetasound::Nodes::MidiStreamTrackFilter
 		FMidiStreamEventTrackChannelFilter Filter;
 	};
 
-	class FMidiStreamTrackFilterNode : public FNodeFacade
+	class FMidiStreamTrackFilterNode_V0 : public FNodeFacade
 	{
 	public:
-		FMidiStreamTrackFilterNode(const FNodeInitData& InInitData)
-			: FNodeFacade(InInitData.InstanceName, InInitData.InstanceID, TFacadeOperatorClass<FOp>())
+		FMidiStreamTrackFilterNode_V0(const FNodeInitData& InInitData)
+			: FNodeFacade(InInitData.InstanceName, InInitData.InstanceID, TFacadeOperatorClass<FMidiStreamTrackFilterOperator_V0>())
 		{}
-		virtual ~FMidiStreamTrackFilterNode() = default;
+		virtual ~FMidiStreamTrackFilterNode_V0() = default;
 	};
 
-	METASOUND_REGISTER_NODE(FMidiStreamTrackFilterNode)
+	METASOUND_REGISTER_NODE(FMidiStreamTrackFilterNode_V0)
 		
-	const FNodeClassMetadata& FOp::GetNodeInfo()
+	const FNodeClassMetadata& FMidiStreamTrackFilterOperator_V0::GetNodeInfo()
 	{
 		auto InitNodeInfo = []() -> FNodeClassMetadata
 		{
@@ -86,12 +251,13 @@ namespace HarmonixMetasound::Nodes::MidiStreamTrackFilter
 			Info.ClassName        = GetClassName();
 			Info.MajorVersion     = 0;
 			Info.MinorVersion     = 1;
-			Info.DisplayName      = METASOUND_LOCTEXT("MidiStreamTrackFilterNode_DisplayName", "Midi Stream Track Filter");
-			Info.Description      = METASOUND_LOCTEXT("MidiStreamTrackFilterNode_Description", "Duplicates the incoming midi stream to its output after filtering tracks from the input Midi Stream");
+			Info.DisplayName      = METASOUND_LOCTEXT("MidiStreamTrackFilterNodeV0_DisplayName", "Midi Stream Track Filter");
+			Info.Description      = METASOUND_LOCTEXT("MidiStreamTrackFilterNodeV0_Description", "Duplicates the incoming midi stream to its output after filtering tracks from the input Midi Stream");
 			Info.Author           = PluginAuthor;
 			Info.PromptIfMissing  = PluginNodeMissingPrompt;
 			Info.DefaultInterface = GetVertexInterface();
 			Info.CategoryHierarchy.Emplace(NodeCategories::Music);
+			Info.bDeprecated = true;
 			return Info;
 		};
 
@@ -100,14 +266,12 @@ namespace HarmonixMetasound::Nodes::MidiStreamTrackFilter
 		return Info;
 	}
 
-	const FVertexInterface& FOp::GetVertexInterface()
+	const FVertexInterface& FMidiStreamTrackFilterOperator_V0::GetVertexInterface()
 	{
-		using namespace CommonPinNames;
-
 		static const FVertexInterface Interface(
 			FInputVertexInterface(
 				TInputDataVertex<bool>(METASOUND_GET_PARAM_NAME_AND_METADATA(Inputs::Enable), true),
-				TInputDataVertex<FString>(METASOUND_GET_PARAM_NAME_AND_METADATA(Inputs::MidiTrackIndexFilterSpecifier), FString(TEXT("*"))),
+				TInputDataVertex<FString>(METASOUND_GET_PARAM_NAME_AND_METADATA(CommonPinNames::Inputs::MidiTrackIndexFilterSpecifier), FString(TEXT("*"))),
 				TInputDataVertex<FMidiStream>(METASOUND_GET_PARAM_NAME_AND_METADATA(Inputs::MidiStream))
 			),
 			FOutputVertexInterface(
@@ -118,24 +282,22 @@ namespace HarmonixMetasound::Nodes::MidiStreamTrackFilter
 		return Interface;
 	}
 
-	TUniquePtr<IOperator> FOp::CreateOperator(const FBuildOperatorParams& InParams, FBuildResults& OutResults)
+	TUniquePtr<IOperator> FMidiStreamTrackFilterOperator_V0::CreateOperator(const FBuildOperatorParams& InParams, FBuildResults& OutResults)
 	{
-		using namespace CommonPinNames;
-
-		const FMidiStreamTrackFilterNode& LoggerNode = static_cast<const FMidiStreamTrackFilterNode&>(InParams.Node);
+		const FMidiStreamTrackFilterNode_V0& LoggerNode = static_cast<const FMidiStreamTrackFilterNode_V0&>(InParams.Node);
 
 		const FInputVertexInterfaceData& InputData = InParams.InputData;
 		FBoolReadRef InEnabled = InputData.GetOrCreateDefaultDataReadReference<bool>(METASOUND_GET_PARAM_NAME(Inputs::Enable), InParams.OperatorSettings);
-		FStringReadRef InTrackIndexFilter = InputData.GetOrCreateDefaultDataReadReference<FString>(METASOUND_GET_PARAM_NAME(Inputs::MidiTrackIndexFilterSpecifier), InParams.OperatorSettings);
+		FStringReadRef InTrackIndexFilter = InputData.GetOrCreateDefaultDataReadReference<FString>(METASOUND_GET_PARAM_NAME(CommonPinNames::Inputs::MidiTrackIndexFilterSpecifier), InParams.OperatorSettings);
 		FMidiStreamReadRef InMidiStream = InputData.GetOrConstructDataReadReference<FMidiStream>(METASOUND_GET_PARAM_NAME(Inputs::MidiStream), InParams.OperatorSettings);
 
-		return MakeUnique<FOp>(InParams, 
+		return MakeUnique<FMidiStreamTrackFilterOperator_V0>(InParams, 
 			InEnabled,
 			InTrackIndexFilter,
 			InMidiStream);
 	}
 
-	FOp::FOp(const FBuildOperatorParams& InParams,
+	FMidiStreamTrackFilterOperator_V0::FMidiStreamTrackFilterOperator_V0(const FBuildOperatorParams& InParams,
 											   const FBoolReadRef& InEnabled,
 											   const FStringReadRef& InTrackSelect,
 											   const FMidiStreamReadRef& InMidiStream)
@@ -147,35 +309,21 @@ namespace HarmonixMetasound::Nodes::MidiStreamTrackFilter
 		Reset(InParams);
 	}
 
-	void FOp::BindInputs(FInputVertexInterfaceData& InVertexData)
+	void FMidiStreamTrackFilterOperator_V0::BindInputs(FInputVertexInterfaceData& InVertexData)
 	{
-		using namespace CommonPinNames;
-
 		InVertexData.BindReadVertex(METASOUND_GET_PARAM_NAME(Inputs::Enable), EnableInPin);
-		InVertexData.BindReadVertex(METASOUND_GET_PARAM_NAME(Inputs::MidiTrackIndexFilterSpecifier), TrackSelectInPin);
+		InVertexData.BindReadVertex(METASOUND_GET_PARAM_NAME(CommonPinNames::Inputs::MidiTrackIndexFilterSpecifier), TrackSelectInPin);
 		InVertexData.BindReadVertex(METASOUND_GET_PARAM_NAME(Inputs::MidiStream), MidiStreamInPin);
 	}
 
-	void FOp::BindOutputs(FOutputVertexInterfaceData& InVertexData)
+	void FMidiStreamTrackFilterOperator_V0::BindOutputs(FOutputVertexInterfaceData& InVertexData)
 	{
 		using namespace CommonPinNames;
 
 		InVertexData.BindReadVertex(METASOUND_GET_PARAM_NAME(Outputs::MidiStream), MidiStreamOutPin);
 	}
 
-	FDataReferenceCollection FOp::GetInputs() const
-	{
-		checkNoEntry();
-		return {};
-	}
-
-	FDataReferenceCollection FOp::GetOutputs() const
-	{
-		checkNoEntry();
-		return {};
-	}
-
-	void FOp::Reset(const FResetParams& ResetParams)
+	void FMidiStreamTrackFilterOperator_V0::Reset(const FResetParams& ResetParams)
 	{
 		MidiStreamOutPin->PrepareBlock();
 		
@@ -199,7 +347,7 @@ namespace HarmonixMetasound::Nodes::MidiStreamTrackFilter
 		}
 	}
 
-	void FOp::Execute()
+	void FMidiStreamTrackFilterOperator_V0::Execute()
 	{
 		MidiStreamOutPin->PrepareBlock();
 
@@ -214,6 +362,7 @@ namespace HarmonixMetasound::Nodes::MidiStreamTrackFilter
 			}
 		}
 
+		MidiStreamOutPin->PrepareBlock();
 		MidiStreamOutPin->Copy(MidiStreamInPin, [&](const FMidiStreamEvent& InEvent) {return Filter(InEvent); });
 	}
 }
diff --git a/Engine/Plugins/Runtime/Harmonix/Source/HarmonixMetasound/Public/HarmonixMetasound/MidiOps/MidiChannelFilter.h b/Engine/Plugins/Runtime/Harmonix/Source/HarmonixMetasound/Public/HarmonixMetasound/MidiOps/MidiChannelFilter.h
--- a/Engine/Plugins/Runtime/Harmonix/Source/HarmonixMetasound/Public/HarmonixMetasound/MidiOps/MidiChannelFilter.h
+++ b/Engine/Plugins/Runtime/Harmonix/Source/HarmonixMetasound/Public/HarmonixMetasound/MidiOps/MidiChannelFilter.h
@@ -3,6 +3,8 @@
 #pragma once
 
 #include "MidiOps.h"
+#include "StuckNoteGuard.h"
+
 #include "HarmonixMetasound/DataTypes/MidiStream.h"
 
 namespace Harmonix::Midi::Ops
@@ -38,8 +40,6 @@ namespace Harmonix::Midi::Ops
 		// bitfield to store which channels are enabled
 		uint16 EnabledChannels = AllChannelsOff;
 
-		// Keep track of active voices in case we turn off a channel.
-		// If we do turn off a channel, we will immediately send note offs for any active voices in that channel so they don't become stuck.
-		TArray<FMidiVoiceId> ActiveVoices[MidiConstants::kNumChannels];
+		FStuckNoteGuard StuckNoteGuard;
 	};
 }
diff --git a/Engine/Plugins/Runtime/Harmonix/Source/HarmonixMetasound/Public/HarmonixMetasound/MidiOps/MidiNoteFilter.h b/Engine/Plugins/Runtime/Harmonix/Source/HarmonixMetasound/Public/HarmonixMetasound/MidiOps/MidiNoteFilter.h
--- a/Engine/Plugins/Runtime/Harmonix/Source/HarmonixMetasound/Public/HarmonixMetasound/MidiOps/MidiNoteFilter.h
+++ b/Engine/Plugins/Runtime/Harmonix/Source/HarmonixMetasound/Public/HarmonixMetasound/MidiOps/MidiNoteFilter.h
@@ -2,6 +2,8 @@
 
 #pragma once
 
+#include "StuckNoteGuard.h"
+
 #include "HarmonixDsp/Parameters/Parameter.h"
 
 #include "HarmonixMetasound/DataTypes/MidiStream.h"
@@ -49,13 +51,6 @@ namespace Harmonix::Midi::Ops
 		FMidiDataParam MaxVelocity{ 0, 127, 127 };
 
 	private:
-		// Keep track of active notes so if the filter changes such that they would become stuck, we can unstick them immediately.
-		struct FActiveVoice
-		{
-			uint8 Channel;
-			uint8 NoteNumber;
-			uint8 Velocity;
-		};
-		TMap<FMidiVoiceId, FActiveVoice> ActiveVoices;
+		FStuckNoteGuard StuckNoteGuard;
 	};
 }
diff --git a/Engine/Plugins/Runtime/Harmonix/Source/HarmonixMetasound/Public/HarmonixMetasound/MidiOps/MidiOps.h b/Engine/Plugins/Runtime/Harmonix/Source/HarmonixMetasound/Public/HarmonixMetasound/MidiOps/MidiOps.h
--- a/Engine/Plugins/Runtime/Harmonix/Source/HarmonixMetasound/Public/HarmonixMetasound/MidiOps/MidiOps.h
+++ b/Engine/Plugins/Runtime/Harmonix/Source/HarmonixMetasound/Public/HarmonixMetasound/MidiOps/MidiOps.h
@@ -8,6 +8,8 @@ DECLARE_LOG_CATEGORY_EXTERN(LogHarmonixMidiOps, Log, All)
 
 namespace Harmonix::Midi::Ops
 {
+	using FMidiTrackIndex = uint16;
+	
 	constexpr uint16 AllChannelsOn = 0xffff;
 	constexpr uint16 AllChannelsOff = 0;
 	
diff --git a/Engine/Plugins/Runtime/Harmonix/Source/HarmonixMetasound/Public/HarmonixMetasound/MidiOps/MidiTrackFilter.h b/Engine/Plugins/Runtime/Harmonix/Source/HarmonixMetasound/Public/HarmonixMetasound/MidiOps/MidiTrackFilter.h
new file mode 100644
--- /dev/null
+++ b/Engine/Plugins/Runtime/Harmonix/Source/HarmonixMetasound/Public/HarmonixMetasound/MidiOps/MidiTrackFilter.h
@@ -0,0 +1,40 @@
+// Copyright Epic Games, Inc. All Rights Reserved.
+
+#pragma once
+
+#include "StuckNoteGuard.h"
+
+#include "HarmonixMetasound/DataTypes/MidiStream.h"
+
+namespace Harmonix::Midi::Ops
+{
+	/**
+	 * Filters MIDI events based on which track they're on
+	 */
+	class HARMONIXMETASOUND_API FMidiTrackFilter
+	{
+	public:
+		/**
+		 * Pass the filtered events from the input stream to the output stream
+		 * @param InStream The MIDI stream to filter
+		 * @param OutStream The filtered MIDI stream
+		 */
+		void Process(const HarmonixMetasound::FMidiStream& InStream, HarmonixMetasound::FMidiStream& OutStream);
+
+		/**
+		 * Set the range of tracks to include. Track 0 is reserved as the "conductor" track, which only contains timing information.
+		 * The first track with notes and other messages on it will be track 1. 
+		 * @param InMinTrackIdx The first track to include in the range of tracks (inclusive).
+		 * @param InMaxTrackIdx The last track to include in th range of tracks (inclusive).
+		 * @param InIncludeConductorTrack Whether or not to include the conductor track (track 0)
+		 */
+		void SetTrackRange(uint16 InMinTrackIdx, uint16 InMaxTrackIdx, bool InIncludeConductorTrack);
+
+	private:
+		uint16 MinTrackIdx = 0;
+		uint16 MaxTrackIdx = 0;
+		bool IncludeConductorTrack = false;
+		
+		FStuckNoteGuard StuckNoteGuard;
+	};
+}
diff --git a/Engine/Plugins/Runtime/Harmonix/Source/HarmonixMetasound/Public/HarmonixMetasound/MidiOps/StuckNoteGuard.h b/Engine/Plugins/Runtime/Harmonix/Source/HarmonixMetasound/Public/HarmonixMetasound/MidiOps/StuckNoteGuard.h
new file mode 100644
--- /dev/null
+++ b/Engine/Plugins/Runtime/Harmonix/Source/HarmonixMetasound/Public/HarmonixMetasound/MidiOps/StuckNoteGuard.h
@@ -0,0 +1,28 @@
+// Copyright Epic Games, Inc. All Rights Reserved.
+
+#pragma once
+
+#include "HarmonixMetasound/DataTypes/MidiStream.h"
+
+namespace Harmonix::Midi::Ops
+{
+	/**
+	 * Tracks incoming note messages and sends note offs when a note would become "stuck" because of a filter change
+	 */
+	class HARMONIXMETASOUND_API FStuckNoteGuard
+	{
+	public:
+		using FIncludeNotePredicate = TFunctionRef<bool(const HarmonixMetasound::FMidiStreamEvent& Event)>;
+
+		/**
+		 * Track incoming note messages and send note offs for notes that would have become "stuck" because the caller's filter changed.
+		 * @param InStream The input MIDI stream
+		 * @param OutStream The output MIDI stream
+		 * @param Predicate The predicate which indicates whether an event is still included in the caller's filter
+		 */
+		void Process(const HarmonixMetasound::FMidiStream& InStream, HarmonixMetasound::FMidiStream& OutStream, const FIncludeNotePredicate& Predicate);
+
+	private:
+		TArray<HarmonixMetasound::FMidiStreamEvent> ActiveNotes;
+	};
+}
diff --git a/Engine/Plugins/Runtime/Harmonix/Source/HarmonixMetasound/Public/HarmonixMetasound/Nodes/MidiStreamTrackFilterNode.h b/Engine/Plugins/Runtime/Harmonix/Source/HarmonixMetasound/Public/HarmonixMetasound/Nodes/MidiStreamTrackFilterNode.h
--- a/Engine/Plugins/Runtime/Harmonix/Source/HarmonixMetasound/Public/HarmonixMetasound/Nodes/MidiStreamTrackFilterNode.h
+++ b/Engine/Plugins/Runtime/Harmonix/Source/HarmonixMetasound/Public/HarmonixMetasound/Nodes/MidiStreamTrackFilterNode.h
@@ -6,7 +6,22 @@
 
 #include "HarmonixMetasound/Common.h"
 
-namespace HarmonixMetasound::Nodes::MidiStreamTrackFilter
+namespace HarmonixMetasound::Nodes::MidiTrackFilter
 {
-	const HARMONIXMETASOUND_API Metasound::FNodeClassName& GetClassName();
+	HARMONIXMETASOUND_API const Metasound::FNodeClassName& GetClassName();
+	HARMONIXMETASOUND_API int32 GetCurrentMajorVersion();
+
+	namespace Inputs
+	{
+		DECLARE_METASOUND_PARAM_EXTERN(Enable);
+		DECLARE_METASOUND_PARAM_EXTERN(MidiStream);
+		DECLARE_METASOUND_PARAM_EXTERN(MinTrackIndex);
+		DECLARE_METASOUND_PARAM_EXTERN(MaxTrackIndex);
+		DECLARE_METASOUND_PARAM_EXTERN(IncludeConductorTrack);
+	}
+
+	namespace Outputs
+	{
+		DECLARE_METASOUND_PARAM_EXTERN(MidiStream);
+	}
 }
diff --git a/Engine/Plugins/Runtime/Harmonix/Source/HarmonixMetasoundTests/Private/MidiOps/MidiTrackFilterTests.cpp b/Engine/Plugins/Runtime/Harmonix/Source/HarmonixMetasoundTests/Private/MidiOps/MidiTrackFilterTests.cpp
new file mode 100644
--- /dev/null
+++ b/Engine/Plugins/Runtime/Harmonix/Source/HarmonixMetasoundTests/Private/MidiOps/MidiTrackFilterTests.cpp
@@ -0,0 +1,101 @@
+// Copyright Epic Games, Inc. All Rights Reserved.
+
+#include "HarmonixMetasound/MidiOps/MidiTrackFilter.h"
+#include "Misc/AutomationTest.h"
+
+#if WITH_DEV_AUTOMATION_TESTS
+
+namespace Harmonix::Midi::Ops::Tests
+{
+	IMPLEMENT_SIMPLE_AUTOMATION_TEST(
+		FMidiTrackFilterProcessTest,
+		"Harmonix.Midi.Ops.MidiTrackFilter.Process",
+		EAutomationTestFlags::EditorContext | EAutomationTestFlags::EngineFilter)
+	bool FMidiTrackFilterProcessTest::RunTest(const FString&)
+	{
+		FMidiTrackFilter Filter;
+		
+		Metasound::FOperatorSettings OperatorSettings{ 48000, 100 };
+		HarmonixMetasound::FMidiStream InputStream{ OperatorSettings };
+		HarmonixMetasound::FMidiStream OutputStream{ OperatorSettings };
+
+		constexpr uint16 MinTrack = 1;
+		constexpr uint16 MaxTrack = 37;
+		constexpr uint8 Channel = 7;
+		constexpr uint8 Data1 = 78;
+		constexpr uint8 Data2 = 87;
+		constexpr float Tempo = 135.7f;
+		const int32 MidiTempo = MidiConstants::BPMToMidiTempo(Tempo);
+		constexpr uint8 TimeSigNumerator = 7;
+		constexpr uint8 TimeSigDenominator = 8;
+
+		const auto AddMidiMessageToTrack = [](HarmonixMetasound::FMidiStream& Stream, FMidiMsg&& Msg, uint16 TrackIndex)
+		{
+			HarmonixMetasound::FMidiStreamEvent Event{ static_cast<uint32>(0), Msg };
+			Event.TrackIndex = TrackIndex;
+			Stream.AddMidiEvent(Event);
+		};
+		
+		// Default: no events pass
+		{
+			// Add a tempo and time sig change
+			AddMidiMessageToTrack(InputStream, FMidiMsg{ MidiTempo }, 0);
+			AddMidiMessageToTrack(InputStream, FMidiMsg{ TimeSigNumerator, TimeSigDenominator }, 0);
+			// Add events to each track
+			for (uint16 TrackIdx = MinTrack; TrackIdx <= MaxTrack; ++TrackIdx)
+			{
+				AddMidiMessageToTrack(InputStream, FMidiMsg::CreateNoteOn(Channel, Data1, Data2), TrackIdx);
+				AddMidiMessageToTrack(InputStream, FMidiMsg::CreateControlChange(Channel, Data1, Data2), TrackIdx);
+			}
+
+			Filter.Process(InputStream, OutputStream);
+			UTEST_EQUAL("No events in output", OutputStream.GetEventsInBlock().Num(), 0);
+		}
+
+		// Include conductor track
+		{
+			Filter.SetTrackRange(0, 0, true);
+			Filter.Process(InputStream, OutputStream);
+			const TArray<HarmonixMetasound::FMidiStreamEvent>& Events = OutputStream.GetEventsInBlock();
+			UTEST_EQUAL("Only tempo and time sig events in output", Events.Num(), 2);
+			for (const HarmonixMetasound::FMidiStreamEvent& Event : Events)
+			{
+				if (Event.MidiMessage.Type == FMidiMsg::EType::Tempo)
+				{
+					UTEST_EQUAL("Tempo is correct", Event.MidiMessage.GetMicrosecPerQuarterNote(), MidiTempo);
+				}
+				else if (Event.MidiMessage.Type == FMidiMsg::EType::TimeSig)
+				{
+					UTEST_EQUAL("Time sig numerator is correct", Event.MidiMessage.GetTimeSigNumerator(), TimeSigNumerator);
+					UTEST_EQUAL("Time sig denominator is correct", Event.MidiMessage.GetTimeSigDenominator(), TimeSigDenominator);
+				}
+			}
+		}
+
+		// Include some tracks within the range
+		constexpr uint16 MinTrackToInclude = MinTrack + 3;
+		constexpr uint16 MaxTrackToInclude = MaxTrack - 4;
+		check(MaxTrackToInclude >= MinTrackToInclude);
+		constexpr uint16 NumTracksIncluded = (MaxTrackToInclude - MinTrackToInclude) + 1;
+		{
+			OutputStream.PrepareBlock();
+			Filter.SetTrackRange(MinTrackToInclude, MaxTrackToInclude, false);
+			Filter.Process(InputStream, OutputStream);
+			constexpr int32 NumEventsExpected = NumTracksIncluded * 2; // we include a note on and a CC above
+			UTEST_EQUAL("Correct number of events", OutputStream.GetEventsInBlock().Num(), NumEventsExpected);
+		}
+
+		// Change the filter and expect notes to become unstuck
+		{
+			InputStream.PrepareBlock();
+			OutputStream.PrepareBlock();
+			Filter.SetTrackRange(0, 0, false);
+			Filter.Process(InputStream, OutputStream);
+			UTEST_EQUAL("Correct number of note offs", OutputStream.GetEventsInBlock().Num(), NumTracksIncluded);
+		}
+		
+		return true;
+	}
+}
+
+#endif
\ No newline at end of file
diff --git a/Engine/Plugins/Runtime/Harmonix/Source/HarmonixMetasoundTests/Private/MidiOps/StuckNoteGuardTests.cpp b/Engine/Plugins/Runtime/Harmonix/Source/HarmonixMetasoundTests/Private/MidiOps/StuckNoteGuardTests.cpp
new file mode 100644
--- /dev/null
+++ b/Engine/Plugins/Runtime/Harmonix/Source/HarmonixMetasoundTests/Private/MidiOps/StuckNoteGuardTests.cpp
@@ -0,0 +1,87 @@
+// Copyright Epic Games, Inc. All Rights Reserved.
+
+#include "HarmonixMetasound/MidiOps/StuckNoteGuard.h"
+#include "Misc/AutomationTest.h"
+
+#if WITH_DEV_AUTOMATION_TESTS
+
+namespace Harmonix::Midi::Ops::Tests
+{
+	IMPLEMENT_SIMPLE_AUTOMATION_TEST(
+		FStuckNoteGuardBasicTest,
+		"Harmonix.Midi.Ops.StuckNoteGuard.Basic",
+		EAutomationTestFlags::EditorContext | EAutomationTestFlags::EngineFilter)
+	bool FStuckNoteGuardBasicTest::RunTest(const FString&)
+	{
+		FStuckNoteGuard Guard;
+
+		Metasound::FOperatorSettings OperatorSettings{ 48000, 100 };
+		HarmonixMetasound::FMidiStream InputStream{ OperatorSettings };
+		HarmonixMetasound::FMidiStream OutputStream{ OperatorSettings };
+
+		// Use a filter that only includes messages on channel 3
+		const auto Channel3Filter = [](const HarmonixMetasound::FMidiStreamEvent& Event)
+		{
+			return Event.MidiMessage.IsStd() && Event.MidiMessage.GetStdChannel() == 3;
+		};
+
+		constexpr uint8 NoteNumber = 43;
+		constexpr uint8 Velocity = 54;
+
+		// Add some note ons on all the channels
+		for (uint8 ChannelIdx = 0; ChannelIdx < MidiConstants::kNumChannels; ++ChannelIdx)
+		{
+			InputStream.AddMidiEvent({ static_cast<uint32>(0), FMidiMsg::CreateNoteOn(ChannelIdx, NoteNumber, Velocity)});
+		}
+
+		// Process and expect no messages (because we don't have any stuck notes yet)
+		Guard.Process(InputStream, OutputStream, Channel3Filter);
+		UTEST_EQUAL("Got no events", OutputStream.GetEventsInBlock().Num(), 0);
+
+		// Now change the filter to include only messages on channel 4, process, and expect to get a note off for the newly-stuck note on channel 3
+		const auto Channel4Filter = [](const HarmonixMetasound::FMidiStreamEvent& Event)
+		{
+			return Event.MidiMessage.IsStd() && Event.MidiMessage.GetStdChannel() == 4;
+		};
+		InputStream.PrepareBlock();
+		OutputStream.PrepareBlock();
+		Guard.Process(InputStream, OutputStream, Channel4Filter);
+		const auto& Events = OutputStream.GetEventsInBlock();
+		UTEST_EQUAL("Got one event", Events.Num(), 1);
+		UTEST_TRUE("Got a note off", Events[0].MidiMessage.IsNoteOff());
+		UTEST_EQUAL("Got the right note number", Events[0].MidiMessage.GetStdData1(), NoteNumber);
+		UTEST_EQUAL("Got the note off on the right channel", Events[0].MidiMessage.GetStdChannel(), 3);
+
+		// Add some note ons on all the channels
+		InputStream.PrepareBlock();
+		OutputStream.PrepareBlock();
+		for (uint8 ChannelIdx = 0; ChannelIdx < MidiConstants::kNumChannels; ++ChannelIdx)
+		{
+			InputStream.AddMidiEvent({ static_cast<uint32>(0), FMidiMsg::CreateNoteOn(ChannelIdx, NoteNumber, Velocity)});
+		}
+
+		// Process and expect no messages (because we don't have any stuck notes yet)
+		Guard.Process(InputStream, OutputStream, Channel4Filter);
+		UTEST_EQUAL("Got no events", OutputStream.GetEventsInBlock().Num(), 0);
+
+		// Now send note offs on all the channels, process, and expect no additional note offs
+		InputStream.PrepareBlock();
+		OutputStream.PrepareBlock();
+		for (uint8 ChannelIdx = 0; ChannelIdx < MidiConstants::kNumChannels; ++ChannelIdx)
+		{
+			InputStream.AddMidiEvent({ static_cast<uint32>(0), FMidiMsg::CreateNoteOff(ChannelIdx, NoteNumber)});
+		}
+		Guard.Process(InputStream, OutputStream, Channel4Filter);
+		UTEST_EQUAL("Got no events", OutputStream.GetEventsInBlock().Num(), 0);
+
+		// Change the filter back to channel 3, and expect no new note offs
+		InputStream.PrepareBlock();
+		OutputStream.PrepareBlock();
+		Guard.Process(InputStream, OutputStream, Channel3Filter);
+		UTEST_EQUAL("Got no events", OutputStream.GetEventsInBlock().Num(), 0);
+		
+		return true;
+	}
+}
+
+#endif
diff --git a/Engine/Plugins/Runtime/Harmonix/Source/HarmonixMetasoundTests/Private/Nodes/MidiStreamTrackFilterNode.cpp b/Engine/Plugins/Runtime/Harmonix/Source/HarmonixMetasoundTests/Private/Nodes/MidiStreamTrackFilterNode.cpp
--- a/Engine/Plugins/Runtime/Harmonix/Source/HarmonixMetasoundTests/Private/Nodes/MidiStreamTrackFilterNode.cpp
+++ b/Engine/Plugins/Runtime/Harmonix/Source/HarmonixMetasoundTests/Private/Nodes/MidiStreamTrackFilterNode.cpp
@@ -1,134 +1,83 @@
 // Copyright Epic Games, Inc. All Rights Reserved.
 
 #include "NodeTestGraphBuilder.h"
-#include "HarmonixDsp/AudioBuffer.h"
-#include "HarmonixMetasound/Common.h"
-#include "HarmonixMetasound/DataTypes/MidiAsset.h"
-#include "HarmonixMetasound/DataTypes/MidiStream.h"
+#include "HarmonixMetasound/MidiOps/MidiTrackFilter.h"
 #include "HarmonixMetasound/Nodes/MidiStreamTrackFilterNode.h"
-#include "HarmonixMidi/MidiMsg.h"
 #include "Misc/AutomationTest.h"
 
 #if WITH_DEV_AUTOMATION_TESTS
 
-namespace HarmonixMetasoundTests::MidiStreamTrackFilterNode
+namespace HarmonixMetasound::Nodes::MidiTrackFilter::Tests
 {
-	using GraphBuilder = Metasound::Test::FNodeTestGraphBuilder;
-	using namespace Metasound;
-	using namespace Metasound::Frontend;
-	using namespace HarmonixMetasound;
-
-	UMidiFile* BuildMidiFile()
-	{
-		// Make a midi file...
-		UMidiFile* TheMidi = NewObject<UMidiFile>();
-		// Set the initial tempo...
-		TheMidi->GetSongMaps()->GetTempoMap().AddTempoInfoPoint(MidiConstants::BPMToMidiTempo(120), 0);
-		// Make two tracks to test filtering
-		FMidiTrack* FirstTrack = TheMidi->AddTrack("FirstTrack");
-		FMidiTrack* SecondTrack = TheMidi->AddTrack("SecondTrack");
-		// Add midi events to the two tracks
-		FMidiMsg FirstMessage = FMidiMsg::CreateNoteOn(1, 1, 1);		
-		FMidiMsg SecondMessage = FMidiMsg::CreateNoteOn(1, 2, 2);
-		FMidiEvent FirstEvent(0, FirstMessage);
-		FMidiEvent SecondEvent(0, SecondMessage);		
-		FirstTrack->AddEvent(FirstEvent);
-		SecondTrack->AddEvent(SecondEvent);
-		// Tell the midi file its tracks have been changed so it can recalculate song length data...
-		TheMidi->TracksChanged();
-		// Now round appropriately...
-		TheMidi->ConformMidiFileLength(EMidiFileLengthConformOption::RoundUp);
-		return TheMidi;
-	}
-
 	IMPLEMENT_SIMPLE_AUTOMATION_TEST(
-		FMidiStreamTrackIndexFunctionalityTest,
-		"Harmonix.Metasound.Nodes.MidiStreamTrackFilterNode.TestFunctionality",
-		EAutomationTestFlags::EditorContext | EAutomationTestFlags::EngineFilter)
-		bool FMidiStreamTrackIndexFunctionalityTest::RunTest(const FString&)
+	FMidiTrackFilterNodeParityTest,
+	"Harmonix.Metasound.Nodes.MidiTrackFilter.Parity",
+	EAutomationTestFlags::EditorContext | EAutomationTestFlags::EngineFilter)
+	bool FMidiTrackFilterNodeParityTest::RunTest(const FString&)
 	{
-		//************************************************************************************************************************
-		// Build the graph.
-		//************************************************************************************************************************
-		GraphBuilder Builder;
-		Builder.AddOutput("AudioOut", GetMetasoundDataTypeName<FAudioBuffer>());
-
-		// transport...
-		Frontend::FNodeHandle TransportNode = Builder.AddNode({ HarmonixMetasound::HarmonixNodeNamespace, "TriggerToTransport", "" }, 0);
-		UTEST_TRUE("Transport Node Created", TransportNode.Get().IsValid());
-		// hoist up play trigger input...
-		Builder.AddAndConnectDataReferenceInput(TransportNode, CommonPinNames::Inputs::TransportPlayName, GetMetasoundDataTypeName<FTrigger>());
-
-		// midi player...
-		Frontend::FNodeHandle MidiPlayerNode = Builder.AddNode({ HarmonixMetasound::HarmonixNodeNamespace, "MidiPlayer", "" }, 0);
-		UTEST_TRUE("MidiPlayer Node Created", MidiPlayerNode.Get().IsValid());
-		// hoist up midi asset input...
-		Builder.AddAndConnectDataReferenceInput(MidiPlayerNode, CommonPinNames::Inputs::MidiFileAssetName, GetMetasoundDataTypeName<FMidiAsset>());
-
-		// wire transport to midi player...
-		bool ConnectionSuccess = Builder.ConnectNodes(TransportNode, CommonPinNames::Outputs::TransportName, MidiPlayerNode, CommonPinNames::Inputs::TransportName);
-		UTEST_TRUE("Transport Out connected to MidiPlayer node", ConnectionSuccess);
-
-		// Midi stream track index filter...
-		FNodeHandle MidiStreamTrackFilterNode = Builder.AddNode(HarmonixMetasound::Nodes::MidiStreamTrackFilter::GetClassName(), 0);
-		UTEST_TRUE("MidiStreamTrackFilter Node Created", MidiStreamTrackFilterNode.Get().IsValid());
-		// hoist up filter input...
-		Builder.AddAndConnectDataReferenceInput(MidiStreamTrackFilterNode, CommonPinNames::Inputs::MidiTrackIndexFilterSpecifierName, GetMetasoundDataTypeName<FString>());
-		// hoist up output trigger...
-		Builder.AddAndConnectDataReferenceOutput(MidiStreamTrackFilterNode, CommonPinNames::Outputs::MidiStreamName, GetMetasoundDataTypeName<FMidiStream>(), CommonPinNames::Outputs::MidiStreamName);
-
-		// wire midi player to midi stream track index filter...
-		ConnectionSuccess = Builder.ConnectNodes(MidiPlayerNode, CommonPinNames::Outputs::MidiStreamName, MidiStreamTrackFilterNode, CommonPinNames::Inputs::MidiStreamName);
-		UTEST_TRUE("MidiPlayer Out connected to MidiStreamTrackFilter node", ConnectionSuccess);
-
-		//************************************************************************************************************************
-		// Make the generator.
-		//************************************************************************************************************************
-		constexpr int32 NumSamplesPerBlock = 480; // 10ms per block
-		const TUniquePtr<FMetasoundGenerator> Generator = Builder.BuildGenerator(48000, NumSamplesPerBlock);
-		UTEST_TRUE("Graph successfully built", Generator.IsValid());
-
-		//************************************************************************************************************************
-		// Get reference to outputs we want to check.
-		//************************************************************************************************************************
-		TOptional<FMidiStreamReadRef> OutputMidiStream = Generator->GetOutputReadReference<FMidiStream>(CommonPinNames::Outputs::MidiStreamName);
-		UTEST_TRUE("Output exists", OutputMidiStream.IsSet());		
+		const TUniquePtr<Metasound::FMetasoundGenerator> Generator =
+			Metasound::Test::FNodeTestGraphBuilder::MakeSingleNodeGraph(GetClassName(), GetCurrentMajorVersion());
+		UTEST_TRUE("Built graph", Generator.IsValid());
+		TOptional<Metasound::TDataReadReference<FMidiStream>> NodeMidiOutput = Generator->GetOutputReadReference<FMidiStream>(Outputs::MidiStreamName);
+		UTEST_TRUE("Got node MIDI output", NodeMidiOutput.IsSet());
+		
+		Harmonix::Midi::Ops::FMidiTrackFilter FilterForComparison;
+		FMidiStream InputStream{ Generator->OperatorSettings };
+		FMidiStream OutputStream{ Generator->OperatorSettings };
+
+		FMidiStreamEvent NoteOnEvent{ static_cast<uint32>(0), FMidiMsg::CreateNoteOn(2, 43, 45) };
+		NoteOnEvent.TrackIndex = 3;
+		
+		// Default: no tracks included, no events output
+		{
+			// Add an event to the node input
+			Generator->ApplyToInputValue<FMidiStream>(Inputs::MidiStreamName, [&NoteOnEvent](FMidiStream& Stream)
+			{
+				Stream.AddMidiEvent(NoteOnEvent);
+			});
 
-		//************************************************************************************************************************
-		// Make initial inputs.
-		//************************************************************************************************************************
-		// Midi file to play...
-		UMidiFile* MidiFile = BuildMidiFile();
-		TSharedPtr<Audio::IProxyData> MidiFileProxy = MidiFile->CreateProxyData({ "MidiStreamTrackFilterTest" });
-		Generator->SetInputValue<FMidiAsset>(CommonPinNames::Inputs::MidiFileAssetName, MidiFileProxy);
+			// Process the node
+			Metasound::FAudioBuffer Buffer{ Generator->OperatorSettings };
+			Generator->OnGenerateAudio(Buffer.GetData(), Buffer.Num());
 
-		// Single Track Index to Filter
-		Generator->SetInputValue<FString>(CommonPinNames::Inputs::MidiTrackIndexFilterSpecifierName, "2");
+			// Process the filter
+			InputStream.AddMidiEvent(NoteOnEvent);
+			FilterForComparison.Process(InputStream, OutputStream);
 
-		// Trigger transport...
-		Generator->ApplyToInputValue<FTrigger>(CommonPinNames::Inputs::TransportPlayName, [](FTrigger& TriggerIn) { TriggerIn.TriggerFrame(0); });
+			// Make sure we got no events
+			UTEST_EQUAL("Got no events (test)", (*NodeMidiOutput)->GetEventsInBlock().Num(), 0);
+			UTEST_EQUAL("Got no events (control)", OutputStream.GetEventsInBlock().Num(), 0);
 
-		//************************************************************************************************************************
-		// Generate 10ms...
-		//************************************************************************************************************************
-		TAudioBuffer<float> Buffer{ Generator->GetNumChannels(), NumSamplesPerBlock, EAudioBufferCleanupMode::Delete };
-		Generator->OnGenerateAudio(Buffer.GetRawChannelData(0), Buffer.GetNumTotalValidSamples());
+			// Leave the inputs where they are
+		}
 
-		// There are 2 note-on messages in the midi stream, in two different tracks.
-		// Ensure that the MidiStreamTrackFilter node correctly filtered out one track.
-		int32 MidiMessageCount = 0;
-		for (const FMidiStreamEvent& MidiStreamEvent : (*OutputMidiStream)->GetEventsInBlock())
+		// Include the track for the note on and expect the message to come out the other end
 		{
-			if (MidiStreamEvent.MidiMessage.IsNoteOn())
-			{
-				++MidiMessageCount;
-			}
+			Generator->SetInputValue<int32>(Inputs::MinTrackIndexName, NoteOnEvent.TrackIndex);
+			Generator->SetInputValue<int32>(Inputs::MaxTrackIndexName, NoteOnEvent.TrackIndex);
+			FilterForComparison.SetTrackRange(NoteOnEvent.TrackIndex, NoteOnEvent.TrackIndex, false);
+
+			// Process the node
+			Metasound::FAudioBuffer Buffer{ Generator->OperatorSettings };
+			Generator->OnGenerateAudio(Buffer.GetData(), Buffer.Num());
+
+			// Process the filter
+			FilterForComparison.Process(InputStream, OutputStream);
+
+			const TArray<FMidiStreamEvent>& NodeEvents = (*NodeMidiOutput)->GetEventsInBlock();
+			UTEST_EQUAL("Got one event (test)", NodeEvents.Num(), 1);
+			UTEST_TRUE("Event is note on (test)", NodeEvents[0].MidiMessage.IsNoteOn());
+			const TArray<FMidiStreamEvent>& ControlEvents = OutputStream.GetEventsInBlock();
+			UTEST_EQUAL("Got one event (control)", ControlEvents.Num(), NodeEvents.Num());
+			UTEST_TRUE("Event is note on (control)", ControlEvents[0].MidiMessage.IsNoteOn());
+
+			UTEST_EQUAL("Same channel", NodeEvents[0].MidiMessage.GetStdChannel(), ControlEvents[0].MidiMessage.GetStdChannel());
+			UTEST_EQUAL("Same note number", NodeEvents[0].MidiMessage.GetStdData1(), ControlEvents[0].MidiMessage.GetStdData1());
+			UTEST_EQUAL("Same velocity", NodeEvents[0].MidiMessage.GetStdData2(), ControlEvents[0].MidiMessage.GetStdData2());
 		}
-
-		UTEST_EQUAL("Number of midi events matches filtered", MidiMessageCount, 1);
-
+		
 		return true;
 	}
 }
 
-#endif
+#endif
\ No newline at end of file
