[Harmonix] MIDI Note Transpose: remove filter in favor of composition

#tests automated tests

#changelist validated
#virtualized

diff --git a/Engine/Plugins/Runtime/Harmonix/Source/HarmonixMetasound/Private/HarmonixMetasound/Nodes/MidiStreamTransposerNode.cpp b/Engine/Plugins/Runtime/Harmonix/Source/HarmonixMetasound/Private/HarmonixMetasound/Nodes/MidiStreamTransposerNode.cpp
--- a/Engine/Plugins/Runtime/Harmonix/Source/HarmonixMetasound/Private/HarmonixMetasound/Nodes/MidiStreamTransposerNode.cpp
+++ b/Engine/Plugins/Runtime/Harmonix/Source/HarmonixMetasound/Private/HarmonixMetasound/Nodes/MidiStreamTransposerNode.cpp
@@ -1,5 +1,7 @@
 // Copyright Epic Games, Inc. All Rights Reserved.
 
+#include "HarmonixMetasound/Nodes/MidiStreamTransposerNode.h"
+
 #include "MetasoundExecutableOperator.h"
 #include "MetasoundFacade.h"
 #include "MetasoundNodeInterface.h"
@@ -16,18 +18,192 @@ DEFINE_LOG_CATEGORY_STATIC(LogMidiStreamTransposer, Log, All);
 
 #define LOCTEXT_NAMESPACE "HarmonixMetaSound"
 
-namespace HarmonixMetasound
+namespace HarmonixMetasound::Nodes::MidiNoteTranspose
 {
 	using namespace Metasound;
 
-	class FMidiStreamTransposerOperator : public TExecutableOperator<FMidiStreamTransposerOperator>
+	const FNodeClassName& GetClassName()
+	{
+		static const FNodeClassName ClassName{ HarmonixNodeNamespace, TEXT("MidiStreamTransposer"), TEXT("") };
+		return ClassName;
+	}
+
+	int32 GetCurrentMajorVersion()
+	{
+		return 1;
+	}
+
+	namespace Inputs
+	{
+		DEFINE_METASOUND_PARAM_ALIAS(Enable, CommonPinNames::Inputs::Enable);
+		DEFINE_METASOUND_PARAM_ALIAS(MidiStream, CommonPinNames::Inputs::MidiStream);
+		DEFINE_METASOUND_PARAM_ALIAS(Transposition, CommonPinNames::Inputs::Transposition);
+	}
+
+	namespace Outputs
+	{
+		DEFINE_METASOUND_PARAM_ALIAS(MidiStream, CommonPinNames::Outputs::MidiStream);
+	}
+
+	class FMidiNoteTransposeOperator final : public TExecutableOperator<FMidiNoteTransposeOperator>
+	{
+	public:
+		static const FNodeClassMetadata& GetNodeInfo()
+		{
+			auto InitNodeInfo = []() -> FNodeClassMetadata
+			{
+				FNodeClassMetadata Info;
+				Info.ClassName        = { HarmonixNodeNamespace, TEXT("MidiStreamTransposer"), TEXT("")};
+				Info.MajorVersion     = 1;
+				Info.MinorVersion     = 0;
+				Info.DisplayName      = METASOUND_LOCTEXT("MidiNoteTransposeNode_DisplayName", "MIDI Note Transpose");
+				Info.Description      = METASOUND_LOCTEXT("MidiNoteTransposeNode_Description", "Duplicates the incoming MIDI stream to its output with the note on/off messages transposed by the specified number of semitones.");
+				Info.Author           = PluginAuthor;
+				Info.PromptIfMissing  = PluginNodeMissingPrompt;
+				Info.DefaultInterface = GetVertexInterface();
+				Info.CategoryHierarchy.Emplace(NodeCategories::Music);
+				return Info;
+			};
+
+			static const FNodeClassMetadata Info = InitNodeInfo();
+
+			return Info;
+		}
+		
+		static const FVertexInterface& GetVertexInterface()
+		{
+			static const FVertexInterface Interface(
+				FInputVertexInterface(
+					TInputDataVertex<bool>(METASOUND_GET_PARAM_NAME_AND_METADATA(Inputs::Enable), true),
+					TInputDataVertex<FMidiStream>(METASOUND_GET_PARAM_NAME_AND_METADATA(Inputs::MidiStream)),
+					TInputDataVertex<int32>(METASOUND_GET_PARAM_NAME_AND_METADATA(Inputs::Transposition), 0)
+				),
+				FOutputVertexInterface(
+					TOutputDataVertex<FMidiStream>(METASOUND_GET_PARAM_NAME_AND_METADATA(Outputs::MidiStream))
+				)
+			);
+
+			return Interface;
+		}
+
+		struct FInputs
+		{
+			FBoolReadRef Enabled;
+			FMidiStreamReadRef MidiStream;
+			FInt32ReadRef Transposition;
+		};
+
+		struct FOutputs
+		{
+			FMidiStreamWriteRef MidiStream;
+		};
+
+		static TUniquePtr<IOperator> CreateOperator(const FBuildOperatorParams& InParams, FBuildResults& OutResults)
+		{
+			const FInputVertexInterfaceData& InputData = InParams.InputData;
+
+			FInputs Inputs
+			{
+				InputData.GetOrCreateDefaultDataReadReference<bool>(METASOUND_GET_PARAM_NAME(Inputs::Enable), InParams.OperatorSettings),
+				InputData.GetOrConstructDataReadReference<FMidiStream>(METASOUND_GET_PARAM_NAME(Inputs::MidiStream), InParams.OperatorSettings),
+				InputData.GetOrCreateDefaultDataReadReference<int32>(METASOUND_GET_PARAM_NAME(Inputs::Transposition), InParams.OperatorSettings)
+			};
+
+			FOutputs Outputs
+			{
+				FMidiStreamWriteRef::CreateNew(InParams.OperatorSettings)
+			};
+
+			return MakeUnique<FMidiNoteTransposeOperator>(InParams, MoveTemp(Inputs), MoveTemp(Outputs));
+		}
+
+		FMidiNoteTransposeOperator(const FBuildOperatorParams& InParams, FInputs&& InInputs, FOutputs&& InOutputs)
+			: Inputs(MoveTemp(InInputs))
+			, Outputs(MoveTemp(InOutputs))
+		{
+			Reset(InParams);
+		}
+
+		virtual void BindInputs(FInputVertexInterfaceData& InVertexData) override
+		{
+			InVertexData.BindReadVertex(METASOUND_GET_PARAM_NAME(Inputs::Enable), Inputs.Enabled);
+			InVertexData.BindReadVertex(METASOUND_GET_PARAM_NAME(Inputs::MidiStream), Inputs.MidiStream);
+			InVertexData.BindReadVertex(METASOUND_GET_PARAM_NAME(Inputs::Transposition), Inputs.Transposition);
+
+			ReassignOutputClock = true;
+		}
+
+		virtual void BindOutputs(FOutputVertexInterfaceData& InVertexData) override
+		{
+			InVertexData.BindReadVertex(METASOUND_GET_PARAM_NAME(Outputs::MidiStream), Outputs.MidiStream);
+
+			ReassignOutputClock = true;
+		}
+
+		void Reset(const FResetParams&)
+		{
+			ReassignOutputClock = true;
+		}
+
+		void Execute()
+		{
+			if (!*Inputs.Enabled)
+			{
+				return;
+			}
+			
+			if (ReassignOutputClock)
+			{
+				if (const FMidiClockReadRef* Clock = Inputs.MidiStream->GetMidiClockSource())
+				{
+					Outputs.MidiStream->SetClockSource(*Clock);
+				}
+
+				ReassignOutputClock = false;
+			}
+
+			const int32 Transposition = *Inputs.Transposition;
+
+			const auto Transformer = [Transposition](const FMidiStreamEvent& InEvent) -> FMidiStreamEvent
+			{
+				FMidiStreamEvent NewEvent = InEvent;
+				
+				if (NewEvent.MidiMessage.IsNoteOn() || NewEvent.MidiMessage.IsNoteOff())
+				{
+					NewEvent.MidiMessage.Data1 = FMath::Clamp(NewEvent.MidiMessage.Data1 + Transposition, 0, 127);
+				}
+				
+				return NewEvent;
+			};
+			
+			Outputs.MidiStream->Copy(Inputs.MidiStream, [](const FMidiStreamEvent&){return true;}, Transformer);
+		}
+		
+	private:
+		FInputs Inputs;
+		FOutputs Outputs;
+		bool ReassignOutputClock = true;
+	};
+
+	class FMidiNoteTransposeNode final : public FNodeFacade
+	{
+	public:
+		explicit FMidiNoteTransposeNode(const FNodeInitData& InInitData)
+			: FNodeFacade(InInitData.InstanceName, InInitData.InstanceID, TFacadeOperatorClass<FMidiNoteTransposeOperator>())
+		{}
+		virtual ~FMidiNoteTransposeNode() override = default;
+	};
+
+	METASOUND_REGISTER_NODE(FMidiNoteTransposeNode)
+
+	class FMidiStreamTransposerOperator_V0 : public TExecutableOperator<FMidiStreamTransposerOperator_V0>
 	{
 	public:
 		static const FNodeClassMetadata& GetNodeInfo();
 		static const FVertexInterface& GetVertexInterface();
 		static TUniquePtr<IOperator> CreateOperator(const FBuildOperatorParams& InParams, FBuildResults& OutResults);
 
-		FMidiStreamTransposerOperator(const FBuildOperatorParams& InParams,
+		FMidiStreamTransposerOperator_V0(const FBuildOperatorParams& InParams,
 								  const FBoolReadRef&       InEnabled,
 								  const FStringReadRef&     InTrackSelect,
 								  const FStringReadRef&     InChannelSelect,
@@ -36,8 +212,6 @@ namespace HarmonixMetasound
 
 		virtual void BindInputs(FInputVertexInterfaceData& InVertexData) override;
 		virtual void BindOutputs(FOutputVertexInterfaceData& InVertexData) override;
-		virtual FDataReferenceCollection GetInputs() const override;
-		virtual FDataReferenceCollection GetOutputs() const override;
 
 		void Reset(const FResetParams& ResetParams);
 
@@ -60,18 +234,18 @@ namespace HarmonixMetasound
 		FMidiStreamEventTrackChannelFilter Filter;
 	};
 
-	class FMidiStreamTransposerNode : public FNodeFacade
+	class FMidiStreamTransposerNode_v0 : public FNodeFacade
 	{
 	public:
-		FMidiStreamTransposerNode(const FNodeInitData& InInitData)
-			: FNodeFacade(InInitData.InstanceName, InInitData.InstanceID, TFacadeOperatorClass<FMidiStreamTransposerOperator>())
+		FMidiStreamTransposerNode_v0(const FNodeInitData& InInitData)
+			: FNodeFacade(InInitData.InstanceName, InInitData.InstanceID, TFacadeOperatorClass<FMidiStreamTransposerOperator_V0>())
 		{}
-		virtual ~FMidiStreamTransposerNode() = default;
+		virtual ~FMidiStreamTransposerNode_v0() = default;
 	};
 
-	METASOUND_REGISTER_NODE(FMidiStreamTransposerNode)
+	METASOUND_REGISTER_NODE(FMidiStreamTransposerNode_v0)
 		
-	const FNodeClassMetadata& FMidiStreamTransposerOperator::GetNodeInfo()
+	const FNodeClassMetadata& FMidiStreamTransposerOperator_V0::GetNodeInfo()
 	{
 		auto InitNodeInfo = []() -> FNodeClassMetadata
 		{
@@ -85,6 +259,7 @@ namespace HarmonixMetasound
 			Info.PromptIfMissing  = PluginNodeMissingPrompt;
 			Info.DefaultInterface = GetVertexInterface();
 			Info.CategoryHierarchy.Emplace(NodeCategories::Music);
+			Info.bDeprecated = true;
 			return Info;
 		};
 
@@ -93,15 +268,13 @@ namespace HarmonixMetasound
 		return Info;
 	}
 
-	const FVertexInterface& FMidiStreamTransposerOperator::GetVertexInterface()
+	const FVertexInterface& FMidiStreamTransposerOperator_V0::GetVertexInterface()
 	{
-		using namespace CommonPinNames;
-
 		static const FVertexInterface Interface(
 			FInputVertexInterface(
 				TInputDataVertex<bool>(METASOUND_GET_PARAM_NAME_AND_METADATA(Inputs::Enable), true),
-				TInputDataVertex<FString>(METASOUND_GET_PARAM_NAME_AND_METADATA(Inputs::MidiTrackIndexFilterSpecifier), FString(TEXT("*"))),
-				TInputDataVertex<FString>(METASOUND_GET_PARAM_NAME_AND_METADATA(Inputs::MidiChannelFilterSpecifier), FString(TEXT("*"))),
+				TInputDataVertex<FString>(METASOUND_GET_PARAM_NAME_AND_METADATA(CommonPinNames::Inputs::MidiTrackIndexFilterSpecifier), FString(TEXT("*"))),
+				TInputDataVertex<FString>(METASOUND_GET_PARAM_NAME_AND_METADATA(CommonPinNames::Inputs::MidiChannelFilterSpecifier), FString(TEXT("*"))),
 				TInputDataVertex<int32>(METASOUND_GET_PARAM_NAME_AND_METADATA(Inputs::Transposition), 0),
 				TInputDataVertex<FMidiStream>(METASOUND_GET_PARAM_NAME_AND_METADATA(Inputs::MidiStream))
 			),
@@ -113,20 +286,18 @@ namespace HarmonixMetasound
 		return Interface;
 	}
 
-	TUniquePtr<IOperator> FMidiStreamTransposerOperator::CreateOperator(const FBuildOperatorParams& InParams, FBuildResults& OutResults)
+	TUniquePtr<IOperator> FMidiStreamTransposerOperator_V0::CreateOperator(const FBuildOperatorParams& InParams, FBuildResults& OutResults)
 	{
-		using namespace CommonPinNames;
-
-		const FMidiStreamTransposerNode& LoggerNode = static_cast<const FMidiStreamTransposerNode&>(InParams.Node);
+		const FMidiStreamTransposerNode_v0& LoggerNode = static_cast<const FMidiStreamTransposerNode_v0&>(InParams.Node);
 
 		const FInputVertexInterfaceData& InputData = InParams.InputData;
 		FBoolReadRef InEnabled = InputData.GetOrCreateDefaultDataReadReference<bool>(METASOUND_GET_PARAM_NAME(Inputs::Enable), InParams.OperatorSettings);
-		FStringReadRef InTrackFilter = InputData.GetOrCreateDefaultDataReadReference<FString>(METASOUND_GET_PARAM_NAME(Inputs::MidiTrackIndexFilterSpecifier), InParams.OperatorSettings);
-		FStringReadRef InMidiChannelFilter = InputData.GetOrCreateDefaultDataReadReference<FString>(METASOUND_GET_PARAM_NAME(Inputs::MidiChannelFilterSpecifier), InParams.OperatorSettings);
+		FStringReadRef InTrackFilter = InputData.GetOrCreateDefaultDataReadReference<FString>(CommonPinNames::METASOUND_GET_PARAM_NAME(Inputs::MidiTrackIndexFilterSpecifier), InParams.OperatorSettings);
+		FStringReadRef InMidiChannelFilter = InputData.GetOrCreateDefaultDataReadReference<FString>(CommonPinNames::METASOUND_GET_PARAM_NAME(Inputs::MidiChannelFilterSpecifier), InParams.OperatorSettings);
 		FInt32ReadRef InTransposition = InputData.GetOrCreateDefaultDataReadReference<int32>(METASOUND_GET_PARAM_NAME(Inputs::Transposition), InParams.OperatorSettings);
 		FMidiStreamReadRef InMidiStream = InputData.GetOrConstructDataReadReference<FMidiStream>(METASOUND_GET_PARAM_NAME(Inputs::MidiStream), InParams.OperatorSettings);
 
-		return MakeUnique<FMidiStreamTransposerOperator>(InParams, 
+		return MakeUnique<FMidiStreamTransposerOperator_V0>(InParams, 
 			InEnabled,
 			InTrackFilter,
 			InMidiChannelFilter,
@@ -134,7 +305,7 @@ namespace HarmonixMetasound
 			InMidiStream);
 	}
 
-	FMidiStreamTransposerOperator::FMidiStreamTransposerOperator(const FBuildOperatorParams& InParams,
+	FMidiStreamTransposerOperator_V0::FMidiStreamTransposerOperator_V0(const FBuildOperatorParams& InParams,
 											   const FBoolReadRef& InEnabled,
 											   const FStringReadRef& InTrackSelect,
 											   const FStringReadRef& InChannelSelect,
@@ -150,37 +321,23 @@ namespace HarmonixMetasound
 		Reset(InParams);
 	}
 
-	void FMidiStreamTransposerOperator::BindInputs(FInputVertexInterfaceData& InVertexData)
+	void FMidiStreamTransposerOperator_V0::BindInputs(FInputVertexInterfaceData& InVertexData)
 	{
 		using namespace CommonPinNames;
 
 		InVertexData.BindReadVertex(METASOUND_GET_PARAM_NAME(Inputs::Enable), EnableInPin);
-		InVertexData.BindReadVertex(METASOUND_GET_PARAM_NAME(Inputs::MidiTrackIndexFilterSpecifier), TrackSelectInPin);
-		InVertexData.BindReadVertex(METASOUND_GET_PARAM_NAME(Inputs::MidiChannelFilterSpecifier), ChannelSelectInPin);
+		InVertexData.BindReadVertex(METASOUND_GET_PARAM_NAME(CommonPinNames::Inputs::MidiTrackIndexFilterSpecifier), TrackSelectInPin);
+		InVertexData.BindReadVertex(METASOUND_GET_PARAM_NAME(CommonPinNames::Inputs::MidiChannelFilterSpecifier), ChannelSelectInPin);
 		InVertexData.BindReadVertex(METASOUND_GET_PARAM_NAME(Inputs::Transposition), TranspositionInPin);
 		InVertexData.BindReadVertex(METASOUND_GET_PARAM_NAME(Inputs::MidiStream), MidiStreamInPin);
 	}
 
-	void FMidiStreamTransposerOperator::BindOutputs(FOutputVertexInterfaceData& InVertexData)
+	void FMidiStreamTransposerOperator_V0::BindOutputs(FOutputVertexInterfaceData& InVertexData)
 	{
-		using namespace CommonPinNames;
-
 		InVertexData.BindReadVertex(METASOUND_GET_PARAM_NAME(Outputs::MidiStream), MidiStreamOutPin);
 	}
 
-	FDataReferenceCollection FMidiStreamTransposerOperator::GetInputs() const
-	{
-		checkNoEntry();
-		return {};
-	}
-
-	FDataReferenceCollection FMidiStreamTransposerOperator::GetOutputs() const
-	{
-		checkNoEntry();
-		return {};
-	}
-
-	void FMidiStreamTransposerOperator::Reset(const FResetParams& ResetParams)
+	void FMidiStreamTransposerOperator_V0::Reset(const FResetParams& ResetParams)
 	{
 		MidiStreamOutPin->PrepareBlock();
 		
@@ -204,7 +361,7 @@ namespace HarmonixMetasound
 		}
 	}
 
-	void FMidiStreamTransposerOperator::Execute()
+	void FMidiStreamTransposerOperator_V0::Execute()
 	{
 		if (CurrentTrackFilter != *TrackSelectInPin)
 		{
diff --git a/Engine/Plugins/Runtime/Harmonix/Source/HarmonixMetasound/Public/HarmonixMetasound/Nodes/MidiStreamTransposerNode.h b/Engine/Plugins/Runtime/Harmonix/Source/HarmonixMetasound/Public/HarmonixMetasound/Nodes/MidiStreamTransposerNode.h
new file mode 100644
--- /dev/null
+++ b/Engine/Plugins/Runtime/Harmonix/Source/HarmonixMetasound/Public/HarmonixMetasound/Nodes/MidiStreamTransposerNode.h
@@ -0,0 +1,25 @@
+// Copyright Epic Games, Inc. All Rights Reserved.
+
+#pragma once
+
+#include "MetasoundNodeInterface.h"
+
+#include "HarmonixMetasound/Common.h"
+
+namespace HarmonixMetasound::Nodes::MidiNoteTranspose
+{
+	HARMONIXMETASOUND_API const Metasound::FNodeClassName& GetClassName();
+	HARMONIXMETASOUND_API int32 GetCurrentMajorVersion();
+
+	namespace Inputs
+	{
+		DECLARE_METASOUND_PARAM_EXTERN(Enable);
+		DECLARE_METASOUND_PARAM_EXTERN(MidiStream);
+		DECLARE_METASOUND_PARAM_EXTERN(Transposition);
+	}
+
+	namespace Outputs
+	{
+		DECLARE_METASOUND_PARAM_EXTERN(MidiStream);
+	}
+}
diff --git a/Engine/Plugins/Runtime/Harmonix/Source/HarmonixMetasoundTests/Private/Nodes/MidiStreamTransposerNode.cpp b/Engine/Plugins/Runtime/Harmonix/Source/HarmonixMetasoundTests/Private/Nodes/MidiStreamTransposerNode.cpp
--- a/Engine/Plugins/Runtime/Harmonix/Source/HarmonixMetasoundTests/Private/Nodes/MidiStreamTransposerNode.cpp
+++ b/Engine/Plugins/Runtime/Harmonix/Source/HarmonixMetasoundTests/Private/Nodes/MidiStreamTransposerNode.cpp
@@ -1,13 +1,13 @@
 // Copyright Epic Games, Inc. All Rights Reserved.
 
 #include "NodeTestGraphBuilder.h"
-#include "HarmonixDsp/AudioBuffer.h"
-#include "HarmonixMetasound/Common.h"
+#include "HarmonixMetasound/DataTypes/MidiStream.h"
+#include "HarmonixMetasound/Nodes/MidiStreamTransposerNode.h"
 #include "Misc/AutomationTest.h"
 
 #if WITH_DEV_AUTOMATION_TESTS
 
-namespace HarmonixMetasoundTests::MidiStreamTransposerNode
+namespace HarmonixMetasoundTests::MidiNoteTransposeNode
 {
 	using GraphBuilder = Metasound::Test::FNodeTestGraphBuilder;
 	using namespace Metasound;
@@ -16,26 +16,91 @@ namespace HarmonixMetasoundTests::MidiStreamTransposerNode
 
 	IMPLEMENT_SIMPLE_AUTOMATION_TEST(
 		FMidiStreamTransposerCreateNodeTest,
-		"Harmonix.Metasound.Nodes.MidiStreamTransposerNode.CreateNode",
+		"Harmonix.Metasound.Nodes.MidiNoteTransposeNode.Basic",
 		EAutomationTestFlags::EditorContext | EAutomationTestFlags::EngineFilter)
 		bool FMidiStreamTransposerCreateNodeTest::RunTest(const FString&)
 	{
 		// Build the graph.
 		constexpr int32 NumSamplesPerBlock = 256;
 		const TUniquePtr<FMetasoundGenerator> Generator = GraphBuilder::MakeSingleNodeGraph(
-			{ HarmonixMetasound::HarmonixNodeNamespace, "MidiStreamTransposer", "" },
-			0,
+			Nodes::MidiNoteTranspose::GetClassName(),
+			Nodes::MidiNoteTranspose::GetCurrentMajorVersion(),
 			48000,
 			NumSamplesPerBlock);
 		UTEST_TRUE("Graph successfully built", Generator.IsValid());
 
-		// execute a block
+		TOptional<TDataReadReference<FMidiStream>> NodeMidiOutput =
+			Generator->GetOutputReadReference<FMidiStream>(Nodes::MidiNoteTranspose::Outputs::MidiStreamName);
+		UTEST_TRUE("Got node MIDI output", NodeMidiOutput.IsSet());
+		
+		FAudioBuffer Buffer{ Generator->OperatorSettings };
+		
+		const TArray<uint8> NoteNumbers{ 3, 23, 57, 87, 99, 120 };
+
+		// Default: no transposition
 		{
-			TAudioBuffer<float> Buffer{ Generator->GetNumChannels(), NumSamplesPerBlock, EAudioBufferCleanupMode::Delete};
-			Generator->OnGenerateAudio(Buffer.GetRawChannelData(0), Buffer.GetNumTotalValidSamples());
+			// Add the notes
+			Generator->ApplyToInputValue<FMidiStream>(
+				Nodes::MidiNoteTranspose::Inputs::MidiStreamName,
+				[&NoteNumbers](FMidiStream& Stream)
+				{
+					for (const uint8 NoteNumber : NoteNumbers)
+					{
+						Stream.AddMidiEvent({ static_cast<uint32>(0), FMidiMsg::CreateNoteOn(7, NoteNumber, 127) });
+					}
+				});
+
+			// Process
+			Generator->OnGenerateAudio(Buffer.GetData(), Buffer.Num());
+
+			// The notes should be unchanged in the output
+			const TArray<FMidiStreamEvent>& Events = (*NodeMidiOutput)->GetEventsInBlock();
+			UTEST_EQUAL("Got the right number of events", Events.Num(), NoteNumbers.Num());
+			const TArray<uint8> ExpectedNoteNumbers = NoteNumbers;
+			
+			for (const FMidiStreamEvent& Event : Events)
+			{
+				if (Event.MidiMessage.IsNoteOn())
+				{
+					UTEST_TRUE("Note number unchanged", ExpectedNoteNumbers.Contains(Event.MidiMessage.GetStdData1()));
+				}
+				else
+				{
+					UTEST_TRUE("Unexpected event", false);
+				}
+			}
 		}
 
-		// No output to validate, finish test
+		// Try some different transposition values
+		const TArray<int32> TranspositionAmounts{ -23, -3, 0, 4, 12, 40, 1000 };
+		for (const int32 TranspositionAmount : TranspositionAmounts)
+		{
+			// Set transposition
+			Generator->SetInputValue<int32>(Nodes::MidiNoteTranspose::Inputs::TranspositionName, TranspositionAmount);
+			// Process
+			Generator->OnGenerateAudio(Buffer.GetData(), Buffer.Num());
+
+			// The notes should be transposed, but clamped to MIDI note range (0-127)
+			const TArray<FMidiStreamEvent>& Events = (*NodeMidiOutput)->GetEventsInBlock();
+			UTEST_EQUAL("Got the right number of events", Events.Num(), NoteNumbers.Num());
+			TArray<uint8> ExpectedNoteNumbers;
+			for (const uint8 OriginalNoteNumber : NoteNumbers)
+			{
+				ExpectedNoteNumbers.Add(FMath::Clamp(OriginalNoteNumber + TranspositionAmount, 0, 127));
+			}
+			
+			for (const FMidiStreamEvent& Event : Events)
+			{
+				if (Event.MidiMessage.IsNoteOn())
+				{
+					UTEST_TRUE("Note number unchanged", ExpectedNoteNumbers.Contains(Event.MidiMessage.GetStdData1()));
+				}
+				else
+				{
+					UTEST_TRUE("Unexpected event", false);
+				}
+			}
+		}
 
 		return true;
 	}
