[Harmonix] Remove FMidiClock::WriteNoAdvance. It was a no-op/assert anyway.

#tests automated tests
#rnx

diff --git a/Engine/Plugins/Runtime/Harmonix/Source/HarmonixMetasound/Private/HarmonixMetasound/DataTypes/MidiClock.cpp b/Engine/Plugins/Runtime/Harmonix/Source/HarmonixMetasound/Private/HarmonixMetasound/DataTypes/MidiClock.cpp
--- a/Engine/Plugins/Runtime/Harmonix/Source/HarmonixMetasound/Private/HarmonixMetasound/DataTypes/MidiClock.cpp
+++ b/Engine/Plugins/Runtime/Harmonix/Source/HarmonixMetasound/Private/HarmonixMetasound/DataTypes/MidiClock.cpp
@@ -486,12 +486,6 @@ namespace HarmonixMetasound
 		FramesUntilNextProcess -= FramesToProcess;
 	}
 
-	void FMidiClock::WriteNoAdvance(int32 StartFrameIndex, int32 EndFrameIndex)
-	{
-		EMusicPlayerTransportState ClockTransport = GetTransportStateAtBlockSampleFrame(StartFrameIndex);
-		check(ClockTransport != EMusicPlayerTransportState::Playing);
-	}
-
 	void FMidiClock::SeekTo(int32 BlockFrameIndex, const FMusicSeekTarget& InTarget, int32 InPrerollBars)
 	{
 		CurrentBlockFrameIndex = BlockFrameIndex;
diff --git a/Engine/Plugins/Runtime/Harmonix/Source/HarmonixMetasound/Private/HarmonixMetasound/Nodes/MetronomeNode.cpp b/Engine/Plugins/Runtime/Harmonix/Source/HarmonixMetasound/Private/HarmonixMetasound/Nodes/MetronomeNode.cpp
--- a/Engine/Plugins/Runtime/Harmonix/Source/HarmonixMetasound/Private/HarmonixMetasound/Nodes/MetronomeNode.cpp
+++ b/Engine/Plugins/Runtime/Harmonix/Source/HarmonixMetasound/Private/HarmonixMetasound/Nodes/MetronomeNode.cpp
@@ -420,11 +420,9 @@ namespace HarmonixMetasound
 			case EMusicPlayerTransportState::Invalid:
 			case EMusicPlayerTransportState::Preparing:
 				AddTransportStateChangeToMidiClock(StartFrameIndex, EMusicPlayerTransportState::Prepared);
-				DrivingMidiClock.WriteNoAdvance(StartFrameIndex, EndFrameIndex);
 				return EMusicPlayerTransportState::Prepared;
 
 			case EMusicPlayerTransportState::Prepared:
-				DrivingMidiClock.WriteNoAdvance(StartFrameIndex, EndFrameIndex);
 				return CurrentState;
 
 			case EMusicPlayerTransportState::Starting:
@@ -458,21 +456,17 @@ namespace HarmonixMetasound
 
 			case EMusicPlayerTransportState::Pausing:
 				AddTransportStateChangeToMidiClock(StartFrameIndex, EMusicPlayerTransportState::Paused);
-				DrivingMidiClock.WriteNoAdvance(StartFrameIndex, EndFrameIndex);
 				return EMusicPlayerTransportState::Paused;
 
 			case EMusicPlayerTransportState::Paused:
-				DrivingMidiClock.WriteNoAdvance(StartFrameIndex, EndFrameIndex);
 				return EMusicPlayerTransportState::Paused;
 
 			case EMusicPlayerTransportState::Stopping:
 				AddTransportStateChangeToMidiClock(StartFrameIndex, EMusicPlayerTransportState::Prepared);
-				DrivingMidiClock.WriteNoAdvance(StartFrameIndex, EndFrameIndex);
 				return EMusicPlayerTransportState::Prepared;
 
 			case EMusicPlayerTransportState::Killing:
 				AddTransportStateChangeToMidiClock(StartFrameIndex, EMusicPlayerTransportState::Prepared);
-				DrivingMidiClock.WriteNoAdvance(StartFrameIndex, EndFrameIndex);
 				return EMusicPlayerTransportState::Prepared;
 
 			default:
diff --git a/Engine/Plugins/Runtime/Harmonix/Source/HarmonixMetasound/Private/HarmonixMetasound/Nodes/MidiPlayerNode.cpp b/Engine/Plugins/Runtime/Harmonix/Source/HarmonixMetasound/Private/HarmonixMetasound/Nodes/MidiPlayerNode.cpp
--- a/Engine/Plugins/Runtime/Harmonix/Source/HarmonixMetasound/Private/HarmonixMetasound/Nodes/MidiPlayerNode.cpp
+++ b/Engine/Plugins/Runtime/Harmonix/Source/HarmonixMetasound/Private/HarmonixMetasound/Nodes/MidiPlayerNode.cpp
@@ -579,16 +579,12 @@ namespace HarmonixMetasound
 			case EMusicPlayerTransportState::Preparing:
 				// midi clock out
 				MidiClockOut->AddTransportStateChangeToBlock({ StartFrameIndex, 0.0f, EMusicPlayerTransportState::Prepared });
-				MidiClockOut->WriteNoAdvance(StartFrameIndex, EndFrameIndex);
 
 				// midi out
 				MidiOutPin->AddTransportStateChangeMessage(StartFrameIndex, EMusicPlayerTransportState::Prepared);
 				return EMusicPlayerTransportState::Prepared;
 
 			case EMusicPlayerTransportState::Prepared:
-				// midi clock out
-				MidiClockOut->WriteNoAdvance(StartFrameIndex, EndFrameIndex);
-
 				// midi out
 				MidiOutPin->AddTransportStateChangeMessage(StartFrameIndex, CurrentState);
 				return CurrentState;
@@ -622,22 +618,17 @@ namespace HarmonixMetasound
 			case EMusicPlayerTransportState::Pausing:
 				// midi clock out
 				MidiClockOut->AddTransportStateChangeToBlock({ StartFrameIndex, 0.0f, EMusicPlayerTransportState::Paused });
-				MidiClockOut->WriteNoAdvance(StartFrameIndex, EndFrameIndex);
 
 				// midi out
 				return EMusicPlayerTransportState::Paused;
 
 			case EMusicPlayerTransportState::Paused:
-				// midi clock out
-				MidiClockOut->WriteNoAdvance(StartFrameIndex, EndFrameIndex);
-
 				// midi out
 				return EMusicPlayerTransportState::Paused;
 
 			case EMusicPlayerTransportState::Stopping:
 				// midi clock out
 				MidiClockOut->AddTransportStateChangeToBlock({ StartFrameIndex, 0.0f, EMusicPlayerTransportState::Prepared });
-				MidiClockOut->WriteNoAdvance(StartFrameIndex, EndFrameIndex);
 
 				// midi out
 				MidiOutPin->AddTransportStateChangeMessage(StartFrameIndex, EMusicPlayerTransportState::Stopping);
@@ -646,7 +637,6 @@ namespace HarmonixMetasound
 			case EMusicPlayerTransportState::Killing:
 				// midi clock out
 				MidiClockOut->AddTransportStateChangeToBlock({ StartFrameIndex, 0.0f, EMusicPlayerTransportState::Prepared });
-				MidiClockOut->WriteNoAdvance(StartFrameIndex, EndFrameIndex);
 
 				// midi out
 				MidiOutPin->AddTransportStateChangeMessage(StartFrameIndex, EMusicPlayerTransportState::Killing);
diff --git a/Engine/Plugins/Runtime/Harmonix/Source/HarmonixMetasound/Public/HarmonixMetasound/DataTypes/MidiClock.h b/Engine/Plugins/Runtime/Harmonix/Source/HarmonixMetasound/Public/HarmonixMetasound/DataTypes/MidiClock.h
--- a/Engine/Plugins/Runtime/Harmonix/Source/HarmonixMetasound/Public/HarmonixMetasound/DataTypes/MidiClock.h
+++ b/Engine/Plugins/Runtime/Harmonix/Source/HarmonixMetasound/Public/HarmonixMetasound/DataTypes/MidiClock.h
@@ -203,7 +203,6 @@ namespace HarmonixMetasound
 		//*****************************************************************************************
 
 		void WriteAdvance(int32 StartFrameIndex, int32 EndFrameIndex, float InSpeed = 1.0f);
-		void WriteNoAdvance(int32 StartFrameIndex, int32 EndFrameIndex);
 		void SeekTo(int32 BlockFrameIndex, const FMusicSeekTarget& InTarget, int32 InPrerollBars);
 		// copy speed and tempo changes from in clock to this clock
 		// with an optional Speed multiplier to adjust the out going speed on this clock
