[Harmonix] Clean up copy/paste error

#rnx
#tests automated tests

diff --git a/Engine/Plugins/Runtime/Harmonix/Source/HarmonixMetasound/Private/HarmonixMetasound/Nodes/MidiStreamTransposerNode.cpp b/Engine/Plugins/Runtime/Harmonix/Source/HarmonixMetasound/Private/HarmonixMetasound/Nodes/MidiStreamTransposerNode.cpp
--- a/Engine/Plugins/Runtime/Harmonix/Source/HarmonixMetasound/Private/HarmonixMetasound/Nodes/MidiStreamTransposerNode.cpp
+++ b/Engine/Plugins/Runtime/Harmonix/Source/HarmonixMetasound/Private/HarmonixMetasound/Nodes/MidiStreamTransposerNode.cpp
@@ -292,8 +292,8 @@ namespace HarmonixMetasound::Nodes::MidiNoteTranspose
 
 		const FInputVertexInterfaceData& InputData = InParams.InputData;
 		FBoolReadRef InEnabled = InputData.GetOrCreateDefaultDataReadReference<bool>(METASOUND_GET_PARAM_NAME(Inputs::Enable), InParams.OperatorSettings);
-		FStringReadRef InTrackFilter = InputData.GetOrCreateDefaultDataReadReference<FString>(CommonPinNames::METASOUND_GET_PARAM_NAME(Inputs::MidiTrackIndexFilterSpecifier), InParams.OperatorSettings);
-		FStringReadRef InMidiChannelFilter = InputData.GetOrCreateDefaultDataReadReference<FString>(CommonPinNames::METASOUND_GET_PARAM_NAME(Inputs::MidiChannelFilterSpecifier), InParams.OperatorSettings);
+		FStringReadRef InTrackFilter = InputData.GetOrCreateDefaultDataReadReference<FString>(METASOUND_GET_PARAM_NAME(CommonPinNames::Inputs::MidiTrackIndexFilterSpecifier), InParams.OperatorSettings);
+		FStringReadRef InMidiChannelFilter = InputData.GetOrCreateDefaultDataReadReference<FString>(METASOUND_GET_PARAM_NAME(CommonPinNames::Inputs::MidiChannelFilterSpecifier), InParams.OperatorSettings);
 		FInt32ReadRef InTransposition = InputData.GetOrCreateDefaultDataReadReference<int32>(METASOUND_GET_PARAM_NAME(Inputs::Transposition), InParams.OperatorSettings);
 		FMidiStreamReadRef InMidiStream = InputData.GetOrConstructDataReadReference<FMidiStream>(METASOUND_GET_PARAM_NAME(Inputs::MidiStream), InParams.OperatorSettings);
 
@@ -323,8 +323,6 @@ namespace HarmonixMetasound::Nodes::MidiNoteTranspose
 
 	void FMidiStreamTransposerOperator_V0::BindInputs(FInputVertexInterfaceData& InVertexData)
 	{
-		using namespace CommonPinNames;
-
 		InVertexData.BindReadVertex(METASOUND_GET_PARAM_NAME(Inputs::Enable), EnableInPin);
 		InVertexData.BindReadVertex(METASOUND_GET_PARAM_NAME(CommonPinNames::Inputs::MidiTrackIndexFilterSpecifier), TrackSelectInPin);
 		InVertexData.BindReadVertex(METASOUND_GET_PARAM_NAME(CommonPinNames::Inputs::MidiChannelFilterSpecifier), ChannelSelectInPin);
