Fix engine FTickableGameObjects to never tick if they are part of a class default object.
Checking IsTemplate in GetTickableTickType only happens on the first tick, while IsAllowedToTick is called every frame.

diff --git a/Engine/Plugins/Runtime/CommonUI/Source/CommonUI/Private/Input/CommonBoundActionBar.cpp b/Engine/Plugins/Runtime/CommonUI/Source/CommonUI/Private/Input/CommonBoundActionBar.cpp
--- a/Engine/Plugins/Runtime/CommonUI/Source/CommonUI/Private/Input/CommonBoundActionBar.cpp
+++ b/Engine/Plugins/Runtime/CommonUI/Source/CommonUI/Private/Input/CommonBoundActionBar.cpp
@@ -47,7 +47,7 @@ void UCommonBoundActionBar::Tick(float DeltaTime)
 
 ETickableTickType UCommonBoundActionBar::GetTickableTickType() const
 {
-	return ETickableTickType::Always;
+	return (IsTemplate() ? ETickableTickType::Never : ETickableTickType::Always);
 }
 
 TStatId UCommonBoundActionBar::GetStatId() const
diff --git a/Engine/Plugins/Runtime/Harmonix/Source/HarmonixMetasound/Public/HarmonixMetasound/Subsystems/MidiClockUpdateSubsystem.h b/Engine/Plugins/Runtime/Harmonix/Source/HarmonixMetasound/Public/HarmonixMetasound/Subsystems/MidiClockUpdateSubsystem.h
--- a/Engine/Plugins/Runtime/Harmonix/Source/HarmonixMetasound/Public/HarmonixMetasound/Subsystems/MidiClockUpdateSubsystem.h
+++ b/Engine/Plugins/Runtime/Harmonix/Source/HarmonixMetasound/Public/HarmonixMetasound/Subsystems/MidiClockUpdateSubsystem.h
@@ -49,6 +49,7 @@ public:
 	// Begin FTickableGameObject
 	virtual bool IsTickable() const override;
 	virtual void Tick(float DeltaTime) override;
+	virtual ETickableTickType GetTickableTickType() const override { return (IsTemplate() ? ETickableTickType::Never : ETickableTickType::Conditional); }
 	virtual TStatId GetStatId() const override;
 	// End FTickableGameObject
 
diff --git a/Engine/Plugins/Runtime/ModelViewViewModel/Source/ModelViewViewModel/Private/View/MVVMBindingSubsystem.h b/Engine/Plugins/Runtime/ModelViewViewModel/Source/ModelViewViewModel/Private/View/MVVMBindingSubsystem.h
--- a/Engine/Plugins/Runtime/ModelViewViewModel/Source/ModelViewViewModel/Private/View/MVVMBindingSubsystem.h
+++ b/Engine/Plugins/Runtime/ModelViewViewModel/Source/ModelViewViewModel/Private/View/MVVMBindingSubsystem.h
@@ -36,7 +36,7 @@ public:
 	}
 	virtual ETickableTickType GetTickableTickType() const override
 	{
-		return ETickableTickType::Always;
+		return (IsTemplate() ? ETickableTickType::Never : ETickableTickType::Always);
 	}
 
 	virtual void Tick(float DeltaTime) override;
diff --git a/Engine/Source/Runtime/Engine/Private/GameFramework/InputDeviceSubsystem.cpp b/Engine/Source/Runtime/Engine/Private/GameFramework/InputDeviceSubsystem.cpp
--- a/Engine/Source/Runtime/Engine/Private/GameFramework/InputDeviceSubsystem.cpp
+++ b/Engine/Source/Runtime/Engine/Private/GameFramework/InputDeviceSubsystem.cpp
@@ -263,7 +263,7 @@ UWorld* UInputDeviceSubsystem::GetTickableGameObjectWorld() const
 
 ETickableTickType UInputDeviceSubsystem::GetTickableTickType() const
 {
-	return ETickableTickType::Conditional;
+	return (IsTemplate() ? ETickableTickType::Never : ETickableTickType::Conditional);
 }
 
 bool UInputDeviceSubsystem::IsAllowedToTick() const
diff --git a/Engine/Source/Runtime/Engine/Public/WorldPartition/DataLayer/ExternalDataLayerEngineSubsystem.h b/Engine/Source/Runtime/Engine/Public/WorldPartition/DataLayer/ExternalDataLayerEngineSubsystem.h
--- a/Engine/Source/Runtime/Engine/Public/WorldPartition/DataLayer/ExternalDataLayerEngineSubsystem.h
+++ b/Engine/Source/Runtime/Engine/Public/WorldPartition/DataLayer/ExternalDataLayerEngineSubsystem.h
@@ -38,6 +38,7 @@ public:
 #endif
 	virtual bool IsTickable() const override;
 	virtual void Tick(float DeltaTime) override;
+	virtual ETickableTickType GetTickableTickType() const override { return (IsTemplate() ? ETickableTickType::Never : ETickableTickType::Conditional); }
 	virtual TStatId GetStatId() const override;
 
 	ENGINE_API void RegisterExternalDataLayerAsset(const UExternalDataLayerAsset* InExternalDataLayerAsset, const UObject* InClient);
