Fix Fusion Patch Pan Settings not working.

Don't "Snap" to target when "Setting Target"...

#tests PIE

diff --git a/Engine/Plugins/Runtime/Harmonix/Source/HarmonixDsp/Private/HarmonixDsp/Panner.cpp b/Engine/Plugins/Runtime/Harmonix/Source/HarmonixDsp/Private/HarmonixDsp/Panner.cpp
--- a/Engine/Plugins/Runtime/Harmonix/Source/HarmonixDsp/Private/HarmonixDsp/Panner.cpp
+++ b/Engine/Plugins/Runtime/Harmonix/Source/HarmonixDsp/Private/HarmonixDsp/Panner.cpp
@@ -15,7 +15,7 @@ void FPanner::SetPanOffsetTarget_Stereo(float InMinusOneToOne)
 	float LeftRad = -FGainTable::Get().GetDirectChannelAzimuthInCurrentLayout(ESpeakerChannelAssignment::LeftFront);
 	float RightRad = -LeftRad;
 	OffsetPan = (((InMinusOneToOne + 1.0f) / 2.0f) * (RightRad - LeftRad)) + LeftRad;
-	UpdatePanRamper(true);
+	UpdatePanRamper(false);
 }
 
 void FPanner::UpdatePanRamper(bool Snap)
