[Harmonix] MIDI Player Node: make it work with Live Update enabled

#tests automated tests

diff --git a/Engine/Plugins/Runtime/Harmonix/Source/HarmonixMetasound/Private/HarmonixMetasound/Nodes/MidiPlayerNode.cpp b/Engine/Plugins/Runtime/Harmonix/Source/HarmonixMetasound/Private/HarmonixMetasound/Nodes/MidiPlayerNode.cpp
--- a/Engine/Plugins/Runtime/Harmonix/Source/HarmonixMetasound/Private/HarmonixMetasound/Nodes/MidiPlayerNode.cpp
+++ b/Engine/Plugins/Runtime/Harmonix/Source/HarmonixMetasound/Private/HarmonixMetasound/Nodes/MidiPlayerNode.cpp
@@ -15,7 +15,6 @@
 #include "HarmonixMetasound/DataTypes/MidiAsset.h"
 #include "HarmonixMetasound/DataTypes/MusicTransport.h"
 
-#include "HarmonixMidi/MidiPlayCursorMgr.h"
 #include "HarmonixMidi/MidiPlayCursor.h"
 #include "HarmonixMidi/MidiVoiceId.h"
 
@@ -86,8 +85,19 @@ namespace HarmonixMetasound
 		FMidiFileProxyPtr CurrentMidiFile;
 		FSampleCount BlockSize      = 0;
 		int32 CurrentBlockSpanStart = 0;
+		bool NeedsTransportInit = true;
 		
 		virtual void SetupNewMidiFile(const FMidiFileProxyPtr& NewMidi);
+
+		void InitTransportIfNeeded()
+		{
+			if (NeedsTransportInit)
+			{
+				InitTransportImpl();
+				NeedsTransportInit = false;
+			}
+		}
+		virtual void InitTransportImpl() = 0;
 	};
 
 	class FExternallyClockedMidiPlayerOperator : public FMidiPlayerOperator
@@ -105,13 +115,14 @@ namespace HarmonixMetasound
 
 		virtual void BindInputs(FInputVertexInterfaceData& InVertexData) override;
 		virtual void BindOutputs(FOutputVertexInterfaceData& InVertexData) override;
-		virtual FDataReferenceCollection GetInputs() const override;
-		virtual FDataReferenceCollection GetOutputs() const override;
 
 		virtual void Reset(const FResetParams& Params) override;
 
 		virtual void Execute() override;
 
+	protected:
+		virtual void InitTransportImpl() override;
+
 	private:
 		//** INPUTS **********************************
 		FMidiClockReadRef MidiClockIn;
@@ -140,8 +151,9 @@ namespace HarmonixMetasound
 
 		virtual void BindInputs(FInputVertexInterfaceData& InVertexData) override;
 		virtual void BindOutputs(FOutputVertexInterfaceData& InVertexData) override;
-		virtual FDataReferenceCollection GetInputs() const override;
-		virtual FDataReferenceCollection GetOutputs() const override;
+
+	protected:
+		virtual void InitTransportImpl() override;
 	};
 
 	class FMidiPlayerNode : public FNodeFacade
@@ -298,22 +310,6 @@ namespace HarmonixMetasound
 		FMidiPlayerOperator::BindOutputs(InVertexData);
 	}
 
-	FDataReferenceCollection FExternallyClockedMidiPlayerOperator::GetInputs() const
-	{
-		// This should never be called. Bind(...) is called instead. This method
-		// exists as a stop-gap until the API can be deprecated and removed.
-		checkNoEntry();
-		return {};
-	}
-
-	FDataReferenceCollection FExternallyClockedMidiPlayerOperator::GetOutputs() const
-	{
-		// This should never be called. Bind(...) is called instead. This method
-		// exists as a stop-gap until the API can be deprecated and removed.
-		checkNoEntry();
-		return {};
-	}
-
 	void FExternallyClockedMidiPlayerOperator::Reset(const FResetParams& Params)
 	{
 		FMidiPlayerOperator::Reset(Params);
@@ -330,20 +326,45 @@ namespace HarmonixMetasound
 		FMidiPlayerOperator::BindOutputs(InVertexData);
 	}
 
-	FDataReferenceCollection FSelfClockedMidiPlayerOperator::GetInputs() const
+	void FSelfClockedMidiPlayerOperator::InitTransportImpl()
 	{
-		// This should never be called. Bind(...) is called instead. This method
-		// exists as a stop-gap until the API can be deprecated and removed.
-		checkNoEntry();
-		return {};
-	}
+		// Get the node caught up to its transport input
+		FTransportInitFn InitFn = [this](EMusicPlayerTransportState CurrentState)
+		{
+			switch (CurrentState)
+			{
+			case EMusicPlayerTransportState::Invalid:
+			case EMusicPlayerTransportState::Preparing:
+			case EMusicPlayerTransportState::Prepared:
+			case EMusicPlayerTransportState::Stopping:
+			case EMusicPlayerTransportState::Killing:
+				MidiClockOut->AddTransportStateChangeToBlock({ 0, 0.0f, EMusicPlayerTransportState::Prepared });
+				MidiOutPin->AddTransportStateChangeMessage(0, EMusicPlayerTransportState::Prepared);
+				return EMusicPlayerTransportState::Prepared;
 
-	FDataReferenceCollection FSelfClockedMidiPlayerOperator::GetOutputs() const
-	{
-		// This should never be called. Bind(...) is called instead. This method
-		// exists as a stop-gap until the API can be deprecated and removed.
-		checkNoEntry();
-		return {};
+			case EMusicPlayerTransportState::Starting:
+			case EMusicPlayerTransportState::Playing:
+			case EMusicPlayerTransportState::Continuing:
+				MidiClockOut->ResetAndStart(0, true);
+				MidiOutPin->AddTransportStateChangeMessage(0, EMusicPlayerTransportState::Playing);
+				return EMusicPlayerTransportState::Playing;
+
+			case EMusicPlayerTransportState::Seeking: // seeking is omitted from init, shouldn't happen
+				checkNoEntry();
+				return EMusicPlayerTransportState::Invalid;
+
+			case EMusicPlayerTransportState::Pausing:
+			case EMusicPlayerTransportState::Paused:
+				MidiClockOut->AddTransportStateChangeToBlock({ 0, 0.0f, EMusicPlayerTransportState::Paused });
+				MidiOutPin->AddTransportStateChangeMessage(0, EMusicPlayerTransportState::Paused);
+				return EMusicPlayerTransportState::Paused;
+
+			default:
+				checkNoEntry();
+				return EMusicPlayerTransportState::Invalid;
+			}
+		};
+		Init(*TransportInPin, MoveTemp(InitFn));
 	}
 
 	void FMidiPlayerOperator::Execute()
@@ -355,6 +376,8 @@ namespace HarmonixMetasound
 		{
 			SetupNewMidiFile(MidiAssetInPin->GetMidiProxy());
 		}
+
+		InitTransportIfNeeded();
 	}
 
 	bool FMidiPlayerOperator::IsPlaying() const
@@ -525,8 +548,45 @@ namespace HarmonixMetasound
 			}
 		};
 		ExecuteTransportSpans(TransportInPin, BlockSize, TransportHandler, HandleMidiClockEventsInBlock);
+	}
+
+	void FExternallyClockedMidiPlayerOperator::InitTransportImpl()
+	{
+		// Get the node caught up to its transport input
+		// We don't send clock events for the externally-clocked player because those should already be handled
+		FTransportInitFn InitFn = [this](EMusicPlayerTransportState CurrentState)
+		{
+			switch (CurrentState)
+			{
+			case EMusicPlayerTransportState::Invalid:
+			case EMusicPlayerTransportState::Preparing:
+			case EMusicPlayerTransportState::Prepared:
+			case EMusicPlayerTransportState::Stopping:
+			case EMusicPlayerTransportState::Killing:
+				MidiOutPin->AddTransportStateChangeMessage(0, EMusicPlayerTransportState::Prepared);
+				return EMusicPlayerTransportState::Prepared;
+
+			case EMusicPlayerTransportState::Starting:
+			case EMusicPlayerTransportState::Playing:
+			case EMusicPlayerTransportState::Continuing:
+				MidiOutPin->AddTransportStateChangeMessage(0, EMusicPlayerTransportState::Playing);
+				return EMusicPlayerTransportState::Playing;
+
+			case EMusicPlayerTransportState::Seeking: // seeking is omitted from init, shouldn't happen
+				checkNoEntry();
+				return EMusicPlayerTransportState::Invalid;
 
-		GetTransportState();
+			case EMusicPlayerTransportState::Pausing:
+			case EMusicPlayerTransportState::Paused:
+				MidiOutPin->AddTransportStateChangeMessage(0, EMusicPlayerTransportState::Paused);
+				return EMusicPlayerTransportState::Paused;
+
+			default:
+				checkNoEntry();
+				return EMusicPlayerTransportState::Invalid;
+			}
+		};
+		Init(*TransportInPin, MoveTemp(InitFn));
 	}
 
 	void FExternallyClockedMidiPlayerOperator::UpdateLoopOffsetTickFromTick(int32 Tick)
@@ -662,6 +722,8 @@ namespace HarmonixMetasound
 		InVertexData.SetValue(METASOUND_GET_PARAM_NAME(MidiPlayerNodePinNames::KillVoicesOnMidiChange), bKillVoicesOnMidiChange);
 
 		SetupNewMidiFile(MidiAssetInPin->GetMidiProxy());
+
+		NeedsTransportInit = true;
 	}
 
 	void FMidiPlayerOperator::BindOutputs(FOutputVertexInterfaceData& InVertexData)
@@ -669,6 +731,8 @@ namespace HarmonixMetasound
 		using namespace CommonPinNames;
 		InVertexData.BindReadVertex(METASOUND_GET_PARAM_NAME(Outputs::MidiStream), MidiOutPin);
 		InVertexData.BindReadVertex(METASOUND_GET_PARAM_NAME(Outputs::MidiClock), MidiClockOut);
+
+		NeedsTransportInit = true;
 	}
 	
 	void FMidiPlayerOperator::Reset(const FResetParams& Params)
@@ -678,6 +742,8 @@ namespace HarmonixMetasound
 
 		MidiOutPin->SetClockSource(MidiClockOut);
 		MidiClockOut->ResetAndStart(0, true);
+
+		NeedsTransportInit = true;
 	}
 
 	void FMidiPlayerOperator::Reset(bool ForceNoBroadcast /*= false*/)
