Some UMidiFile code cleanup, rationalization, and minor bug fixes. Includes...
- Importing Format 0 midi files now converts them to Format 1 with a conductor track and one track for standard events.
- Adding comparison operators for a bunch of midi data structs to make writing tests easier
- Adding an "Export Standard Midi File" selection to UMidiFile asset context menu in the content browser to make testing easier (*generally a handy feature)
- Adding an automated test for UMidiFile import export (Currently not functional, but the framework is there.)

diff --git a/Engine/Plugins/Runtime/Harmonix/Source/HarmonixMidi/Private/HarmonixMidi/MidiFile.cpp b/Engine/Plugins/Runtime/Harmonix/Source/HarmonixMidi/Private/HarmonixMidi/MidiFile.cpp
--- a/Engine/Plugins/Runtime/Harmonix/Source/HarmonixMidi/Private/HarmonixMidi/MidiFile.cpp
+++ b/Engine/Plugins/Runtime/Harmonix/Source/HarmonixMidi/Private/HarmonixMidi/MidiFile.cpp
@@ -41,8 +41,6 @@ private:
 	virtual void OnText(int32 Tick, const FString& Str, uint8 Type) override;
 	virtual void OnTimeSignature(int32 Tick, int32 Numerator, int32 Denominator, bool FailOnError = true) override;
 
-	virtual void AddDummyConductorTrack() override;
-
 	UMidiFile& File;
 	int32  LastTick = 0;
 	bool   CurrentTrackHasName = false;
@@ -736,12 +734,6 @@ FMidiEventReceiver::FMidiEventReceiver(UMidiFile& file)
 {
 }
 
-void FMidiEventReceiver::AddDummyConductorTrack()
-{
-	// Add a dummy conductor track for midi type 0 files
-	File.GetTracks().Emplace("Conductor");
-}
-
 void FMidiEventReceiver::OnNewTrack(int32 newTrackIndex)
 {
 	File.GetTracks().Emplace();
diff --git a/Engine/Plugins/Runtime/Harmonix/Source/HarmonixMidi/Private/HarmonixMidi/MidiReader.cpp b/Engine/Plugins/Runtime/Harmonix/Source/HarmonixMidi/Private/HarmonixMidi/MidiReader.cpp
--- a/Engine/Plugins/Runtime/Harmonix/Source/HarmonixMidi/Private/HarmonixMidi/MidiReader.cpp
+++ b/Engine/Plugins/Runtime/Harmonix/Source/HarmonixMidi/Private/HarmonixMidi/MidiReader.cpp
@@ -238,6 +238,17 @@ void FStdMidiFileReader::ReadNextEventImpl()
 		ReadEvent(*InputArchive);
 		break;
 	case EState::NewTrack:
+		if (Format == 0 && CurrentTrackIndex == 0)
+		{
+			// we already ready the one and only track as the "conductor track".
+			// we have to read it again as "track 1", so...
+			InputArchive->Seek(LastTracksFilePosition);
+			TrackFilteringMode = ETrackFilteringMode::NonConductorEvents;
+		}
+		else 
+		{
+			LastTracksFilePosition = InputArchive->Tell();
+		}
 		ReadTrackHeader(*InputArchive);
 		break;
 	case EState::Start:
@@ -262,19 +273,33 @@ void FStdMidiFileReader::ReadFileHeader(FArchive& Archive)
 	//   0: A single multi-channel track
 	//   1: 1 or more simultaneous tracks
 	//   2: 1 or more sequentially independent tracks
-	// we only support type 1.
-	int16 Format;
+	// we only support type 0 & 1.
 	Archive << Format;
 	ensureMsgf(Format == 1 || Format == 0, TEXT("%s: Only type 0 or 1 MIDI files are supported; this file is type %d"), *Filename, Format);
 
+	if (Format == 0)
+	{
+		// We are starting to read a format 0 midi file, so on the first pass
+		// over the one and only track in the file gather the conductor type events...
+		TrackFilteringMode = ETrackFilteringMode::ConductorEvents;
+	}
+	else
+	{
+		TrackFilteringMode = ETrackFilteringMode::None;
+	}
+
 	Archive << NumTracks;
 
+	if (!ensureMsgf(Format != 0 || NumTracks == 1, TEXT("%s: Format 0 file expected only one track but found %d"), *Filename, NumTracks))
+	{
+		Failed = true;
+		return;
+	}
 	if (Format == 0)
 	{
-		// for midi type 0
-		Receiver->AddDummyConductorTrack();
-		++CurrentTrackIndex;
-		++NumTracks;  // we added a track
+		// we are going to convert it to a format 1 file, with a conductor track
+		// and one track for non-conductor events...
+		NumTracks = 2;
 	}
 
 	if (NumTracks <= 0)
@@ -322,6 +347,12 @@ void FStdMidiFileReader::ReadTrackHeader(FArchive& Archive)
 	}
 
 	TrackEndPos = int32(Archive.Tell() + TrackHeader.GetLength());
+	if (TrackEndPos > (int32)Archive.TotalSize())
+	{
+		UE_LOG(LogMidi, Error, TEXT("%s: MIDI track %d data length as recorded in the header exceeds the amount of data in the file!"), *Filename, CurrentTrackIndex);
+		Failed = true;
+		return;
+	}
 
 	++CurrentTrackIndex; // update track num
 
@@ -334,6 +365,19 @@ void FStdMidiFileReader::ReadTrackHeader(FArchive& Archive)
 
 	// Let the midi receiver know that this is a new track.
 	Receiver->OnNewTrack(CurrentTrackIndex);
+
+	if (Format == 0)
+	{
+		if (CurrentTrackIndex == 0)
+		{
+			CurrentTrackName = TEXT("conductor");
+		}
+		else
+		{
+			CurrentTrackName = TEXT("Track-0");
+		}
+		Receiver->OnText(0, CurrentTrackName, MidiConstants::kMeta_TrackName);
+	}
 }
 
 // Read an event from the track
@@ -440,7 +484,7 @@ void FStdMidiFileReader::ReadMidiEvent(int32 Tick, uint8 Status, uint8 Data1, FA
 		break;
 	}
 
-	if (ValidMidiEvent)
+	if (ValidMidiEvent && TrackFilteringMode != ETrackFilteringMode::ConductorEvents)
 	{
 		// At this time, we only care about Note On / Note Off events.
 		QueueChannelMsg(Tick, Status, Data1, Data2);
@@ -481,66 +525,40 @@ void FStdMidiFileReader::ReadMetaEvent(int32 Tick, uint8 Type, FArchive& Archive
 {
 	int32 Length = Midi::VarLenNumber::Read(Archive);
 	int32 StartPos = int32(Archive.Tell());
+	FString WorkingString;
 
 	switch (Type)
 	{
-	case MidiConstants::kMeta_Text:
-	case MidiConstants::kMeta_Copyright:
 	case MidiConstants::kMeta_TrackName:
+		CurrentTrackName = ReadText(Archive, Length);
+		Receiver->OnText(Tick, CurrentTrackName, Type);
+		break;
+	case MidiConstants::kMeta_Copyright:
+	case MidiConstants::kMeta_Marker:
+	case MidiConstants::kMeta_CuePoint:
+		WorkingString = ReadText(Archive, Length);
+		if (TrackFilteringMode != ETrackFilteringMode::NonConductorEvents)
+		{
+			Receiver->OnText(Tick, WorkingString, Type);
+		}
+		break;
+	case MidiConstants::kMeta_Text:
 	case MidiConstants::kMeta_Lyric:
+		WorkingString = ReadText(Archive, Length);
+		if (TrackFilteringMode != ETrackFilteringMode::ConductorEvents)
 		{
-			unsigned char* AsUtf8;
-			if (TextEncoding == MidiConstants::EMidiTextEventEncoding::Latin1)
-			{
-				unsigned char* AsLatin = (unsigned char*)FMemory::Malloc(Length + 1);
-				Archive.Serialize(AsLatin, Length);
-				AsLatin[Length] = 0;
-				AsUtf8 = (unsigned char*)FMemory::Malloc(Length * 2 + 2);
-
-				unsigned char* in = AsLatin;
-				unsigned char* out = AsUtf8;
-				while (*in)
-				{
-					if (*in < 128)
-					{
-						*out++ = *in++;
-					}
-					else
-					{
-						*out++ = 0xc0 | *in >> 6;
-						*out++ = 0x80 | (*in++ & 0x3f);
-					}
-				}
-				*out = 0;
-
-				FMemory::Free(AsLatin);
-			}
-			else
-			{
-				AsUtf8 = (unsigned char*)FMemory::Malloc(Length + 1);
-				Archive.Serialize(AsUtf8, Length);
-				AsUtf8[Length] = 0;
-			}
-
-			FString AsFString = UTF8_TO_TCHAR(AsUtf8);
-
-			if (Type == MidiConstants::kMeta_TrackName)
-			{
-				// Save ourselves a copy for better error messaging...
-				CurrentTrackName = AsFString;
-			}
-
-			Receiver->OnText(Tick, AsFString, Type);
-			FMemory::Free(AsUtf8);
+			Receiver->OnText(Tick, WorkingString, Type);
 		}
 		break;
-
 	case MidiConstants::kMeta_Tempo:
 		{
 			uint8 Byte1, Byte2, Byte3;
 			Archive << Byte1 << Byte2 << Byte3;
 			int32 Tempo = (Byte1 << 16) + (Byte2 << 8) + Byte3;
-			Receiver->OnTempo(Tick, Tempo);
+			if (TrackFilteringMode != ETrackFilteringMode::NonConductorEvents)
+			{
+				Receiver->OnTempo(Tick, Tempo);
+			}
 		}
 		break;
 
@@ -581,20 +599,22 @@ void FStdMidiFileReader::ReadMetaEvent(int32 Tick, uint8 Type, FArchive& Archive
 					*MidiTickFormat(Tick, BarMap.Get(), Midi::EMusicTimeStringFormat::Position), Numerator);
 			}
 
-			// before telling any receivers about the time signature we update our own local tempo map
-			// so we can print better error and warning messages...
-			FMusicTimestamp Timestamp;
-			check(Tick == 0 || BarMap->GetNumTimeSignaturePoints() > 0);
-			int32 BarIndex = BarMap->TickToBarIncludingCountIn(Tick);
-			if (!BarMap->AddTimeSignatureAtBarIncludingCountIn(BarIndex, Numerator, Denominator, true, false))
+			if (TrackFilteringMode != ETrackFilteringMode::NonConductorEvents)
 			{
-				UE_LOG(LogMidi, Warning, TEXT("%s (%s): Time signature %d/%d at %s overlaps or conflicts with nearby time signatures"),
-					*Filename, *CurrentTrackName, Numerator, Denominator, *MidiTickFormat(Tick, BarMap.Get(), Midi::EMusicTimeStringFormat::Position));
+				// before telling any receivers about the time signature we update our own local tempo map
+				// so we can print better error and warning messages...
+				FMusicTimestamp Timestamp;
+				check(Tick == 0 || BarMap->GetNumTimeSignaturePoints() > 0);
+				int32 BarIndex = BarMap->TickToBarIncludingCountIn(Tick);
+				if (!BarMap->AddTimeSignatureAtBarIncludingCountIn(BarIndex, Numerator, Denominator, true, false))
+				{
+					UE_LOG(LogMidi, Warning, TEXT("%s (%s): Time signature %d/%d at %s overlaps or conflicts with nearby time signatures"),
+						*Filename, *CurrentTrackName, Numerator, Denominator, *MidiTickFormat(Tick, BarMap.Get(), Midi::EMusicTimeStringFormat::Position));
+				}
+				// fill in our bar map
+				Receiver->OnTimeSignature(Tick, Numerator, Denominator, false);
 			}
 
-			// fill in our bar map
-			Receiver->OnTimeSignature(Tick, Numerator, Denominator, false);
-
 			// skip over next two bytes, "MIDI clocks per metronome click"
 			// and "32nd notes per MIDI quarter note"; they're not too useful
 			Archive.Seek(Archive.Tell() + 2);
@@ -606,8 +626,6 @@ void FStdMidiFileReader::ReadMetaEvent(int32 Tick, uint8 Type, FArchive& Archive
 	case MidiConstants::kMeta_KeySig:
 	case MidiConstants::kMeta_SMPTE:
 	case MidiConstants::kMeta_InstrumentName:
-	case MidiConstants::kMeta_Marker:
-	case MidiConstants::kMeta_CuePoint:
 		// these are valid events, but not currently supported by
 		// MidiReceiver; do nothing
 		break;
@@ -627,6 +645,46 @@ void FStdMidiFileReader::ReadMetaEvent(int32 Tick, uint8 Type, FArchive& Archive
 	Archive.Seek(StartPos + Length);
 }
 
+FString FStdMidiFileReader::ReadText(FArchive& Archive, int32 Length)
+{
+	unsigned char* AsUtf8 = nullptr;
+	if (TextEncoding == MidiConstants::EMidiTextEventEncoding::Latin1)
+	{
+		unsigned char* AsLatin = (unsigned char*)FMemory::Malloc(Length + 1);
+		Archive.Serialize(AsLatin, Length);
+		AsLatin[Length] = 0;
+		AsUtf8 = (unsigned char*)FMemory::Malloc(Length * 2 + 2);
+
+		unsigned char* in = AsLatin;
+		unsigned char* out = AsUtf8;
+		while (*in)
+		{
+			if (*in < 128)
+			{
+				*out++ = *in++;
+			}
+			else
+			{
+				*out++ = 0xc0 | *in >> 6;
+				*out++ = 0x80 | (*in++ & 0x3f);
+			}
+		}
+		*out = 0;
+
+		FMemory::Free(AsLatin);
+	}
+	else
+	{
+		AsUtf8 = (unsigned char*)FMemory::Malloc(Length + 1);
+		Archive.Serialize(AsUtf8, Length);
+		AsUtf8[Length] = 0;
+	}
+
+	FString AsFString = StringCast<TCHAR>((const UTF8CHAR*)AsUtf8).Get();
+	FMemory::Free(AsUtf8);
+	return AsFString;
+}
+
 // Queue all midi that is on the same tick, for sorting. 
 // Note that tick is already in terms of desired ticks per quarter note.
 void FStdMidiFileReader::QueueChannelMsg(int32 Tick, uint8 Status, uint8 Data1, uint8 Data2)
@@ -638,15 +696,18 @@ void FStdMidiFileReader::QueueChannelMsg(int32 Tick, uint8 Status, uint8 Data1,
 // Check if MidiList has anything in it. If so, sort and send out.
 void FStdMidiFileReader::ProcessMidiList()
 {
-	EventsOnSameTick.StableSort(RawMidiLess());
-	for (auto& MidiItem : EventsOnSameTick)
+	if (TrackFilteringMode != ETrackFilteringMode::ConductorEvents)
 	{
-		Receiver->OnMidiMessage(MidiListTick, MidiItem.Status, MidiItem.Data1, MidiItem.Data2);
-
-		// state changed because SkipCurrentTrack was called. Stop processing.
-		if (State != EState::InTrack)
+		EventsOnSameTick.StableSort(RawMidiLess());
+		for (auto& MidiItem : EventsOnSameTick)
 		{
-			break;
+			Receiver->OnMidiMessage(MidiListTick, MidiItem.Status, MidiItem.Data1, MidiItem.Data2);
+
+			// state changed because SkipCurrentTrack was called. Stop processing.
+			if (State != EState::InTrack)
+			{
+				break;
+			}
 		}
 	}
 	EventsOnSameTick.Empty();
diff --git a/Engine/Plugins/Runtime/Harmonix/Source/HarmonixMidi/Private/HarmonixMidi/MidiTrack.cpp b/Engine/Plugins/Runtime/Harmonix/Source/HarmonixMidi/Private/HarmonixMidi/MidiTrack.cpp
--- a/Engine/Plugins/Runtime/Harmonix/Source/HarmonixMidi/Private/HarmonixMidi/MidiTrack.cpp
+++ b/Engine/Plugins/Runtime/Harmonix/Source/HarmonixMidi/Private/HarmonixMidi/MidiTrack.cpp
@@ -24,17 +24,6 @@ FMidiTrack::FMidiTrack(const FString& Name)
 	AddEvent(FMidiEvent(0, FMidiMsg::CreateText(stringIndex, MidiConstants::kMeta_TrackName)));
 }
 
-/*
-FArchive& operator<<(FArchive& Archive, FMidiTrack& Track)
-{
-	Archive << Track.Sorted;
-	Archive << Track.PrimaryMidiChannel;
-	Archive << Track.Events;
-	Archive << Track.Strings;
-	return Archive;
-}
-*/
-
 const FMidiEventList& FMidiTrack::GetEvents() const
 {
 	check(Sorted);
diff --git a/Engine/Plugins/Runtime/Harmonix/Source/HarmonixMidi/Public/HarmonixMidi/BarMap.h b/Engine/Plugins/Runtime/Harmonix/Source/HarmonixMidi/Public/HarmonixMidi/BarMap.h
--- a/Engine/Plugins/Runtime/Harmonix/Source/HarmonixMidi/Public/HarmonixMidi/BarMap.h
+++ b/Engine/Plugins/Runtime/Harmonix/Source/HarmonixMidi/Public/HarmonixMidi/BarMap.h
@@ -113,6 +113,13 @@ public:
 		, TimeSignature(InTimeSignature)
 	{}
 
+	bool operator==(const FTimeSignaturePoint& Other) const
+	{
+		return	BarIndex == Other.BarIndex &&
+				BeatIndex == Other.BeatIndex &&
+				TimeSignature == Other.TimeSignature;
+	}
+
 	UPROPERTY()
 	int32 BarIndex; // 0 based since it is internal!
 	UPROPERTY()
@@ -138,7 +145,8 @@ public:
 	FBarMap()
 		: TicksPerQuarterNote(MidiConstants::kTicksPerQuarterNoteInt)
 	{}
-	
+	friend bool operator==(const FBarMap& Left, const FBarMap& Right);
+
 	void Empty();
 	void Copy(const FBarMap& Other, int32 StartTick = 0, int32 EndTick = -1);
 	bool IsEmpty() const;
diff --git a/Engine/Plugins/Runtime/Harmonix/Source/HarmonixMidi/Public/HarmonixMidi/BeatMap.h b/Engine/Plugins/Runtime/Harmonix/Source/HarmonixMidi/Public/HarmonixMidi/BeatMap.h
--- a/Engine/Plugins/Runtime/Harmonix/Source/HarmonixMidi/Public/HarmonixMidi/BeatMap.h
+++ b/Engine/Plugins/Runtime/Harmonix/Source/HarmonixMidi/Public/HarmonixMidi/BeatMap.h
@@ -40,6 +40,11 @@ public:
 		, Type(InType)
 	{}
 
+	bool operator==(const FBeatMapPoint& Other) const
+	{
+		return PulseBar == Other.PulseBar && Type == Other.Type;
+	}
+
 	UPROPERTY()
 	int32 PulseBar;
 	UPROPERTY()
@@ -69,6 +74,11 @@ public:
 		, LastIncludedBeatIndex(InLastIncludedBeatIndex)
 	{}
 
+	bool operator==(const FPulseBar& Other) const
+	{
+		return StartTick == Other.StartTick && LengthTicks == Other.LengthTicks && FirstIncludedBeatIndex == Other.FirstIncludedBeatIndex && LastIncludedBeatIndex == Other.LastIncludedBeatIndex;
+	}
+
 	UPROPERTY()
 	int32 StartTick;
 	UPROPERTY()
@@ -91,6 +101,7 @@ public:
 	FBeatMap()
 		: TicksPerQuarterNote(MidiConstants::kTicksPerQuarterNoteInt)
 	{}
+	friend bool operator==(const FBeatMap& Left, const FBeatMap& Right);
 
 	void Empty();
 	void Copy(const FBeatMap& Other, int32 StartTick = 0, int32 EndTick = -1);
diff --git a/Engine/Plugins/Runtime/Harmonix/Source/HarmonixMidi/Public/HarmonixMidi/ChordMap.h b/Engine/Plugins/Runtime/Harmonix/Source/HarmonixMidi/Public/HarmonixMidi/ChordMap.h
--- a/Engine/Plugins/Runtime/Harmonix/Source/HarmonixMidi/Public/HarmonixMidi/ChordMap.h
+++ b/Engine/Plugins/Runtime/Harmonix/Source/HarmonixMidi/Public/HarmonixMidi/ChordMap.h
@@ -24,6 +24,11 @@ public:
 		, Name(ChordName)
 	{}
 
+	bool operator==(const FChordMapPoint& Other) const
+	{
+		return Name == Other.Name;
+	}
+
 	UPROPERTY()
 	FName  Name;
 };
@@ -44,6 +49,7 @@ public:
 	FChordProgressionMap()
 		: TicksPerQuarterNote(MidiConstants::kTicksPerQuarterNoteInt)
 	{}
+	friend bool operator==(const FChordProgressionMap& Left, const FChordProgressionMap& Right);
 
 	void Finalize(int32 LastTick);
 
diff --git a/Engine/Plugins/Runtime/Harmonix/Source/HarmonixMidi/Public/HarmonixMidi/MidiEvent.h b/Engine/Plugins/Runtime/Harmonix/Source/HarmonixMidi/Public/HarmonixMidi/MidiEvent.h
--- a/Engine/Plugins/Runtime/Harmonix/Source/HarmonixMidi/Public/HarmonixMidi/MidiEvent.h
+++ b/Engine/Plugins/Runtime/Harmonix/Source/HarmonixMidi/Public/HarmonixMidi/MidiEvent.h
@@ -19,6 +19,11 @@ struct HARMONIXMIDI_API FMidiEvent
 public:
 	FMidiEvent(int32 Tick, const FMidiMsg& InMessage);
 
+	bool operator==(const FMidiEvent& Other) const
+	{
+		return Tick == Other.Tick && Message == Other.Message;
+	}
+
 	/** Save the event out in standard midi file format */
 	void WriteStdMidi(FMidiWriter& Writer, const FMidiTrack& Track) const;
 
diff --git a/Engine/Plugins/Runtime/Harmonix/Source/HarmonixMidi/Public/HarmonixMidi/MidiFile.h b/Engine/Plugins/Runtime/Harmonix/Source/HarmonixMidi/Public/HarmonixMidi/MidiFile.h
--- a/Engine/Plugins/Runtime/Harmonix/Source/HarmonixMidi/Public/HarmonixMidi/MidiFile.h
+++ b/Engine/Plugins/Runtime/Harmonix/Source/HarmonixMidi/Public/HarmonixMidi/MidiFile.h
@@ -65,6 +65,7 @@ struct HARMONIXMIDI_API FMidiFileData
 	*/
 	void AddTimeSigChange(int32 TrackIdx, int32 Tick, int32 TimeSigNum, int32 TimeSigDenom);
 
+	friend bool operator==(const FMidiFileData& Left, const FMidiFileData& Right);
 };
 
 template<>
@@ -91,6 +92,12 @@ public:
 	using FMidiTrackList = TArray<FMidiTrack>;
 
 	UMidiFile();
+
+	// A comparison operator. This allows for differences in members related to how
+	// a midi file was generated/imported, but compares the underlying "renderable"
+	// midi data.
+	friend bool operator==(const UMidiFile& Left, const UMidiFile& Right);
+
 	virtual void GetAssetRegistryTags(FAssetRegistryTagsContext Context) const override;
 	virtual void PostInitProperties() override;
 
@@ -107,7 +114,7 @@ public:
 	/** A method for exporting the midi track data to a standard midi file. */
 	void SaveStdMidiFile(const FString& FilePath);
 	/** A method for exporting the midi track data to a standard midi file. */
-	void SaveStdMidiFile(TSharedPtr<FArchive> Archive, const FString& Filename);
+	void SaveStdMidiFile(TSharedPtr<FArchive> Archive, const FString& Filename = FString());
 
 	FMidiTrack* AddTrack(const FString& Name);
 
diff --git a/Engine/Plugins/Runtime/Harmonix/Source/HarmonixMidi/Public/HarmonixMidi/MidiMsg.h b/Engine/Plugins/Runtime/Harmonix/Source/HarmonixMidi/Public/HarmonixMidi/MidiMsg.h
--- a/Engine/Plugins/Runtime/Harmonix/Source/HarmonixMidi/Public/HarmonixMidi/MidiMsg.h
+++ b/Engine/Plugins/Runtime/Harmonix/Source/HarmonixMidi/Public/HarmonixMidi/MidiMsg.h
@@ -109,6 +109,10 @@ public:
 		, Data2(0)
 	{}
 
+	bool operator==(const FMidiMsg& Other) const
+	{
+		return Type == Other.Type && Status == Other.Status && Data1 == Other.Data1 && Data2 == Other.Data2;
+	}
 
 	EType Type;
 	// Use a union instead of subclassing MidiMsg so that MidiMsg is 
diff --git a/Engine/Plugins/Runtime/Harmonix/Source/HarmonixMidi/Public/HarmonixMidi/MidiReader.h b/Engine/Plugins/Runtime/Harmonix/Source/HarmonixMidi/Public/HarmonixMidi/MidiReader.h
--- a/Engine/Plugins/Runtime/Harmonix/Source/HarmonixMidi/Public/HarmonixMidi/MidiReader.h
+++ b/Engine/Plugins/Runtime/Harmonix/Source/HarmonixMidi/Public/HarmonixMidi/MidiReader.h
@@ -85,6 +85,7 @@ private:
 	void ReadMidiEvent(int32 Tick, uint8 Status, uint8 Data1, FArchive& Archive);
 	void ReadSystemEvent(int32 Tick, uint8 Status, FArchive& Archive);
 	void ReadMetaEvent(int32 Tick, uint8 Type, FArchive& Archive);
+	FString ReadText(FArchive& Archive, int32 Length);
 
 	// functions for sorting the midi before sending it out to the receivers
 	void QueueChannelMsg(int32 Tick, uint8 Status, uint8 Data1, uint8 Data2);
@@ -110,7 +111,9 @@ private:
 	int32 DestinationTicksPerQuarterNote = MidiConstants::kTicksPerQuarterNoteInt;
 	float TickConversionFactor = 1.0f; // factor to convert from file's tick-per-quarternote to kTicksPerQuarterNote 
 
+	int16 Format = 1;             // Format 0 or Format 1 standard midi file?
 	int16 NumTracks = 0;          // Number of tracks.
+	int64 LastTracksFilePosition; // Position in the file of the current track's first byte.
 	int32 TrackEndPos = 0;        // end byte position of the current track
 	int32 CurrentFileTick = 0;    // Current tick in units of File's ticks-per-quarter.
 	int32 CurrentTick = 0;        // Current tick count for current track where one quarter note is TicksPerQuarterNote.
@@ -119,6 +122,14 @@ private:
 	int32 CurrentTrackIndex = -1; // What track are we currently reading?
 	FString CurrentTrackName;
 
+	enum class ETrackFilteringMode
+	{
+		None,
+		ConductorEvents,
+		NonConductorEvents,
+	};
+	ETrackFilteringMode TrackFilteringMode = ETrackFilteringMode::None;
+
 	// stuff for sorting Midi messages on the same tick
 	struct FRawMidiMsg
 	{
diff --git a/Engine/Plugins/Runtime/Harmonix/Source/HarmonixMidi/Public/HarmonixMidi/MidiReceiver.h b/Engine/Plugins/Runtime/Harmonix/Source/HarmonixMidi/Public/HarmonixMidi/MidiReceiver.h
--- a/Engine/Plugins/Runtime/Harmonix/Source/HarmonixMidi/Public/HarmonixMidi/MidiReceiver.h
+++ b/Engine/Plugins/Runtime/Harmonix/Source/HarmonixMidi/Public/HarmonixMidi/MidiReceiver.h
@@ -34,9 +34,6 @@ public:
 	/** called when all tracks have been read. */
 	virtual void Finalize(int32 InLastFileTick) {}
 
-	// For Midi type 0 files
-	virtual void AddDummyConductorTrack() {};
-
 	virtual void OnNewTrack(int32 NewTrackIndex) = 0;
 	virtual void OnEndOfTrack(int32 LastTick) = 0;
 	virtual void OnAllTracksRead() = 0;
diff --git a/Engine/Plugins/Runtime/Harmonix/Source/HarmonixMidi/Public/HarmonixMidi/MidiTrack.h b/Engine/Plugins/Runtime/Harmonix/Source/HarmonixMidi/Public/HarmonixMidi/MidiTrack.h
--- a/Engine/Plugins/Runtime/Harmonix/Source/HarmonixMidi/Public/HarmonixMidi/MidiTrack.h
+++ b/Engine/Plugins/Runtime/Harmonix/Source/HarmonixMidi/Public/HarmonixMidi/MidiTrack.h
@@ -23,6 +23,8 @@ public:
 	FMidiTrack();
 	FMidiTrack(const FString& name);
 
+	friend bool operator==(const FMidiTrack& Left, const FMidiTrack& Right);
+
 	const FMidiEventList& GetEvents() const;  		// this will ASSERT if events aren't sorted!
 	const FMidiEventList& GetUnsortedEvents() const; // this won't sort
 
diff --git a/Engine/Plugins/Runtime/Harmonix/Source/HarmonixMidi/Public/HarmonixMidi/SectionMap.h b/Engine/Plugins/Runtime/Harmonix/Source/HarmonixMidi/Public/HarmonixMidi/SectionMap.h
--- a/Engine/Plugins/Runtime/Harmonix/Source/HarmonixMidi/Public/HarmonixMidi/SectionMap.h
+++ b/Engine/Plugins/Runtime/Harmonix/Source/HarmonixMidi/Public/HarmonixMidi/SectionMap.h
@@ -23,6 +23,11 @@ public:
 		, Name(InName)
 	{}
 
+	bool operator==(const FSongSection& Other) const
+	{
+		return Name == Other.Name;
+	}
+
 	UPROPERTY(BlueprintReadOnly, Category = "SongSection")
 	FString Name;
 };
@@ -39,6 +44,7 @@ public:
 	FSectionMap()
 		: TicksPerQuarterNote(MidiConstants::kTicksPerQuarterNoteInt)
 	{}
+	friend bool operator==(const FSectionMap& Left, const FSectionMap& Right);
 
 	void Finalize(int32 LastTick);
 
diff --git a/Engine/Plugins/Runtime/Harmonix/Source/HarmonixMidi/Public/HarmonixMidi/SongMaps.h b/Engine/Plugins/Runtime/Harmonix/Source/HarmonixMidi/Public/HarmonixMidi/SongMaps.h
--- a/Engine/Plugins/Runtime/Harmonix/Source/HarmonixMidi/Public/HarmonixMidi/SongMaps.h
+++ b/Engine/Plugins/Runtime/Harmonix/Source/HarmonixMidi/Public/HarmonixMidi/SongMaps.h
@@ -25,6 +25,13 @@ public:
 	int32 LengthBars = 0;
 	UPROPERTY()
 	int32 LastTick = 0;
+
+	bool operator==(const FSongLengthData& Other) const
+	{
+		return	LengthTicks == Other.LengthTicks &&
+				LengthBars == Other.LengthBars && 
+				LastTick == Other.LastTick;
+	}
 };
 
 /**
@@ -42,6 +49,7 @@ struct HARMONIXMIDI_API FSongMaps
 
 public:
 	FSongMaps();
+	friend bool operator==(const FSongMaps& Left, const FSongMaps& Right);
 
 	void Init(int32 InTicksPerQuarterNote);
 	void Copy(const FSongMaps& Other, int32 StartTick = 0, int32 EndTick = -1);
diff --git a/Engine/Plugins/Runtime/Harmonix/Source/HarmonixMidi/Public/HarmonixMidi/TempoMap.h b/Engine/Plugins/Runtime/Harmonix/Source/HarmonixMidi/Public/HarmonixMidi/TempoMap.h
--- a/Engine/Plugins/Runtime/Harmonix/Source/HarmonixMidi/Public/HarmonixMidi/TempoMap.h
+++ b/Engine/Plugins/Runtime/Harmonix/Source/HarmonixMidi/Public/HarmonixMidi/TempoMap.h
@@ -26,6 +26,11 @@ public:
 	{
 	}
 
+	bool operator==(const FTempoInfoPoint& Other) const
+	{
+		return Ms == Other.Ms && MidiTempo == Other.MidiTempo;
+	}
+
 	float GetBPM() const { return MidiConstants::MidiTempoToBPM(MidiTempo); }
 
 	UPROPERTY()
@@ -60,6 +65,8 @@ public:
 	{}
 	virtual ~FTempoMap() {}
 
+	friend bool operator==(const FTempoMap& Left, const FTempoMap& Right);
+
 	void Empty();
 	void Copy(const FTempoMap& Other, int32 StartTick = 0, int32 EndTick = -1);
 	bool IsEmpty() const;
diff --git a/Engine/Plugins/Runtime/Harmonix/Source/HarmonixMidiEditor/HarmonixMidiEditor.Build.cs b/Engine/Plugins/Runtime/Harmonix/Source/HarmonixMidiEditor/HarmonixMidiEditor.Build.cs
--- a/Engine/Plugins/Runtime/Harmonix/Source/HarmonixMidiEditor/HarmonixMidiEditor.Build.cs
+++ b/Engine/Plugins/Runtime/Harmonix/Source/HarmonixMidiEditor/HarmonixMidiEditor.Build.cs
@@ -34,6 +34,9 @@ public class HarmonixMidiEditor : ModuleRules
 					"AssetDefinition",
 					"DetailCustomizations",
 					"InputCore",
+					"ContentBrowser",
+					"AssetTools",
+					"ToolMenus",
 				}
 		 );
 	}
diff --git a/Engine/Plugins/Runtime/Harmonix/Source/HarmonixMidiEditor/Private/AssetDefinition_MidiFile.cpp b/Engine/Plugins/Runtime/Harmonix/Source/HarmonixMidiEditor/Private/AssetDefinition_MidiFile.cpp
--- a/Engine/Plugins/Runtime/Harmonix/Source/HarmonixMidiEditor/Private/AssetDefinition_MidiFile.cpp
+++ b/Engine/Plugins/Runtime/Harmonix/Source/HarmonixMidiEditor/Private/AssetDefinition_MidiFile.cpp
@@ -1,9 +1,19 @@
 // Copyright Epic Games, Inc. All Rights Reserved.
 #include "AssetDefinition_MidiFile.h"
+
+#include "Algo/AnyOf.h"
+#include "ContentBrowserMenuContexts.h"
 #include "HarmonixMidi/MidiFile.h"
+#include "ToolMenus.h"
+#include "DesktopPlatformModule.h"
+#include "Framework/Application/SlateApplication.h"
+#include "Misc/MessageDialog.h"
+#include "GenericPlatform/GenericPlatformMisc.h"
 
 #define LOCTEXT_NAMESPACE "Harmonix_Midi"
 
+FString UAssetDefinition_MidiFile::LastMidiExportFolder;
+
 TSoftClassPtr<UObject> UAssetDefinition_MidiFile::GetAssetClass() const
 {
 	return UMidiFile::StaticClass();
@@ -31,4 +41,103 @@ bool UAssetDefinition_MidiFile::CanImport() const
 	return true;
 }
 
+void UAssetDefinition_MidiFile::RegisterContextMenu()
+{
+	UToolMenu* Menu = UToolMenus::Get()->ExtendMenu("ContentBrowser.AssetContextMenu.MidiFile");
+	FToolMenuSection& Section = Menu->FindOrAddSection("GetAssetActions");
+
+	Section.AddDynamicEntry("MidiFile_ExportMid",
+		FNewToolMenuSectionDelegate::CreateLambda([](FToolMenuSection& InSection)
+			{
+				const TAttribute<FText> Label = LOCTEXT("MidiFile_ExportMid", "Export Standard Midi File (.mid)");
+				const TAttribute<FText> ToolTip = LOCTEXT("MidiFile_ExportMidToolTip", "Exports standard midi file(s)");
+				const FSlateIcon Icon = FSlateIcon(FAppStyle::GetAppStyleSetName(), "ClassIcon.MidiFile");
+				const FToolMenuExecuteAction UIAction = FToolMenuExecuteAction::CreateStatic(&UAssetDefinition_MidiFile::ExecuteExportMidiFile);
+				InSection.AddMenuEntry("MidiFile_ExportMid", Label, ToolTip, Icon, UIAction);
+			}));
+}
+
+void UAssetDefinition_MidiFile::ExecuteExportMidiFile(const FToolMenuContext& MenuContext)
+{
+	if (const UContentBrowserAssetContextMenuContext* Context = UContentBrowserAssetContextMenuContext::FindContextWithAssets(MenuContext))
+	{
+		if (Context->SelectedAssets.Num() > 1)
+		{
+			ExportAllMidiToFolder(Context);
+		}
+		else
+		{
+			IDesktopPlatform* DesktopPlatform = FDesktopPlatformModule::Get();
+			if (DesktopPlatform)
+			{
+				UMidiFile* MidiFile = Context->LoadFirstSelectedObject<UMidiFile>();
+				FString	DefaultFileName = MidiFile->GetName() + TEXT(".mid");
+				TArray<FString> SaveFileNames;
+				const bool bFileSelected = DesktopPlatform->SaveFileDialog(
+					FSlateApplication::Get().FindBestParentWindowHandleForDialogs(nullptr),
+					LOCTEXT("MidiFile_ExportMid_SaveFileDialogTitle", "Save standard midi file as...").ToString(),
+					LastMidiExportFolder.IsEmpty() ? FPaths::ProjectDir() : LastMidiExportFolder,
+					DefaultFileName,
+					TEXT("Standard Midi File (*.mid)|*.mid"),
+					EFileDialogFlags::None,
+					SaveFileNames);
+
+				if (!bFileSelected)
+				{
+					return;
+				}
+
+				if (ensure(SaveFileNames.Num() == 1))
+				{
+					FString OutputFileName = SaveFileNames[0];
+					MidiFile->SaveStdMidiFile(OutputFileName);
+					LastMidiExportFolder = FPaths::GetPath(OutputFileName);
+				}
+			}
+		}
+	}
+}
+
+void UAssetDefinition_MidiFile::ExportAllMidiToFolder(const UContentBrowserAssetContextMenuContext* Context)
+{
+	IDesktopPlatform* DesktopPlatform = FDesktopPlatformModule::Get();
+	if (DesktopPlatform)
+	{
+		FString SelectedFolderName;
+		if (DesktopPlatform->OpenDirectoryDialog(FSlateApplication::Get().FindBestParentWindowHandleForDialogs(nullptr), "Select destination for standard MIDI files...", LastMidiExportFolder.IsEmpty() ? FPaths::ProjectDir() : LastMidiExportFolder, SelectedFolderName))
+		{
+			LastMidiExportFolder = SelectedFolderName;
+			bool bKeepWarning = true;
+			for (UMidiFile* MidiFile : Context->LoadSelectedObjects<UMidiFile>())
+			{
+				FString OutFilePath = FPaths::Combine(LastMidiExportFolder, MidiFile->GetName() + TEXT(".mid"));
+				if (bKeepWarning && FPaths::FileExists(OutFilePath))
+				{
+					EAppReturnType::Type AskResult = AskOverwrite(OutFilePath);
+					if (AskResult == EAppReturnType::NoAll)
+					{
+						break;
+					}
+					if (AskResult == EAppReturnType::No)
+					{
+						continue;
+					}
+					if (AskResult == EAppReturnType::YesAll)
+					{
+						bKeepWarning = false;
+					}
+				}
+				MidiFile->SaveStdMidiFile(OutFilePath);
+			}
+		}
+	}
+}
+
+EAppReturnType::Type UAssetDefinition_MidiFile::AskOverwrite(FString& OutPath)
+{
+	FText OutPathText = FText::FromString(*OutPath);
+	return FMessageDialog::Open(EAppMsgCategory::Warning, EAppMsgType::YesNoYesAllNoAll, 
+		FText::Format(LOCTEXT("MidiFile_ExportMid_OverwriteMessage", "{0} exists.\n\nWould you like to overwrite it?"), OutPathText));
+}
+
 #undef LOCTEXT_NAMESPACE
diff --git a/Engine/Plugins/Runtime/Harmonix/Source/HarmonixMidiEditor/Private/AssetDefinition_MidiFile.h b/Engine/Plugins/Runtime/Harmonix/Source/HarmonixMidiEditor/Private/AssetDefinition_MidiFile.h
--- a/Engine/Plugins/Runtime/Harmonix/Source/HarmonixMidiEditor/Private/AssetDefinition_MidiFile.h
+++ b/Engine/Plugins/Runtime/Harmonix/Source/HarmonixMidiEditor/Private/AssetDefinition_MidiFile.h
@@ -6,6 +6,10 @@
 
 #include "AssetDefinition_MidiFile.generated.h"
 
+struct FToolMenuContext;
+class UContentBrowserAssetContextMenuContext;
+namespace EAppReturnType { enum Type; }
+
 UCLASS()
 class UAssetDefinition_MidiFile : public UAssetDefinitionDefault
 {
@@ -19,5 +23,10 @@ public:
 	virtual TConstArrayView<FAssetCategoryPath> GetAssetCategories() const override;
 	virtual bool CanImport() const override;
 
+	static void RegisterContextMenu();
 private:
+	static void ExecuteExportMidiFile(const FToolMenuContext& MenuContext);
+	static void ExportAllMidiToFolder(const UContentBrowserAssetContextMenuContext* Context);
+	static EAppReturnType::Type AskOverwrite(FString& OutPath);
+	static FString LastMidiExportFolder;
 };
diff --git a/Engine/Plugins/Runtime/Harmonix/Source/HarmonixMidiEditor/Private/HarmonixMidiEditorModule.cpp b/Engine/Plugins/Runtime/Harmonix/Source/HarmonixMidiEditor/Private/HarmonixMidiEditorModule.cpp
--- a/Engine/Plugins/Runtime/Harmonix/Source/HarmonixMidiEditor/Private/HarmonixMidiEditorModule.cpp
+++ b/Engine/Plugins/Runtime/Harmonix/Source/HarmonixMidiEditor/Private/HarmonixMidiEditorModule.cpp
@@ -10,6 +10,8 @@
 #include "MidiFileDetailCustomization.h"
 #include "PropertyEditorDelegates.h"
 #include "IDetailCustomization.h"
+#include "ToolMenus.h"
+#include "AssetDefinition_MidiFile.h"
 
 #define LOCTEXT_NAMESPACE "Harmonix_Midi"
 
@@ -31,6 +33,14 @@ void FHarmonixMidiEditorModule::StartupModule()
 	);
 
 	PropertyModule.NotifyCustomizationModuleChanged();
+
+	UToolMenus::RegisterStartupCallback(FSimpleMulticastDelegate::FDelegate::CreateRaw(this, &FHarmonixMidiEditorModule::RegisterAssetContextMenus));
+}
+
+void FHarmonixMidiEditorModule::RegisterAssetContextMenus()
+{
+	FToolMenuOwnerScoped MenuOwner(this);
+	UAssetDefinition_MidiFile::RegisterContextMenu();
 }
 
 void FHarmonixMidiEditorModule::ShutdownModule()
diff --git a/Engine/Plugins/Runtime/Harmonix/Source/HarmonixMidiEditor/Private/MidiFileDetailCustomization.cpp b/Engine/Plugins/Runtime/Harmonix/Source/HarmonixMidiEditor/Private/MidiFileDetailCustomization.cpp
--- a/Engine/Plugins/Runtime/Harmonix/Source/HarmonixMidiEditor/Private/MidiFileDetailCustomization.cpp
+++ b/Engine/Plugins/Runtime/Harmonix/Source/HarmonixMidiEditor/Private/MidiFileDetailCustomization.cpp
@@ -57,7 +57,7 @@ void FMidiFileDetailCustomization::CustomizeDetails(IDetailLayoutBuilder& Detail
 		//if the file doesn't need to be conformed (and hence never conformed), display the integer bar length, otherwise display fractional bar length
 		if (!MidiFileBeingEdited->ShouldConformMidiFileLength(EMidiFileLengthConformOption::RoundDown) && !MidiFileBeingEdited->ShouldConformMidiFileLength(EMidiFileLengthConformOption::RoundUp))
 		{
-			FileLengthText = FText::FromString(FString::Printf(TEXT("%d Bar(s)"), (int32)MidiFileFractionalLength));
+			FileLengthText = FText::FromString(FString::Printf(TEXT("%d Bar(s)"), MidiFileLengthConformed));
 		}
 		else {
 			FileLengthText = FText::FromString(FString::Printf(TEXT("%.3f Bar(s)"), MidiFileFractionalLength));
diff --git a/Engine/Plugins/Runtime/Harmonix/Source/HarmonixMidiEditor/Private/MidiFileFactory.cpp b/Engine/Plugins/Runtime/Harmonix/Source/HarmonixMidiEditor/Private/MidiFileFactory.cpp
--- a/Engine/Plugins/Runtime/Harmonix/Source/HarmonixMidiEditor/Private/MidiFileFactory.cpp
+++ b/Engine/Plugins/Runtime/Harmonix/Source/HarmonixMidiEditor/Private/MidiFileFactory.cpp
@@ -106,7 +106,10 @@ void UMidiFileFactory::ShowConformMidiFileLengthDialog(int32 MidiFileAssetIndex)
 
 	// enable the Apply To All checkbox if importing more than 1 file 
 	(ImportedFiles.Num() > 1 && MidiFileAssetIndex != ImportedFiles.Num() - 1) ? ConformMidiFileLengthDialog->ApplyToAllCheckBox->SetVisibility(EVisibility::Visible) : ConformMidiFileLengthDialog->ApplyToAllCheckBox->SetVisibility(EVisibility::Collapsed);
-	
+
+	// we'll need this for a few options down below...	
+	float MidiFileFractionalLength = MidiFileAsset->GetSongMaps()->GetBarIncludingCountInAtTick(MidiFileAsset->GetLastEventTick() + 1);
+
 	//if the file has length less than 1 bar (e.g. 0.25 bars),it cannot be rounded down to 0 bar, this is checked in UMidiFile::ShouldConformFileLength(EMidiFileLengthConformOption Option)
 	//disable the Round Down checkbox and the Round To Nearest checkbox, the only available option should be Round Up
 	if (!MidiFileAsset->ShouldConformMidiFileLength(EMidiFileLengthConformOption::RoundDown) && MidiFileAsset->ShouldConformMidiFileLength(EMidiFileLengthConformOption::RoundUp))
@@ -120,12 +123,32 @@ void UMidiFileFactory::ShowConformMidiFileLengthDialog(int32 MidiFileAssetIndex)
 		ConformMidiFileLengthDialog->RoundUpCheckBox->SetIsChecked(ECheckBoxState::Checked);
 		ConformMidiFileLengthDialog->ConformOption = EMidiFileLengthConformOption::RoundUp;
 	}
+	else
+	{
+		ConformMidiFileLengthDialog->RoundDownCheckBox->SetIsChecked(ECheckBoxState::Unchecked);
+		ConformMidiFileLengthDialog->RoundToNearestCheckBox->SetIsChecked(ECheckBoxState::Unchecked);
+		ConformMidiFileLengthDialog->RoundUpCheckBox->SetIsChecked(ECheckBoxState::Unchecked);
+		if (FMath::Frac(MidiFileFractionalLength) < 0.1f || FMath::Frac(MidiFileFractionalLength) > 0.9f)
+		{
+			ConformMidiFileLengthDialog->RoundToNearestCheckBox->SetIsChecked(ECheckBoxState::Checked);
+			ConformMidiFileLengthDialog->ConformOption = EMidiFileLengthConformOption::Nearest;
+		}
+		else if (FMath::Frac(MidiFileFractionalLength) < 0.2f)
+		{
+			ConformMidiFileLengthDialog->RoundDownCheckBox->SetIsChecked(ECheckBoxState::Checked);
+			ConformMidiFileLengthDialog->ConformOption = EMidiFileLengthConformOption::RoundDown;
+		}
+		else
+		{
+			ConformMidiFileLengthDialog->RoundUpCheckBox->SetIsChecked(ECheckBoxState::Checked);
+			ConformMidiFileLengthDialog->ConformOption = EMidiFileLengthConformOption::RoundUp;
+		}
+	}
 
 	//set text for text block that asks user whether or not to conform midi file length
 	FString MidiFileName = MidiFileAsset->GetName();
-	float MidiFileFractionalLength = MidiFileAsset->GetSongMaps()->GetBarIncludingCountInAtTick(MidiFileAsset->GetLastEventTick());
 	ConformMidiFileLengthDialog->AskConformFileLengthText->SetText(FText::FromString(FString::Printf(
-		TEXT("Midi File Name: %s.mid, Length: %.3f bars"),
+		TEXT("Midi File Name: %s.mid, Length: %.5f bars"),
 		*MidiFileName,
 		MidiFileFractionalLength)));
 
diff --git a/Engine/Plugins/Runtime/Harmonix/Source/HarmonixMidiEditor/Public/HarmonixMidiEditorModule.h b/Engine/Plugins/Runtime/Harmonix/Source/HarmonixMidiEditor/Public/HarmonixMidiEditorModule.h
--- a/Engine/Plugins/Runtime/Harmonix/Source/HarmonixMidiEditor/Public/HarmonixMidiEditorModule.h
+++ b/Engine/Plugins/Runtime/Harmonix/Source/HarmonixMidiEditor/Public/HarmonixMidiEditorModule.h
@@ -16,5 +16,5 @@ public:
 	virtual void StartupModule();
 	virtual void ShutdownModule();
 private:
-
+	void RegisterAssetContextMenus();
 };
diff --git a/Engine/Plugins/Runtime/Harmonix/Source/HarmonixMidiTests/Private/MidiFileImportExportTests.cpp b/Engine/Plugins/Runtime/Harmonix/Source/HarmonixMidiTests/Private/MidiFileImportExportTests.cpp
new file mode 100644
--- /dev/null
+++ b/Engine/Plugins/Runtime/Harmonix/Source/HarmonixMidiTests/Private/MidiFileImportExportTests.cpp
@@ -0,0 +1,130 @@
+// Copyright Epic Games, Inc. All Rights Reserved.
+
+#include "MidiTestUtility.h"
+#include "Serialization/MemoryWriter.h"
+#include "Serialization/MemoryReader.h"
+#include "Misc/AutomationTest.h"
+#include "Misc/Paths.h"
+
+#if WITH_DEV_AUTOMATION_TESTS
+
+namespace HarmonixMidiTests::ImportExportTests
+{
+	using namespace Harmonix::Testing::Utility::MidiTestUtility;
+		
+		/**
+	* Add some events to the input midi file for testing the ConformMidiFileLength() function
+	* 
+	************************
+	* Midi File Structure: *
+	************************
+	*
+	*For each Midi track / channel :
+	*	1 Note On / Note Off pair will be added at the 1st beat of every bar, note duration is 1 beat(TicksPerQuarter)
+	*For every EVEN bar number :
+	*	1 CC event and 1 Text event will be added at the same position as the Note On events
+	*For every ODD bar number :
+	*	1 Pitch Bend event and 1 Poly Pres event will be added at the same position as the Note On events
+	*If input file length(bars) is a fractional number :
+	*	1 Note On / Note Off pair and 1 Text event will be added to the tick position of that fractional bar length
+	*	(Note On and Note Off events are 1 tick away so they don't end up on the same tick)
+	*/
+	void AddEventsToTestMidiFile(UMidiFile* InMidiFile, int32 InNumTracks, int32 InNumChannels, float InFileLengthBars)
+	{
+		constexpr int32 DefaultNoteNumber = 60;//C4
+		constexpr int32 DefaultNoteVelocity = 90;
+		constexpr uint8 DefaultControllerID = 69; //Hold Pedal
+		constexpr uint8 DefaultControlValue = 127;
+		constexpr uint8 DefaultPitchBendValueLSB = 64;
+		constexpr uint8 DefaultPitchBendValueMSB = 127;
+		constexpr uint8 DefaultPolyPressNoteNumber = 62;
+		constexpr uint8 DefaultPolyPresValue = 127;
+
+		//Add events to track 1 - InNumTracks, Track 0 is the conductor track
+		for (int32 TrackIndex = 1; TrackIndex < InNumTracks; ++TrackIndex)
+		{
+			FMidiTrack* CurrentTrack = InMidiFile->GetTrack(TrackIndex);
+			for (int32 Channel = 0; Channel < InNumChannels; ++Channel)
+			{
+				for (int32 Bar = 0; Bar < FMath::CeilToInt32(InFileLengthBars); ++Bar)
+				{
+					AddNoteOnNoteOffPairsToFile(InMidiFile, DefaultNoteNumber, DefaultNoteVelocity, TrackIndex, Channel, (float)Bar);
+
+					//Currently, Midi CC events & Midi Text events are added to bars with EVEN bar numbers,
+					//Midi Poly Press & Pitch Bend event are added to bars with ODD bar numbers
+					if (Bar % 2 == 0)
+					{
+						//Add Control Events to each channels/tracks
+						AddCCEventsToFile(InMidiFile, DefaultControllerID, DefaultNoteNumber, TrackIndex, Channel, (float)Bar);
+						//Add Text events to tracks
+						AddTextEventsToFile(InMidiFile, TEXT("TextInMidiFile"), TrackIndex, Bar);
+					}
+					else
+					{
+						//Add pitch bend events to channels/tracks
+						AddPitchEventsToFile(InMidiFile, DefaultPitchBendValueLSB, DefaultPitchBendValueMSB, TrackIndex, Channel, (float)Bar);
+						//Add Poly Pres events to channels/tracks
+						AddPolyPresEventsToFile(InMidiFile, DefaultPolyPressNoteNumber, DefaultPolyPresValue, TrackIndex, Channel, (float)Bar);
+					}
+				}
+				//if bar length is a fractional number, add additional midi events after the last integer bar
+				if (InFileLengthBars != (int32)InFileLengthBars)
+				{
+					//This will end up getting removed
+					AddNoteOnNoteOffPairsToFile(InMidiFile, DefaultNoteNumber, DefaultNoteVelocity, TrackIndex, Channel, InFileLengthBars);
+					//Add 2 CC Events (same controller ID) to test the conform function where it should remove events with the same type on the last tick
+					AddCCEventsToFile(InMidiFile, DefaultControllerID, DefaultControlValue, TrackIndex, Channel, InFileLengthBars);
+					AddCCEventsToFile(InMidiFile, DefaultControllerID, DefaultControlValue + 1, TrackIndex, Channel, InFileLengthBars);
+
+					//Add 2 Pitch Bend Events to test the conform function where it should remove events with the same type on the last tick
+					AddPitchEventsToFile(InMidiFile, DefaultPitchBendValueLSB, DefaultPitchBendValueMSB, TrackIndex, Channel, InFileLengthBars);
+					AddPitchEventsToFile(InMidiFile, DefaultPitchBendValueLSB + 1, DefaultPitchBendValueMSB + 1, TrackIndex, Channel, InFileLengthBars);
+
+					//Add 2 Pitch Bend Events to test the conform function where it should remove events with the same type on the last tick
+					AddPolyPresEventsToFile(InMidiFile, DefaultPolyPressNoteNumber, DefaultPolyPresValue, TrackIndex, Channel, InFileLengthBars);
+					AddPolyPresEventsToFile(InMidiFile, DefaultPolyPressNoteNumber, DefaultPolyPresValue + 1, TrackIndex, Channel, InFileLengthBars);
+				}
+			}
+			InMidiFile->GetTrack(TrackIndex)->Sort();
+		}
+
+		//update file information in SongLengthData
+		InMidiFile->TracksChanged();
+	}
+
+	IMPLEMENT_SIMPLE_AUTOMATION_TEST(
+		FTestMidiFileImportExport,
+		"Harmonix.Midi.MidiFile.ImportExport",
+		EAutomationTestFlags::EditorContext | EAutomationTestFlags::EngineFilter)
+		bool FTestMidiFileImportExport::RunTest(const FString&)
+	{
+		//input test values for creating a midi file
+		const float FileLengthBars = 5.5f;
+		const int32 NumChannels = 2;
+		const int32 NumTracks = 4;
+		const int32 TimeSigNum6 = 6;
+		const int32 TimeSigDenum8 = 8;
+		const int32 TimeSigNum4 = 4;
+		const int32 TimeSigDenum4 = 4;
+		const int32 Tempo = 120;
+		
+		//create an empty midi file to test for rounding to nearest (down)
+		UMidiFile* GeneratedMidiFile = BuildMidiFile(FileLengthBars, NumChannels, NumTracks, TimeSigNum4, TimeSigDenum4, Tempo);
+		//Add some events to the midi file for testing 
+		AddEventsToTestMidiFile(GeneratedMidiFile, NumTracks, NumChannels, FileLengthBars);
+
+		TArray<uint8> StdMidiFileBytes;
+		TSharedPtr<FMemoryWriter> StdMidiFileOut = MakeShared<FMemoryWriter>(StdMidiFileBytes, true);
+		GeneratedMidiFile->SaveStdMidiFile(StdMidiFileOut);
+
+		UMidiFile* ReimportedMidiFile = NewObject<UMidiFile>();
+		TSharedPtr<FMemoryReader> StdMidiFileIn = MakeShared<FMemoryReader>(StdMidiFileBytes, true);
+		ReimportedMidiFile->LoadStdMidiFile(StdMidiFileIn, TEXT("InMemoryMidiFile"));
+
+		// TODO
+		//TestTrue(TEXT("Generated midi file survives export/import cycle."),*GeneratedMidiFile == *ReimportedMidiFile);
+
+		return true;
+	}
+}
+#endif
\ No newline at end of file
diff --git a/Engine/Plugins/Runtime/Harmonix/Source/HarmonixMidiTests/Private/MidiTestUtility.cpp b/Engine/Plugins/Runtime/Harmonix/Source/HarmonixMidiTests/Private/MidiTestUtility.cpp
--- a/Engine/Plugins/Runtime/Harmonix/Source/HarmonixMidiTests/Private/MidiTestUtility.cpp
+++ b/Engine/Plugins/Runtime/Harmonix/Source/HarmonixMidiTests/Private/MidiTestUtility.cpp
@@ -120,4 +120,165 @@ namespace Harmonix::Testing::Utility::MidiTestUtility
 		);
 		CurrentTrack->AddEvent(CurrentPolyPresEvent);
 	}
-}
\ No newline at end of file
+}
+
+bool operator==(const FBarMap& Left, const FBarMap& Right)
+{
+	if (Left.StartBar != Right.StartBar || Left.TicksPerQuarterNote != Right.TicksPerQuarterNote || Left.Points.Num() != Right.Points.Num())
+	{
+		return false;
+	}
+	for (int32 PointIndex = 0; PointIndex < Left.Points.Num(); ++PointIndex)
+	{
+		if (Left.Points[PointIndex] != Right.Points[PointIndex])
+		{
+			return false;
+		}
+	}
+	return true;
+}
+
+bool operator==(const FBeatMap& Left, const FBeatMap& Right)
+{
+	if (Left.TicksPerQuarterNote != Right.TicksPerQuarterNote || Left.Points.Num() != Right.Points.Num() || Left.Bars.Num() != Right.Bars.Num())
+	{
+		return false;
+	}
+	for (int32 PointIndex = 0; PointIndex < Left.Points.Num(); ++PointIndex)
+	{
+		if (Left.Points[PointIndex] != Right.Points[PointIndex])
+		{
+			return false;
+		}
+	}
+
+	for (int32 BarIndex = 0; BarIndex < Left.Bars.Num(); ++BarIndex)
+	{
+		if (Left.Bars[BarIndex] != Right.Bars[BarIndex])
+		{
+			return false;
+		}
+	}
+	return true;
+}
+
+bool operator==(const FChordProgressionMap& Left, const FChordProgressionMap& Right)
+{
+	if (Left.TicksPerQuarterNote != Right.TicksPerQuarterNote || Left.Points.Num() != Right.Points.Num() || Left.ChordTrackIndex != Right.ChordTrackIndex)
+	{
+		return false;
+	}
+
+	for (int32 PointIndex = 0; PointIndex < Left.Points.Num(); ++PointIndex)
+	{
+		if (Left.Points[PointIndex] != Right.Points[PointIndex])
+		{
+			return false;
+		}
+	}
+	return true;
+}
+
+bool operator==(const FSectionMap& Left, const FSectionMap& Right)
+{
+	if (Left.TicksPerQuarterNote != Right.TicksPerQuarterNote || Left.Points.Num() != Right.Points.Num())
+	{
+		return false;
+	}
+
+	for (int32 PointIndex = 0; PointIndex < Left.Points.Num(); ++PointIndex)
+	{
+		if (Left.Points[PointIndex] != Right.Points[PointIndex])
+		{
+			return false;
+		}
+	}
+	return true;
+}
+
+bool operator==(const FTempoMap& Left, const FTempoMap& Right)
+{
+	if (Left.TicksPerQuarterNote != Right.TicksPerQuarterNote || Left.Points.Num() != Right.Points.Num())
+	{
+		return false;
+	}
+	for (int32 PointIndex = 0; PointIndex < Left.Points.Num(); ++PointIndex)
+	{
+		if (Left.Points[PointIndex] != Right.Points[PointIndex])
+		{
+			return false;
+		}
+	}
+	return true;
+}
+
+bool operator==(const FSongMaps& Left, const FSongMaps& Right)
+{
+	if (Left.TrackNames.Num() != Right.TrackNames.Num())
+	{
+		return false;
+	}
+	for (int32 TrackIndex = 0; TrackIndex < Left.TrackNames.Num(); ++TrackIndex)
+	{
+		if (Left.TrackNames[TrackIndex] != Right.TrackNames[TrackIndex])
+		{
+			return false;
+		}
+	}
+	return	Left.TicksPerQuarterNote == Right.TicksPerQuarterNote &&
+		Left.TempoMap == Right.TempoMap &&
+		Left.BarMap == Right.BarMap &&
+		Left.BeatMap == Right.BeatMap &&
+		Left.SectionMap == Right.SectionMap &&
+		Left.ChordMap == Right.ChordMap &&
+		Left.LengthData == Right.LengthData;
+}
+
+
+bool operator==(const FMidiTrack& Left, const FMidiTrack& Right)
+{
+	if (Left.Events.Num() != Right.Events.Num() || Left.Strings.Num() != Right.Strings.Num())
+	{
+		return false;
+	}
+	for (int32 EventIndex = 0; EventIndex < Left.Events.Num(); ++EventIndex)
+	{
+		if (Left.Events[EventIndex] != Right.Events[EventIndex])
+		{
+			return false;
+		}
+	}
+	for (int32 StringIndex = 0; StringIndex < Left.Strings.Num(); ++StringIndex)
+	{
+		if (Left.Strings[StringIndex] != Right.Strings[StringIndex])
+		{
+			return false;
+		}
+	}
+	return Left.Sorted == Right.Sorted && Left.PrimaryMidiChannel == Right.PrimaryMidiChannel;
+}
+
+bool operator==(const FMidiFileData& Left, const FMidiFileData& Right)
+{
+	if (Left.Tracks.Num() != Right.Tracks.Num())
+	{
+		return false;
+	}
+	for (int32 TrackIndex = 0; TrackIndex < Left.Tracks.Num(); ++TrackIndex)
+	{
+		if (Left.Tracks[TrackIndex] != Right.Tracks[TrackIndex])
+		{
+			return false;
+		}
+	}
+
+	return	Left.MidiFileName == Right.MidiFileName &&
+			Left.TicksPerQuarterNote == Right.TicksPerQuarterNote &&
+			Left.SongMaps == Right.SongMaps &&
+			Left.LastEventTick == Right.LastEventTick;
+}
+
+bool operator==(const UMidiFile& Left, const UMidiFile& Right)
+{
+	return Left.TheMidiData == Right.TheMidiData;
+}
diff --git a/Engine/Plugins/Runtime/Harmonix/Source/HarmonixMidiTests/Private/MidiTestUtility.h b/Engine/Plugins/Runtime/Harmonix/Source/HarmonixMidiTests/Private/MidiTestUtility.h
--- a/Engine/Plugins/Runtime/Harmonix/Source/HarmonixMidiTests/Private/MidiTestUtility.h
+++ b/Engine/Plugins/Runtime/Harmonix/Source/HarmonixMidiTests/Private/MidiTestUtility.h
@@ -45,4 +45,6 @@ namespace Harmonix::Testing::Utility::MidiTestUtility
 	*/
 	void AddPolyPresEventsToFile(UMidiFile* InFile, uint8 InNoteNumber, uint8 InPolyPresValue, int32 InTrackIndex, int32 InChannel, float InBarIndex);
 
-}
\ No newline at end of file
+}
+
+bool operator==(const UMidiFile& Left, const UMidiFile& Right);
