Add ability to create fusion patch from sound waves in the content browser

Parse the names of the SoundWaves to determine the keyzones they should be mapped to.

Adds the ability to
- Sort sound waves by name, note name, or note number, extracting the midi notes from the file name
- "Layout" the keyzones in different ways (distribute, layer, or single note)
- Assign the zones to a "Scale" or evenly distribute them
- Align the root note to the "min", "max", or "center" of each zone

#tests PIE

diff --git a/Engine/Plugins/Runtime/Harmonix/Source/HarmonixDspEditor/HarmonixDspEditor.Build.cs b/Engine/Plugins/Runtime/Harmonix/Source/HarmonixDspEditor/HarmonixDspEditor.Build.cs
--- a/Engine/Plugins/Runtime/Harmonix/Source/HarmonixDspEditor/HarmonixDspEditor.Build.cs
+++ b/Engine/Plugins/Runtime/Harmonix/Source/HarmonixDspEditor/HarmonixDspEditor.Build.cs
@@ -21,6 +21,7 @@ public class HarmonixDspEditor : ModuleRules
 
 		PrivateDependencyModuleNames.AddRange(
 			new string[] {
+				"ContentBrowser",
 				"Engine",
 				"EditorWidgets",
 				"EditorScriptingUtilities",
@@ -31,11 +32,13 @@ public class HarmonixDspEditor : ModuleRules
 				"UnrealEd",
 				"Json",
 				"AssetDefinition",
+				"AssetTools",
 				"PropertyEditor",
 				"DetailCustomizations",
 				"Slate",
 				"SlateCore",
 				"Settings",
+				"ToolMenus",
 				"InputCore"
 			}
 		);
diff --git a/Engine/Plugins/Runtime/Harmonix/Source/HarmonixDspEditor/Private/AssetDefinitions/AssetDefinition_FusionPatch.cpp b/Engine/Plugins/Runtime/Harmonix/Source/HarmonixDspEditor/Private/AssetDefinitions/AssetDefinition_FusionPatch.cpp
new file mode 100644
--- /dev/null
+++ b/Engine/Plugins/Runtime/Harmonix/Source/HarmonixDspEditor/Private/AssetDefinitions/AssetDefinition_FusionPatch.cpp
@@ -0,0 +1,86 @@
+﻿// Copyright Epic Games, Inc. All Rights Reserved.
+
+#include "AssetDefinition_FusionPatch.h"
+
+#include "AssetToolsModule.h"
+#include "Algo/AnyOf.h"
+#include "ContentBrowserMenuContexts.h"
+#include "EditorDialogLibrary.h"
+#include "FusionPatchAssetFactory.h"
+#include "FusionPatchImportOptions.h"
+#include "Sound/SoundWaveProcedural.h"
+#include "ToolMenus.h"
+#include "ToolMenuSection.h"
+#include "HarmonixDsp/FusionSampler/FusionPatch.h"
+#include "FusionPatchAssetFactory.h"
+
+#define LOCTEXT_NAMESPACE "AssetTypeActions"
+
+void FFusionPatchExtension::RegisterMenus()
+{
+	UToolMenu* Menu = UToolMenus::Get()->ExtendMenu("ContentBrowser.AssetContextMenu.SoundWave");
+	FToolMenuSection& Section = Menu->FindOrAddSection("GetAssetActions");
+
+	Section.AddDynamicEntry("SoundWaveAssetConversion_CreateFusionPatch",
+		FNewToolMenuSectionDelegate::CreateLambda([] (FToolMenuSection& InSection)
+		{
+			if (const UContentBrowserAssetContextMenuContext* Context = InSection.FindContext<UContentBrowserAssetContextMenuContext>())
+			{
+				// filter out USoundWaveProcedural assets
+				if (!Algo::AnyOf(Context->SelectedAssets, [](const FAssetData& AssetData) { return AssetData.IsInstanceOf<USoundWaveProcedural>(); }))
+				{
+					const TAttribute<FText> Label = LOCTEXT("SoundWave_CreateFusionPatch", "Create Fusion Patch");
+					const TAttribute<FText> ToolTip = LOCTEXT("SoundWave_CreateFusionPatchToolTip", "Creates a Fusion Patch Asset from the selected sound waves,"
+						"where each sound wave is used as a keyzone in the Fusion Patch");
+					const FSlateIcon Icon = FSlateIcon(FAppStyle::GetAppStyleSetName(), "ClassIcon.FusionPatch");
+					const FToolMenuExecuteAction UIAction = FToolMenuExecuteAction::CreateStatic(&FFusionPatchExtension::ExecuteCreateFusionPatch);
+
+					InSection.AddMenuEntry("SoundWave_CreateFusionPatch", Label, ToolTip, Icon, UIAction);
+				}
+			}
+		}));
+}
+
+void FFusionPatchExtension::ExecuteCreateFusionPatch(const FToolMenuContext& MenuContext)
+{
+	const FString DefaultSuffix = TEXT("_Patch");
+	FAssetToolsModule& AssetToolsModule = FModuleManager::Get().LoadModuleChecked<FAssetToolsModule>("AssetTools");
+
+	UFusionPatchAssetFactory* FusionPatchFactory = NewObject<UFusionPatchAssetFactory>();
+
+	TArray<TWeakObjectPtr<USoundWave>> StagedSoundWaves;
+	if (const UContentBrowserAssetContextMenuContext* Context = UContentBrowserAssetContextMenuContext::FindContextWithAssets(MenuContext))
+	{
+		for (USoundWave* Wave : Context->LoadSelectedObjectsIf<USoundWave>([](const FAssetData& AssetData) { return !AssetData.IsInstanceOf<USoundWaveProcedural>(); }))
+		{
+			StagedSoundWaves.Add(Wave);
+		}
+
+		if (StagedSoundWaves.Num() > 0)
+		{
+			// Determine an appropriate name
+			FString Name;
+			FString PackagePath;
+			AssetToolsModule.Get().CreateUniqueAssetName(StagedSoundWaves[0]->GetPackage()->GetName(), DefaultSuffix, PackagePath, Name);
+
+			UFusionPatchCreateOptions::FArgs Args;
+			Args.StagedSoundWaves = MoveTemp(StagedSoundWaves);
+			Args.AssetName = Name;
+			Args.Directory = FPackageName::GetLongPackagePath(PackagePath);
+			bool WasOkayPressed = false;
+			const UFusionPatchCreateOptions* CreateOptions = UFusionPatchCreateOptions::GetWithDialog(MoveTemp(Args), WasOkayPressed);
+			if (WasOkayPressed)
+			{
+				FString UniqueName;
+				AssetToolsModule.Get().CreateUniqueAssetName(CreateOptions->AssetName, TEXT(""), PackagePath, UniqueName);
+				FusionPatchFactory->CreateOptions = CreateOptions;
+				UObject* CreatedAsset = AssetToolsModule.Get().CreateAsset(UniqueName, CreateOptions->FusionPatchDir.Path, UFusionPatch::StaticClass(), FusionPatchFactory);
+
+				TArray<UObject*> Assets = {CreatedAsset};
+				FAssetToolsModule::GetModule().Get().SyncBrowserToAssets(Assets);
+			}
+		}
+	}
+}
+
+#undef LOCTEXT_NAMESPACE
\ No newline at end of file
diff --git a/Engine/Plugins/Runtime/Harmonix/Source/HarmonixDspEditor/Private/AssetDefinitions/AssetDefinition_FusionPatch.h b/Engine/Plugins/Runtime/Harmonix/Source/HarmonixDspEditor/Private/AssetDefinitions/AssetDefinition_FusionPatch.h
--- a/Engine/Plugins/Runtime/Harmonix/Source/HarmonixDspEditor/Private/AssetDefinitions/AssetDefinition_FusionPatch.h
+++ b/Engine/Plugins/Runtime/Harmonix/Source/HarmonixDspEditor/Private/AssetDefinitions/AssetDefinition_FusionPatch.h
@@ -7,6 +7,8 @@
 
 #include "AssetDefinition_FusionPatch.generated.h"
 
+struct FToolMenuContext;
+
 UCLASS()
 class UAssetDefinition_FusionPatch : public UAssetDefinitionDefault
 {
@@ -24,4 +26,11 @@ public:
 	}
 	virtual bool CanImport() const { return true; }
 
+};
+
+class FFusionPatchExtension
+{
+public:
+	static void RegisterMenus();
+	static void ExecuteCreateFusionPatch(const FToolMenuContext& MenuContext);
 };
\ No newline at end of file
diff --git a/Engine/Plugins/Runtime/Harmonix/Source/HarmonixDspEditor/Private/Customization/FusionPatchImportOptionsCustomization.cpp b/Engine/Plugins/Runtime/Harmonix/Source/HarmonixDspEditor/Private/Customization/FusionPatchImportOptionsCustomization.cpp
new file mode 100644
--- /dev/null
+++ b/Engine/Plugins/Runtime/Harmonix/Source/HarmonixDspEditor/Private/Customization/FusionPatchImportOptionsCustomization.cpp
@@ -0,0 +1,132 @@
+﻿// Copyright Epic Games, Inc. All Rights Reserved.
+#include "FusionPatchImportOptionsCustomization.h"
+#include "FusionPatchImportOptions.h"
+
+#include "DetailLayoutBuilder.h"
+#include "DetailCategoryBuilder.h"
+#include "DetailWidgetRow.h"
+#include "IDetailChildrenBuilder.h"
+#include "PropertyCustomizationHelpers.h"
+
+#include "HarmonixDsp/FusionSampler/Settings/KeyzoneSettings.h"
+#include "Widgets/Layout/SUniformGridPanel.h"
+#include "Widgets/Text/STextBlock.h"
+
+void FFusionPatchCreateOptionsCustomization::CustomizeDetails(IDetailLayoutBuilder& DetailLayout)
+{
+	static const int32 MinDesiredSlotWidth = 80;
+	TSharedPtr<IPropertyHandle> KeyzonesPropertyHandle = DetailLayout.GetProperty(GET_MEMBER_NAME_CHECKED(UFusionPatchCreateOptions, Keyzones));
+	TSharedPtr<IPropertyHandle> SortOptionHandle = DetailLayout.GetProperty(GET_MEMBER_NAME_CHECKED(UFusionPatchCreateOptions, SortOption));
+	check(KeyzonesPropertyHandle);
+	TSharedPtr<IPropertyHandleArray> KeyzonesArrayHandle = KeyzonesPropertyHandle->AsArray();
+	check(KeyzonesArrayHandle);
+	
+	IDetailCategoryBuilder& KeyzonesCategory = DetailLayout.EditCategory(TEXT("Keyzones"), FText::GetEmpty(), ECategoryPriority::Uncommon);
+	KeyzonesCategory.InitiallyCollapsed(true);
+	
+	uint32 NumElements = 0;
+	KeyzonesArrayHandle->GetNumElements(NumElements);
+	if (NumElements > 0)
+	{
+		TSharedRef<IPropertyHandle> Element = KeyzonesArrayHandle->GetElement(0);
+		TSharedPtr<IPropertyHandle> SoundWaveHandle = Element->GetChildHandle(GET_MEMBER_NAME_CHECKED(FKeyzoneSettings, SoundWave));
+		TSharedPtr<IPropertyHandle> MinNoteHandle = Element->GetChildHandle(GET_MEMBER_NAME_CHECKED(FKeyzoneSettings, MinNote));
+		TSharedPtr<IPropertyHandle> MaxNoteHandle = Element->GetChildHandle(GET_MEMBER_NAME_CHECKED(FKeyzoneSettings, MaxNote));
+		TSharedPtr<IPropertyHandle> RootNoteHandle = Element->GetChildHandle(GET_MEMBER_NAME_CHECKED(FKeyzoneSettings, RootNote));
+		TSharedPtr<IPropertyHandle> MinVelocityHandle = Element->GetChildHandle(GET_MEMBER_NAME_CHECKED(FKeyzoneSettings, MinVelocity));
+		TSharedPtr<IPropertyHandle> MaxVelocityHandle = Element->GetChildHandle(GET_MEMBER_NAME_CHECKED(FKeyzoneSettings, MaxVelocity));
+
+		KeyzonesCategory.AddCustomRow(NSLOCTEXT("FusionoPatch_Details", "KeyzonesHeader", "Keyzones"))
+		.NameContent()
+		[
+			SoundWaveHandle->CreatePropertyNameWidget()
+		]
+		.ValueContent()
+		[
+		SNew(SUniformGridPanel)
+			.MinDesiredSlotWidth(MinDesiredSlotWidth)
+		   +SUniformGridPanel::Slot(0, 0)
+		   [
+				MinNoteHandle->CreatePropertyNameWidget()
+		   ]
+		   +SUniformGridPanel::Slot(1, 0)
+		   [
+			   RootNoteHandle->CreatePropertyNameWidget()
+		   ]
+		   +SUniformGridPanel::Slot(2, 0)
+		   [
+			   MaxNoteHandle->CreatePropertyNameWidget()
+		   ]
+		   +SUniformGridPanel::Slot(3, 0)
+		   [
+				MinVelocityHandle->CreatePropertyNameWidget()
+		   ]
+		   +SUniformGridPanel::Slot(4, 0)
+		   [
+		   		MaxVelocityHandle->CreatePropertyNameWidget()
+		   ]
+		];
+	}
+	TSharedPtr<FDetailArrayBuilder> ArrayBuilder = MakeShared<FDetailArrayBuilder>(KeyzonesPropertyHandle.ToSharedRef(), false, false, false);
+	
+	ArrayBuilder->OnGenerateArrayElementWidget(FOnGenerateArrayElementWidget::CreateLambda([](TSharedRef<IPropertyHandle> Element, int32 Index, IDetailChildrenBuilder& ChildrenBuilder)
+	{
+		TSharedPtr<IPropertyHandle> SoundWaveHandle = Element->GetChildHandle(GET_MEMBER_NAME_CHECKED(FKeyzoneSettings, SoundWave));
+		TSharedPtr<IPropertyHandle> MinNoteHandle = Element->GetChildHandle(GET_MEMBER_NAME_CHECKED(FKeyzoneSettings, MinNote));
+		TSharedPtr<IPropertyHandle> MaxNoteHandle = Element->GetChildHandle(GET_MEMBER_NAME_CHECKED(FKeyzoneSettings, MaxNote));
+		TSharedPtr<IPropertyHandle> RootNoteHandle = Element->GetChildHandle(GET_MEMBER_NAME_CHECKED(FKeyzoneSettings, RootNote));
+		TSharedPtr<IPropertyHandle> MinVelocityHandle = Element->GetChildHandle(GET_MEMBER_NAME_CHECKED(FKeyzoneSettings, MinVelocity));
+		TSharedPtr<IPropertyHandle> MaxVelocityHandle = Element->GetChildHandle(GET_MEMBER_NAME_CHECKED(FKeyzoneSettings, MaxVelocity));
+
+		UObject* SoundWave = nullptr;
+		SoundWaveHandle->GetValue(SoundWave);
+		if (!SoundWave)
+		{
+			return;
+		}
+				
+		FText SoundWaveName = FText::FromString(SoundWave->GetPackage()->GetName());
+
+		ChildrenBuilder.AddCustomRow(NSLOCTEXT("FusionPatch_Details", "KeyzoneProperty", "Keyzone"))
+		.EditCondition(TAttribute<bool>::CreateLambda([]() { return false; }), nullptr)
+		.NameContent()
+		[
+			SNew(STextBlock)
+			.Font(IDetailLayoutBuilder::GetDetailFont())
+			.Text(SoundWaveName)
+		]
+		.ValueContent()
+		[
+		 SNew(SUniformGridPanel)
+			.MinDesiredSlotWidth(MinDesiredSlotWidth)
+			+SUniformGridPanel::Slot(0, 0)
+			[
+				MinNoteHandle->CreatePropertyValueWidget()
+			]
+			+SUniformGridPanel::Slot(1, 0)
+			[
+				RootNoteHandle->CreatePropertyValueWidget()
+			]
+			+SUniformGridPanel::Slot(2, 0)
+			[
+				MaxNoteHandle->CreatePropertyValueWidget()
+			]
+			+SUniformGridPanel::Slot(3, 0)
+			[
+				MinVelocityHandle->CreatePropertyValueWidget()
+			]
+			+SUniformGridPanel::Slot(4, 0)
+			[
+				MaxVelocityHandle->CreatePropertyValueWidget()
+			]
+		];
+	}));
+
+	SortOptionHandle->SetOnPropertyValueChanged(FSimpleDelegate::CreateLambda([ArrayBuilder]()
+	{
+		ArrayBuilder->RefreshChildren();
+	}));
+	KeyzonesCategory.AddCustomBuilder(ArrayBuilder.ToSharedRef());
+
+	DetailLayout.HideProperty(KeyzonesPropertyHandle);
+}
diff --git a/Engine/Plugins/Runtime/Harmonix/Source/HarmonixDspEditor/Private/Customization/FusionPatchImportOptionsCustomization.h b/Engine/Plugins/Runtime/Harmonix/Source/HarmonixDspEditor/Private/Customization/FusionPatchImportOptionsCustomization.h
new file mode 100644
--- /dev/null
+++ b/Engine/Plugins/Runtime/Harmonix/Source/HarmonixDspEditor/Private/Customization/FusionPatchImportOptionsCustomization.h
@@ -0,0 +1,28 @@
+﻿// Copyright Epic Games, Inc. All Rights Reserved.
+
+#pragma once
+
+#include "CoreMinimal.h"
+#include "IDetailCustomization.h"
+#include "DetailCategoryBuilder.h"
+#include "UObject/NameTypes.h"
+#include "Templates/SharedPointer.h"
+#include "Layout/Visibility.h"
+
+class SWidget;
+
+//Fusion Patch detail customizatiaon class
+class FFusionPatchCreateOptionsCustomization : public IDetailCustomization
+{
+public:
+	static TSharedRef<IDetailCustomization> MakeInstance()
+	{
+		return MakeShareable(new FFusionPatchCreateOptionsCustomization);
+	}
+
+	virtual void CustomizeDetails(IDetailLayoutBuilder& DetailLayout) override;
+
+private:
+
+	void RebuildKeyzones();
+};
diff --git a/Engine/Plugins/Runtime/Harmonix/Source/HarmonixDspEditor/Private/FusionPatchAssetFactory.cpp b/Engine/Plugins/Runtime/Harmonix/Source/HarmonixDspEditor/Private/FusionPatchAssetFactory.cpp
--- a/Engine/Plugins/Runtime/Harmonix/Source/HarmonixDspEditor/Private/FusionPatchAssetFactory.cpp
+++ b/Engine/Plugins/Runtime/Harmonix/Source/HarmonixDspEditor/Private/FusionPatchAssetFactory.cpp
@@ -13,6 +13,7 @@
 #include "JsonImporterHelper.h"
 #include "FusionPatchJsonImporter.h"
 
+#include "Misc/FeedbackContext.h"
 #include "Misc/MessageDialog.h"
 #include "Misc/Paths.h"
 #include "Misc/PackageName.h"
@@ -20,7 +21,6 @@
 #include "UObject/Package.h"
 #include "HAL/FileManager.h"
 
-
 #define LOCTEXT_NAMESPACE "FusionPatchAssetFactory"
 
 DEFINE_LOG_CATEGORY(LogFusionPatchAssetFactory)
@@ -112,26 +112,6 @@ void UFusionPatchAssetFactory::PostImportCleanUp()
 	ReplaceExistingSamplesResponse = EAppReturnType::No;
 }
 
-const UFusionPatchImportOptions* UFusionPatchAssetFactory::GetImportOptions(const FString& DefaultSamplesImportPath)
-{
-	UFusionPatchImportOptions* ImportOptions = GetMutableDefault<UFusionPatchImportOptions>();
-	if (!ApplyOptionsToAllImport)
-	{
-		// prompt on the first asset and update user settings
-		if (ImportOptions->SamplesImportDir.Path.IsEmpty())
-		{
-			// otherwise, import audio samples into the destination path provided
-			ImportOptions->SamplesImportDir.Path = DefaultSamplesImportPath;
-		}
-		const FText ImportDialogTitle = NSLOCTEXT("FusionPatchImporter", "FusionPatchImportOptionsTitle", "Fusion Patch Import Options");
-		UEditorDialogLibrary::ShowObjectDetailsView(ImportDialogTitle, ImportOptions);
-		const FString DestPath = ImportOptions->SamplesImportDir.Path;
-
-		ApplyOptionsToAllImport = true;
-	}
-	return ImportOptions;
-}
-
 bool UFusionPatchAssetFactory::GetReplaceExistingSamplesResponse(const FString& InName)
 {
 	if (ReplaceExistingSamplesResponse == EAppReturnType::YesAll)
@@ -185,6 +165,40 @@ void UFusionPatchAssetFactory::UpdateFusionPatchImportNotificationItem(TSharedPt
 
 UObject* UFusionPatchAssetFactory::FactoryCreateText(UClass* InClass, UObject* InParent, FName InName, EObjectFlags Flags, UObject* Context, const TCHAR* Type, const TCHAR*& Buffer, const TCHAR* BufferEnd, FFeedbackContext* Warn)
 {
+	const FString LongPackagePath = FPackageName::GetLongPackagePath(InParent->GetOutermost()->GetPathName());
+
+	const UFusionPatchImportOptions* ImportOptions = nullptr;
+	
+	if (!ApplyOptionsToAllImport)
+	{
+		bool WasOkayPressed = false;
+		UFusionPatchImportOptions::FArgs Args;
+		Args.Directory = LongPackagePath;
+		ImportOptions = UFusionPatchImportOptions::GetWithDialog(MoveTemp(Args), WasOkayPressed);
+		if (!WasOkayPressed)
+		{
+			// import cancelled by user
+			return nullptr;
+		}
+		ApplyOptionsToAllImport = true;
+	}
+	else
+	{
+		UFusionPatchImportOptions* MutableOptions = GetMutableDefault<UFusionPatchImportOptions>();
+		if (MutableOptions->SamplesImportDir.Path.IsEmpty())
+		{
+			MutableOptions->SamplesImportDir.Path = LongPackagePath;	
+		}
+		ImportOptions = MutableOptions;
+	}
+
+	if (Warn->ReceivedUserCancel())
+	{
+		return nullptr;
+	}
+
+	const bool ReplaceExistingSamples = GetReplaceExistingSamplesResponse(InName.ToString());
+	
 	const FString SourceFile = GetCurrentFilename();
 	AdditionalImportedObjects.Empty();
 	FString JsonString;
@@ -202,7 +216,7 @@ UObject* UFusionPatchAssetFactory::FactoryCreateText(UClass* InClass, UObject* I
 			FusionPatch = NewObject<UFusionPatch>(InParent, InName, Flags);
 		}
 
-		const FString LongPackagePath = FPackageName::GetLongPackagePath(InParent->GetOutermost()->GetPathName());
+		
 		const FString SourcePath = FPaths::GetPath(SourceFile);
 		
 		//create a notification that displays the import progress at the lower right corner
@@ -211,9 +225,6 @@ UObject* UFusionPatchAssetFactory::FactoryCreateText(UClass* InClass, UObject* I
 		TSharedPtr<SNotificationItem> ImportNotificationItem;
 		ImportNotificationItem = FSlateNotificationManager::Get().AddNotification(ImportNotificationInfo);
 
-		const UFusionPatchImportOptions* ImportOptions = GetImportOptions(LongPackagePath);
-		bool ReplaceExistingSamples = GetReplaceExistingSamplesResponse(InName.ToString());
-
 
 		// Pass import args to parser so it can import sub files
 		FFusionPatchJsonImporter::FImportArgs ImportArgs(InName, SourcePath, LongPackagePath, ImportOptions->SamplesImportDir.Path, ReplaceExistingSamples);
@@ -258,4 +269,23 @@ UObject* UFusionPatchAssetFactory::FactoryCreateText(UClass* InClass, UObject* I
 	return nullptr;
 }
 
+UObject* UFusionPatchAssetFactory::FactoryCreateNew(UClass* InClass, UObject* InParent, FName InName, EObjectFlags Flags, UObject* Context, FFeedbackContext* Warn)
+{
+	UFusionPatch* NewAsset = NewObject<UFusionPatch>(InParent, InName, Flags);
+
+	if (!NewAsset)
+	{
+		return nullptr;
+	}
+
+	if (CreateOptions)
+	{
+		NewAsset->UpdateKeyzones(CreateOptions->Keyzones);
+		NewAsset->UpdateSettings(CreateOptions->FusionPatchSettings);
+		CreateOptions = nullptr;
+	}
+
+	return NewAsset;
+}
+
 #undef LOCTEXT_NAMESPACE
\ No newline at end of file
diff --git a/Engine/Plugins/Runtime/Harmonix/Source/HarmonixDspEditor/Private/FusionPatchAssetFactory.h b/Engine/Plugins/Runtime/Harmonix/Source/HarmonixDspEditor/Private/FusionPatchAssetFactory.h
--- a/Engine/Plugins/Runtime/Harmonix/Source/HarmonixDspEditor/Private/FusionPatchAssetFactory.h
+++ b/Engine/Plugins/Runtime/Harmonix/Source/HarmonixDspEditor/Private/FusionPatchAssetFactory.h
@@ -9,6 +9,9 @@
 
 #include "FusionPatchAssetFactory.generated.h"
 
+class UFusionPatchImportOptions;
+class UFusionPatchCreateOptions;
+
 DECLARE_LOG_CATEGORY_EXTERN(LogFusionPatchAssetFactory, Log, All);
 
 UCLASS()
@@ -27,12 +30,13 @@ public:
 	virtual void PostImportCleanUp() override;
 
 	virtual UObject* FactoryCreateText(UClass* InClass, UObject* InParent, FName InName, EObjectFlags Flags, UObject* Context, const TCHAR* Type, const TCHAR*& Buffer, const TCHAR* BufferEnd, FFeedbackContext* Warn) override;
+	virtual UObject* FactoryCreateNew(UClass* InClass, UObject* InParent, FName InName, EObjectFlags Flags, UObject* Context, FFeedbackContext* Warn) override;
 
+	const UFusionPatchCreateOptions* CreateOptions = nullptr;
 private:
-	const class UFusionPatchImportOptions* GetImportOptions(const FString& DefaultSamplesImportPath);
+
 	bool GetReplaceExistingSamplesResponse(const FString& InName);
 	
-
 	bool ApplyOptionsToAllImport = false;
 
 	EAppReturnType::Type ReplaceExistingSamplesResponse = EAppReturnType::No;
diff --git a/Engine/Plugins/Runtime/Harmonix/Source/HarmonixDspEditor/Private/FusionPatchImportOptions.cpp b/Engine/Plugins/Runtime/Harmonix/Source/HarmonixDspEditor/Private/FusionPatchImportOptions.cpp
new file mode 100644
--- /dev/null
+++ b/Engine/Plugins/Runtime/Harmonix/Source/HarmonixDspEditor/Private/FusionPatchImportOptions.cpp
@@ -0,0 +1,510 @@
+﻿// Copyright Epic Games, Inc. All Rights Reserved.
+
+#include "FusionPatchImportOptions.h"
+#include "EditorDialogLibrary.h"
+#include "UObject/Package.h"
+#include "Internationalization/Regex.h"
+
+#if WITH_EDITOR
+
+struct FMidiScale
+{
+public:
+	
+	FMidiScale(TArray<int8>& InNotes)
+		: Notes(InNotes)
+	{
+	}
+	
+	FMidiScale(TArray<int8>&& InNotes)
+	: Notes(MoveTemp(InNotes))
+	{
+	}
+
+	int8 operator[](int32 Idx) const
+	{
+		int32 NoteIdx = Idx % Notes.Num();
+		int32 OctaveIdx = Idx / Notes.Num();
+		return Notes[NoteIdx] + (OctaveIdx * 12);
+	}
+
+	int32 NumNotes() const { return Notes.Num(); }
+
+	static const FMidiScale MajorScale;
+	
+private:
+	TArray<int8> Notes;
+};
+
+const FMidiScale FMidiScale::MajorScale(TArray<int8>({0, 2, 4, 5, 7, 9, 11}));
+
+const UFusionPatchImportOptions* UFusionPatchImportOptions::GetWithDialog(FArgs&& Args, bool& OutWasOkayPressed)
+{
+	UFusionPatchImportOptions* Options = GetMutableDefault<UFusionPatchImportOptions>();
+	// prompt on the first asset and update user settings
+	if (Options->SamplesImportDir.Path.IsEmpty())
+	{
+		// otherwise, import audio samples into the destination path provided
+		Options->SamplesImportDir.Path = MoveTemp(Args.Directory);
+	}
+	const FText ImportDialogTitle = NSLOCTEXT("FusionPatchImportOptions", "FusionPatchImportOptionsTitle", "Fusion Patch Import Options");
+	OutWasOkayPressed = UEditorDialogLibrary::ShowObjectDetailsView(ImportDialogTitle, Options);
+	return Options;
+}
+
+const UFusionPatchCreateOptions* UFusionPatchCreateOptions::GetWithDialog(FArgs&& Args, bool& OutWasOkayPressed)
+{
+	UFusionPatchCreateOptions* Options = GetMutableDefault<UFusionPatchCreateOptions>();
+	Options->FusionPatchDir.Path = MoveTemp(Args.Directory);
+	Options->AssetName = MoveTemp(Args.AssetName);
+	Options->StagedSoundWaves = MoveTemp(Args.StagedSoundWaves);
+	Options->UpdateKeyzonesWithSettings();
+	const FText ImportDialogTitle = NSLOCTEXT("FusionPatchCreateOptions", "FusionPatchCreateOptionsTitle", "New Fusion Patch Options");
+	FEditorDialogLibraryObjectDetailsViewOptions DialogOptions;
+	DialogOptions.bAllowResizing = true;
+	OutWasOkayPressed = UEditorDialogLibrary::ShowObjectDetailsView(ImportDialogTitle, Options, DialogOptions);
+	Options->StagedSoundWaves.Empty();
+
+	// apply some default settings that will work out of the box
+	Options->FusionPatchSettings.Adsr[0].IsEnabled = true;
+	Options->FusionPatchSettings.Adsr[0].SustainLevel = 1.0f;
+	
+	return Options;
+}
+
+void UFusionPatchCreateOptions::UpdateKeyzonesWithSettings()
+{
+	// autodetect parse option based on the sort option
+	// default to none
+	LockedNotesMask = ELockedNoteFlag::None;
+	switch (SortOption)
+	{
+	case EFusionPatchKeyzoneSortOption::Lexical:
+		SortLexical();
+		break;
+	case EFusionPatchKeyzoneSortOption::NoteNumber:
+		SortNoteNumber();
+		break;
+	case EFusionPatchKeyzoneSortOption::NoteName:
+		SortNoteName();
+		break;
+	case EFusionPatchKeyzoneSortOption::Index:
+		SortIndex();
+		break;
+	}
+
+	// apply scale to the root notes if we didn't parse notes from the note names
+	if (LockedNotesMask == ELockedNoteFlag::None)
+	{
+		const FMidiScale* Scale = nullptr;
+		switch (ScaleOption)
+		{
+		case EFusionPatchKeyzoneNoteScaleOption::MajorScale: Scale = &FMidiScale::MajorScale; break;
+		}
+		
+		if (Scale)
+		{
+			for (int32 Idx = 0; Idx < Keyzones.Num(); ++Idx)
+			{
+				Keyzones[Idx].RootNote = FMath::Clamp(MinNote + (*Scale)[Idx], MinNote, MaxNote);
+				Keyzones[Idx].MinNote = Keyzones[Idx].RootNote;
+				Keyzones[Idx].MaxNote = Keyzones[Idx].RootNote;
+			}
+			LockedNotesMask = ELockedNoteFlag::Root;
+		}
+	}
+
+	// if the min and max notes haven't already been determined from the file names, lay them out
+	if (ELockedNoteFlag::None ==  (LockedNotesMask & ELockedNoteFlag::MinMax))
+	{
+		switch (LayoutOption)
+		{
+		case EFusionPatchKeyzoneNoteLayoutOption::SingleNote:
+			LayoutSingleNote();
+			break;
+		case EFusionPatchKeyzoneNoteLayoutOption::Distribute:
+			LayoutDistribute();
+			break;
+		case EFusionPatchKeyzoneNoteLayoutOption::Layer:
+			LayoutLayer();
+			break;
+		}
+	}
+	
+	// set the root note only if we didn't parse it.
+	if (ELockedNoteFlag::None == (LockedNotesMask & ELockedNoteFlag::Root))
+	{
+		for (FKeyzoneSettings& Keyzone : Keyzones)
+		{
+			switch (RootNoteOption)
+			{
+			case EFusionPatchKeyzoneRootNoteOption::Min:
+				Keyzone.RootNote = Keyzone.MinNote;
+				break;
+			case EFusionPatchKeyzoneRootNoteOption::Max:
+				Keyzone.RootNote = Keyzone.MaxNote;
+				break;
+			case EFusionPatchKeyzoneRootNoteOption::Centered:
+				Keyzone.RootNote = (Keyzone.MaxNote + Keyzone.MinNote) / 2;
+				break;
+			}
+		}
+	}
+}
+
+void UFusionPatchCreateOptions::PostEditChangeProperty(struct FPropertyChangedEvent& PropertyChangedEvent)
+{
+	MaxNote = FMath::Clamp(MaxNote, 0, 127);
+	MinNote = FMath::Min(MinNote, MaxNote);
+	UpdateKeyzonesWithSettings();
+	SaveConfig();
+	Super::PostEditChangeProperty(PropertyChangedEvent);
+}
+
+void UFusionPatchCreateOptions::SortLexical()
+{
+	// assumes the sound waves are already in alphabetical order
+	Keyzones.Reset();
+	for (int32 Idx = 0; Idx < StagedSoundWaves.Num(); ++Idx)
+	{
+		TWeakObjectPtr<USoundWave> Wave = StagedSoundWaves[Idx];
+		FKeyzoneSettings NewKeyzone;
+		NewKeyzone.SoundWave = Wave.Get();
+		NewKeyzone.RootNote = Idx;
+		NewKeyzone.MinNote = Idx;
+		NewKeyzone.MaxNote = Idx;
+		Keyzones.Add(NewKeyzone);
+	}
+}
+
+void UFusionPatchCreateOptions::SortNoteNumber()
+{
+	Keyzones.Reset();
+	for (int32 Idx = 0; Idx < StagedSoundWaves.Num(); ++Idx)
+	{
+		TWeakObjectPtr<USoundWave> Wave = StagedSoundWaves[Idx];
+		FKeyzoneSettings NewKeyzone;
+		NewKeyzone.SoundWave = Wave.Get();
+
+		FKeyzoneNoteParser::FParseResult ParseResult;
+		if (FKeyzoneNoteParser::ParseMinRootMax(Wave->GetPackage()->GetName(), ParseResult))
+		{
+			LockedNotesMask = ELockedNoteFlag::MinRootMax;
+			NewKeyzone.MinNote = ParseResult.MinNote;
+			NewKeyzone.MaxNote = ParseResult.MaxNote;
+			NewKeyzone.RootNote = ParseResult.RootNote;
+		}
+		else if (FKeyzoneNoteParser::ParseMinMax(Wave->GetPackage()->GetName(), ParseResult))
+		{
+			LockedNotesMask = ELockedNoteFlag::MinMax;
+			NewKeyzone.MinNote = ParseResult.MinNote;
+			NewKeyzone.MaxNote = ParseResult.MaxNote;
+			NewKeyzone.RootNote = ParseResult.MinNote;
+		}
+		else if (FKeyzoneNoteParser::ParseRoot(Wave->GetPackage()->GetName(), ParseResult))
+		{
+			LockedNotesMask = ELockedNoteFlag::Root;
+			NewKeyzone.RootNote = ParseResult.RootNote;
+			NewKeyzone.MinNote = ParseResult.RootNote;
+			NewKeyzone.MaxNote = ParseResult.RootNote;
+		}
+		else
+		{
+			NewKeyzone.RootNote = Idx;
+			NewKeyzone.MinNote = Idx;
+			NewKeyzone.MaxNote = Idx;
+		}
+
+		// sort the keyzones by root note
+		const int32 Index = Algo::LowerBoundBy(Keyzones, NewKeyzone.RootNote, &FKeyzoneSettings::RootNote);
+		Keyzones.Insert(NewKeyzone, Index);
+	}
+}
+
+void UFusionPatchCreateOptions::SortNoteName()
+{
+	Keyzones.Reset();
+	for (int32 Idx = 0; Idx < StagedSoundWaves.Num(); ++Idx)
+	{
+		TWeakObjectPtr<USoundWave> Wave = StagedSoundWaves[Idx];
+		FKeyzoneSettings NewKeyzone;
+		NewKeyzone.SoundWave = Wave.Get();
+		
+		FKeyzoneNoteParser::FParseResult ParseResult;
+		if (FKeyzoneNoteParser::ParseNoteName(Wave->GetPackage()->GetName(), ParseResult))
+		{
+			LockedNotesMask = ELockedNoteFlag::Root;
+			NewKeyzone.MinNote = ParseResult.RootNote;
+			NewKeyzone.RootNote = ParseResult.RootNote;
+			NewKeyzone.MaxNote = ParseResult.RootNote;
+		}
+		else
+		{
+			NewKeyzone.MinNote = Idx;
+			NewKeyzone.RootNote = Idx;
+			NewKeyzone.MaxNote = Idx;
+		}
+
+		const int32 Index = Algo::LowerBoundBy(Keyzones, NewKeyzone.RootNote, &FKeyzoneSettings::RootNote);
+		Keyzones.Insert(NewKeyzone, Index);
+	}
+}
+
+void UFusionPatchCreateOptions::SortIndex()
+{
+	Keyzones.Reset();
+	for (int32 Idx = 0; Idx < StagedSoundWaves.Num(); ++Idx)
+	{
+		TWeakObjectPtr<USoundWave> Wave = StagedSoundWaves[Idx];
+		FKeyzoneSettings NewKeyzone;
+		NewKeyzone.SoundWave = Wave.Get();
+		FKeyzoneNoteParser::FParseResult ParseResult;
+		if (FKeyzoneNoteParser::ParseIndex(Wave->GetPackage()->GetName(), ParseResult))
+		{
+			NewKeyzone.MinNote = ParseResult.RootNote;
+			NewKeyzone.RootNote = ParseResult.RootNote;
+			NewKeyzone.MaxNote = ParseResult.RootNote;
+		}
+		else
+		{
+			NewKeyzone.MinNote = Idx;
+			NewKeyzone.RootNote = Idx;
+			NewKeyzone.MaxNote = Idx;
+		}
+		
+		const int32 Index = Algo::LowerBoundBy(Keyzones, NewKeyzone.RootNote, &FKeyzoneSettings::RootNote);
+		Keyzones.Insert(NewKeyzone, Index);
+	}
+}
+
+void UFusionPatchCreateOptions::LayoutSingleNote()
+{
+	int32 NumKeyzones = Keyzones.Num();
+	if (NumKeyzones == 0)
+	{
+		return;
+	}
+
+	if (LockedNotesMask == ELockedNoteFlag::None)
+	{
+		if (NumKeyzones == 1)
+		{
+			FKeyzoneSettings& Keyzone = Keyzones[0];
+			switch (RootNoteOption)
+			{
+			case EFusionPatchKeyzoneRootNoteOption::Min:
+				Keyzone.RootNote = MinNote;
+				break;
+			case EFusionPatchKeyzoneRootNoteOption::Max:
+				Keyzone.RootNote = MaxNote;
+				break;
+			case EFusionPatchKeyzoneRootNoteOption::Centered:
+				Keyzone.RootNote = (MaxNote - MinNote) / 2;
+				break;
+			}
+
+			Keyzone.MinNote = Keyzone.RootNote;
+			Keyzone.MaxNote = Keyzone.RootNote;
+			return;
+		}
+		
+		int32 StepOffset = 0;
+		float StepSize = float(MaxNote - MinNote) / float(NumKeyzones);
+		switch (RootNoteOption)
+		{
+		case EFusionPatchKeyzoneRootNoteOption::Min:
+			StepOffset = 0;
+			break;
+		case EFusionPatchKeyzoneRootNoteOption::Max:
+			StepOffset = 1;
+			break;
+		case EFusionPatchKeyzoneRootNoteOption::Centered:
+			StepSize = float(MaxNote - MinNote) / float(NumKeyzones - 1);
+			break;
+		}
+	
+		for (int32 Idx = 0; Idx < NumKeyzones; ++Idx)
+		{
+			FKeyzoneSettings& Keyzone = Keyzones[Idx];
+			Keyzone.RootNote = FMath::Clamp(MinNote + FMath::FloorToInt((Idx + StepOffset) * StepSize), MinNote, MaxNote);
+			Keyzone.MinNote = Keyzone.RootNote;
+			Keyzone.MaxNote = Keyzone.RootNote;
+		}
+	}
+	else if (LockedNotesMask == ELockedNoteFlag::Root)
+	{
+		for (int32 Idx = 0; Idx < NumKeyzones; ++Idx)
+		{
+			FKeyzoneSettings& Keyzone = Keyzones[Idx];
+			Keyzone.MinNote = Keyzone.RootNote;
+			Keyzone.MaxNote = Keyzone.RootNote;
+		}
+	}
+}
+
+void UFusionPatchCreateOptions::LayoutDistribute()
+{
+	int32 NumKeyzones = Keyzones.Num();
+	if (NumKeyzones == 0)
+	{
+		return;
+	}
+	if (NumKeyzones == 1)
+	{
+		FKeyzoneSettings& Keyzone = Keyzones[0];
+		Keyzone.MinNote = FMath::Min(MinNote, Keyzone.RootNote);
+		Keyzone.MaxNote = FMath::Max(MaxNote, Keyzone.RootNote);
+		return;
+	}
+	
+	// if we didn't parse anything, distribute evenly
+	if (LockedNotesMask == ELockedNoteFlag::None)
+	{
+		float StepSize = float(MaxNote - MinNote) / float(NumKeyzones);
+		for (int32 Idx = 0; Idx < NumKeyzones;  ++Idx)
+		{
+			FKeyzoneSettings& Keyzone = Keyzones[Idx];
+			Keyzone.MinNote = FMath::Clamp(MinNote + FMath::FloorToInt(Idx * StepSize), MinNote, MaxNote);
+			Keyzone.MaxNote = FMath::Clamp(MinNote + FMath::FloorToInt((Idx + 1) * StepSize) - 1, MinNote, MaxNote);
+		}
+		Keyzones[0].MinNote = MinNote;
+		Keyzones.Last().MaxNote = MaxNote;
+	}
+	// if we parsed the root note only, distribute the min and max around the root note
+	else if (LockedNotesMask == ELockedNoteFlag::Root)
+	{
+		switch (RootNoteOption)
+		{
+		case EFusionPatchKeyzoneRootNoteOption::Centered:
+			Keyzones[0].MinNote = FMath::Min(MinNote, Keyzones[0].RootNote);
+			Keyzones.Last().MaxNote = FMath::Max(MaxNote, Keyzones.Last().RootNote);
+			
+			for (int32 Idx = 0; Idx < NumKeyzones; ++Idx)
+			{
+				if (Idx < NumKeyzones - 1)
+				{
+					Keyzones[Idx].MaxNote = (Keyzones[Idx].RootNote + Keyzones[Idx+1].RootNote) / 2;
+				}
+				if (Idx > 0)
+				{
+					Keyzones[Idx].MinNote = (Keyzones[Idx].RootNote + Keyzones[Idx-1].RootNote) / 2;
+				}
+			}
+			break;
+		case EFusionPatchKeyzoneRootNoteOption::Min:
+			for (int32 Idx = 0; Idx < NumKeyzones - 1; ++Idx)
+			{
+				Keyzones[Idx].MaxNote = Keyzones[Idx + 1].RootNote - 1;
+				Keyzones[Idx].MinNote = Keyzones[Idx].RootNote;
+			}
+			Keyzones[0].MinNote = FMath::Min(MinNote, Keyzones[0].RootNote);
+			Keyzones.Last().MaxNote = FMath::Max(MaxNote, Keyzones.Last().RootNote);
+			break;
+		case EFusionPatchKeyzoneRootNoteOption::Max:
+			for (int32 Idx = NumKeyzones - 1; Idx > 0; --Idx)
+			{
+				Keyzones[Idx].MinNote = Keyzones[Idx - 1].RootNote + 1;
+				Keyzones[Idx].MaxNote = Keyzones[Idx].RootNote;
+			}
+			Keyzones[0].MinNote = FMath::Min(MinNote, Keyzones[0].RootNote);
+            Keyzones.Last().MaxNote = FMath::Max(MaxNote, Keyzones.Last().RootNote);
+			break;
+		}
+	}
+}
+
+void UFusionPatchCreateOptions::LayoutLayer()
+{
+	// if just the root note was parsed, clamp the min and max notes
+	if (LockedNotesMask == ELockedNoteFlag::Root)
+	{
+		for (FKeyzoneSettings& Keyzone : Keyzones)
+		{
+			Keyzone.MinNote = FMath::Min(MinNote, Keyzone.RootNote);
+			Keyzone.MaxNote = FMath::Max(MaxNote, Keyzone.RootNote);
+		}
+	}
+	// otherwise, apply min and max directly
+	else if (ELockedNoteFlag::None == (LockedNotesMask & ELockedNoteFlag::MinMax))
+	{
+		for (FKeyzoneSettings& Keyzone : Keyzones)
+		{
+			Keyzone.MinNote = MinNote;
+			Keyzone.MaxNote = MaxNote;
+		}
+	}
+}
+#endif
+
+bool FKeyzoneNoteParser::ParseMinRootMax(const FString& Input, FParseResult& Output)
+{
+	Output.Reset();
+	const FRegexPattern Pattern(TEXT("_(\\d+)_(\\d+)_(\\d+)"));
+
+	FRegexMatcher Matcher(Pattern, Input);
+	if (Matcher.FindNext())
+	{
+		Output.MinNote = FCString::Atoi(*Matcher.GetCaptureGroup(1));
+		Output.RootNote = FCString::Atoi(*Matcher.GetCaptureGroup(2));
+		Output.MaxNote = FCString::Atoi(*Matcher.GetCaptureGroup(3));
+		return true;
+	}
+	return false;
+}
+
+bool FKeyzoneNoteParser::ParseMinMax(const FString& Input, FParseResult& Output)
+{
+	Output.Reset();
+	const FRegexPattern Pattern(TEXT("_(\\d+)_(\\d+)"));
+
+	FRegexMatcher Matcher(Pattern, Input);
+	if (Matcher.FindNext())
+	{
+		Output.MinNote = FCString::Atoi(*Matcher.GetCaptureGroup(1));
+		Output.MaxNote = FCString::Atoi(*Matcher.GetCaptureGroup(2));
+		return true;
+	}
+	return false;
+}
+
+bool FKeyzoneNoteParser::ParseRoot(const FString& Input, FParseResult& Output)
+{
+	Output.Reset();
+	const FRegexPattern Pattern(TEXT("_(\\d+)"));
+
+	FRegexMatcher Matcher(Pattern, Input);
+	if (Matcher.FindNext())
+	{
+		Output.RootNote = FCString::Atoi(*Matcher.GetCaptureGroup(1));
+		return true;
+	}
+	return false;
+}
+
+bool FKeyzoneNoteParser::ParseNoteName(const FString& Input, FParseResult& Output)
+{
+	Output.Reset();
+	const FRegexPattern Pattern(TEXT("_([A-Ga-g][b#]?-?\\d+)"));
+
+	FRegexMatcher Matcher(Pattern, Input);
+	if (Matcher.FindNext())
+	{
+		Output.RootNote = MidiConstants::GetNoteNumberFromNoteName(TCHAR_TO_ANSI(*Matcher.GetCaptureGroup(1)));
+		return true;
+	}
+	return false;
+}
+
+bool FKeyzoneNoteParser::ParseIndex(const FString& Input, FParseResult& Output)
+{
+	Output.Reset();
+	const FRegexPattern Pattern(TEXT("_(\\d+)"));
+
+	FRegexMatcher Matcher(Pattern, Input);
+	if (Matcher.FindNext())
+	{
+		Output.RootNote = FCString::Atoi(*Matcher.GetCaptureGroup(1));
+		return true;
+	}
+	return false;
+}
\ No newline at end of file
diff --git a/Engine/Plugins/Runtime/Harmonix/Source/HarmonixDspEditor/Private/FusionPatchImportOptions.h b/Engine/Plugins/Runtime/Harmonix/Source/HarmonixDspEditor/Private/FusionPatchImportOptions.h
--- a/Engine/Plugins/Runtime/Harmonix/Source/HarmonixDspEditor/Private/FusionPatchImportOptions.h
+++ b/Engine/Plugins/Runtime/Harmonix/Source/HarmonixDspEditor/Private/FusionPatchImportOptions.h
@@ -5,7 +5,9 @@
 
 #include "Sound/SoundWave.h"
 #include "Sound/SoundWaveLoadingBehavior.h"
-
+#include "HarmonixDsp/FusionSampler/Settings/KeyzoneSettings.h"
+#include "HarmonixDsp/FusionSampler/Settings/FusionPatchSettings.h"
+#include "HarmonixMidi/MidiConstants.h"
 #include "FusionPatchImportOptions.generated.h"
 
 UCLASS(config = Editor)
@@ -14,6 +16,12 @@ class UFusionPatchImportOptions : public UObject
 	GENERATED_BODY()
 
 public:
+	struct FArgs
+	{
+		FString Directory;
+	};
+
+	static const UFusionPatchImportOptions* GetWithDialog(FArgs&& Args, bool& OutWasOkayPressed);
 
 	/** The directory to save samples to */
 	UPROPERTY(config, EditAnywhere, Category = "Import Options", Meta = (DisplayName = "Sound Waves Import Folder", ContentDir))
@@ -27,6 +35,7 @@ public:
 	UPROPERTY(EditAnywhere, Category = "Import Options", Meta = (DisplayName = "Sound Wave Compression Type"))
 	ESoundAssetCompressionType SampleCompressionType = ESoundAssetCompressionType::BinkAudio;
 
+	
 #if WITH_EDITOR
 	/** UObject interface */
 	virtual void PostEditChangeProperty(struct FPropertyChangedEvent& PropertyChangedEvent) override
@@ -34,4 +43,157 @@ public:
 		SaveConfig();
 	}
 #endif
-};
\ No newline at end of file
+};
+
+UENUM()
+enum class ELockedNoteFlag : uint8
+{
+	None = 0b000,
+	Min	 = 0b001,
+	Root = 0b010,
+	Max  = 0b100,
+	MinMax = Min | Max,
+	MinRootMax = Min | Root | Max
+};
+ENUM_CLASS_FLAGS(ELockedNoteFlag)
+
+UENUM()
+enum class EFusionPatchKeyzoneSortOption : uint8
+{
+	Lexical		UMETA(DisplayName="Alphabetical", ToolTip="Alphabetical order."),
+	Index		UMETA(DisplayName="Index Order", ToolTip="Sample name contains index, ordering the samples alphabetically then from lowest to highest index (snare_0, snare_1, tom_0)"),
+	NoteNumber	UMETA(DisplayName="Note Number", ToolTip="Sample names contain note numbers with one of the acceptable formats: (name_[Min]_[Root]_[Max], name_[Min]_[Max], name_[Root])"),
+	NoteName	UMETA(DisplayName="Note Name", ToolTip="Sample names contain note names: (eg. 'Eb1', 'G#3', 'B2')"),
+};
+
+UENUM()
+enum class EFusionPatchKeyzoneNoteLayoutOption : uint8
+{
+	SingleNote	UMETA(DisplayName="Single Note", ToolTip="Assign each zone to a single note, 'root' note if parsed by note number or note name, otherwise each note from the selected 'min - max' range"),
+	Distribute	UMETA(DisplayName="Distribute", ToolTip="Evenly distribute all zones, using the values parsed by the note number or note name, otherwise uses 'min-max' note range"),
+	Layer		UMETA(DisplayName="Layer", ToolTip="Every zone will span 'min-max' note range."),
+};
+
+UENUM()
+enum class EFusionPatchKeyzoneNoteScaleOption : uint8
+{
+	None		UMETA(DisplayName="No Scale (Equal Distance)", ToolTip="Distributes the notes evenly based on the selected layout"),
+	MajorScale	UMETA(DisplayName="Major Scale", ToolTip="Distributes the notes across a major scale, starting with the Min Note")
+};
+
+UENUM()
+enum class EFusionPatchKeyzoneRootNoteOption : uint8
+{
+	Min			UMETA(DisplayName="Align Min", ToolTip="Adjusts the range so that the Root note is the same as the Min note of the zone"),
+	Max			UMETA(DisplayName="Align Max", ToolTip="Adjusts the range so that the Root note is the same as the Max note of the zone"),
+	Centered	UMETA(DisplayName="Align Center", ToolTip="Adjusts the range so that the Root note of each keyzone is between its min and max note (Root = (Max + Min) / 2)")	
+};
+
+UCLASS(config = Editor)
+class UFusionPatchCreateOptions : public UObject
+{
+	GENERATED_BODY()
+
+public:
+	struct FArgs
+	{
+		FString Directory;
+		FString AssetName;
+		TArray<TWeakObjectPtr<USoundWave>> StagedSoundWaves;
+	};
+
+	static const UFusionPatchCreateOptions* GetWithDialog(FArgs&& Args, bool& OutWasOkayPressed);
+
+	// Which directory the FusionPatch should be created in
+	UPROPERTY(config, EditAnywhere, Category = "Create Patch", Meta=(DisplayName = "Destination", ContentDir))
+	FDirectoryPath FusionPatchDir;
+
+	// The Name of the new fusion patch
+	UPROPERTY(EditAnywhere, Category = "Create Patch", Meta=(DisplayName = ""))
+	FString AssetName;
+	
+	// how to parse and sort the names of the samples, extract the note numbers (if applicable) and assigning root, min, max notes
+	UPROPERTY(config, EditAnywhere, Category = "Create Patch", meta=(DisplayName="Assign Notes By"))
+	EFusionPatchKeyzoneSortOption SortOption = EFusionPatchKeyzoneSortOption::NoteNumber;
+
+	// how to assign the min and max notes for each keyzone based on ordering and note mapping
+	UPROPERTY(config, EditAnywhere, Category = "Create Patch", meta=(DisplayName="Keyzone Layout"))
+	EFusionPatchKeyzoneNoteLayoutOption LayoutOption = EFusionPatchKeyzoneNoteLayoutOption::Distribute;
+
+	// whether to assign the notes along a scale, or evenly spaced
+	UPROPERTY(Config, EditAnywhere, Category = "Create Patch", meta=(DisplayName="Scale Mapping"))
+	EFusionPatchKeyzoneNoteScaleOption ScaleOption = EFusionPatchKeyzoneNoteScaleOption::None;
+	
+	// how to assign the root note based on the min and max notes
+	UPROPERTY(config, EditAnywhere, Category = "Create Patch", meta=(DisplayName="Note Alignment"))
+	EFusionPatchKeyzoneRootNoteOption RootNoteOption = EFusionPatchKeyzoneRootNoteOption::Min;
+
+	UPROPERTY(config, EditAnywhere, Category = "Create Patch", meta=(ClampMin = 0, ClampMax = 127, UIMin = 0, UIMax = 127))
+	int8 MinNote = 0;
+
+	UPROPERTY(config, EditAnywhere, Category = "Create Patch", meta=(ClampMin = 0, ClampMax = 127, UIMin = 0, UIMax = 127))
+	int8 MaxNote = 127;
+
+	UPROPERTY(EditAnywhere, Category = "Create Patch")
+	TArray<FKeyzoneSettings> Keyzones;
+
+	UPROPERTY()
+	FFusionPatchSettings FusionPatchSettings;
+
+	TArray<TWeakObjectPtr<USoundWave>> StagedSoundWaves;
+
+	ELockedNoteFlag LockedNotesMask = ELockedNoteFlag::None;
+
+#if WITH_EDITOR
+	
+	/** UObject interface */
+	virtual void PostEditChangeProperty(struct FPropertyChangedEvent& PropertyChangedEvent) override;
+	
+	void SortLexical();
+
+	void SortNoteNumber();
+
+	void SortNoteName();
+
+	void SortIndex();
+
+	void LayoutSingleNote();
+
+	void LayoutDistribute();
+
+	void LayoutLayer();
+
+	void UpdateKeyzonesWithSettings();
+	
+	
+#endif
+};
+
+class FKeyzoneNoteParser
+{
+public:
+
+	struct FParseResult
+	{
+		int32 MinNote = -1;
+		int32 RootNote = -1;
+		int32 MaxNote = -1;
+
+		void Reset()
+		{
+			MinNote = -1;
+			RootNote = -1;
+			MaxNote = -1;
+		}
+	};
+
+	static bool ParseMinRootMax(const FString& Input, FParseResult& Output);
+
+	static bool ParseMinMax(const FString& Input, FParseResult& Output);
+		
+	static bool ParseRoot(const FString& Input, FParseResult& Output);
+
+	static bool ParseNoteName(const FString& Input, FParseResult& Output);
+
+	static bool ParseIndex(const FString& Input, FParseResult& Output);
+};
diff --git a/Engine/Plugins/Runtime/Harmonix/Source/HarmonixDspEditor/Private/HarmonixDspEditorModule.cpp b/Engine/Plugins/Runtime/Harmonix/Source/HarmonixDspEditor/Private/HarmonixDspEditorModule.cpp
--- a/Engine/Plugins/Runtime/Harmonix/Source/HarmonixDspEditor/Private/HarmonixDspEditorModule.cpp
+++ b/Engine/Plugins/Runtime/Harmonix/Source/HarmonixDspEditor/Private/HarmonixDspEditorModule.cpp
@@ -8,13 +8,17 @@
 #include "ISettingsModule.h"
 
 #include "PropertyEditorModule.h"
+#include "AssetDefinitions/AssetDefinition_FusionPatch.h"
 #include "Customization/TypedParameterCustomization.h"
 #include "Customization/PitchShifterNameCustomization.h"
 #include "Customization/PitchShifterConfigCustomization.h"
 #include "Customization/PannerDetailsCustomization.h"
 #include "Customization/FusionPatchDetailCustomization.h"
+#include "Customization/FusionPatchImportOptionsCustomization.h"
 #include "Customization/FusionPatchSettingsDetailCustomization.h"
 #include "HarmonixDsp/FusionSampler/FusionPatch.h"
+#include "ToolMenus.h"
+#include "Customization/FusionPatchImportOptionsCustomization.h"
 
 #define LOCTEXT_NAMESPACE "HarmonixDspEditor"
 
@@ -29,13 +33,14 @@ void FHarmonixDspEditorModule::StartupModule()
 	PropertyEditorModule.RegisterCustomPropertyTypeLayout("PannerDetails", FOnGetPropertyTypeCustomizationInstance::CreateStatic(&FPannerDetailsCustomization::MakeInstance));
 	PropertyEditorModule.RegisterCustomPropertyTypeLayout("FusionPatchSettings", FOnGetPropertyTypeCustomizationInstance::CreateStatic(&FFusionPatchSettingsDetailCustomization::MakeInstance));
 	PropertyEditorModule.RegisterCustomClassLayout("StretcherAndPitchShifterFactoryConfig", FOnGetDetailCustomizationInstance::CreateStatic(&FPitchShifterConfigCustomization::MakeInstance));
+	PropertyEditorModule.RegisterCustomClassLayout("FusionPatchCreateOptions", FOnGetDetailCustomizationInstance::CreateStatic(&FFusionPatchCreateOptionsCustomization::MakeInstance));
 	PropertyEditorModule.RegisterCustomClassLayout(UFusionPatch::StaticClass()->GetFName(), FOnGetDetailCustomizationInstance::CreateStatic(&FFusionPatchDetailCustomization::MakeInstance));
-	
+
+	UToolMenus::RegisterStartupCallback(FSimpleMulticastDelegate::FDelegate::CreateRaw(this, &FHarmonixDspEditorModule::RegisterMenus));
 }
 
 void FHarmonixDspEditorModule::ShutdownModule()
 {
-
 	FPropertyEditorModule& PropertyEditorModule = FModuleManager::LoadModuleChecked<FPropertyEditorModule>("PropertyEditor");
 	PropertyEditorModule.UnregisterCustomClassLayout("FusionPatch");
 	PropertyEditorModule.UnregisterCustomClassLayout("StretcherAndPitchShifterFactoryConfig");
@@ -45,6 +50,12 @@ void FHarmonixDspEditorModule::ShutdownModule()
 	PropertyEditorModule.UnregisterCustomPropertyTypeLayout("TypedParameter");
 }
 
+void FHarmonixDspEditorModule::RegisterMenus()
+{
+	FToolMenuOwnerScoped MenuOwner(this);
+	FFusionPatchExtension::RegisterMenus();
+}
+
 IMPLEMENT_MODULE(FHarmonixDspEditorModule, HarmonixDspEditor);
 
 #undef LOCTEXT_NAMESPACE
diff --git a/Engine/Plugins/Runtime/Harmonix/Source/HarmonixDspEditor/Public/HarmonixDspEditorModule.h b/Engine/Plugins/Runtime/Harmonix/Source/HarmonixDspEditor/Public/HarmonixDspEditorModule.h
--- a/Engine/Plugins/Runtime/Harmonix/Source/HarmonixDspEditor/Public/HarmonixDspEditorModule.h
+++ b/Engine/Plugins/Runtime/Harmonix/Source/HarmonixDspEditor/Public/HarmonixDspEditorModule.h
@@ -12,4 +12,7 @@ class FHarmonixDspEditorModule : public IModuleInterface
 public:
 	virtual void StartupModule();
 	virtual void ShutdownModule();
+
+private:
+	void RegisterMenus();
 };
