Use "Sound Wave" name for Keyzone Name dropdown since the SamplePath only exists if the sound wave was imported from a .fusion file.
#tests PIE

diff --git a/Engine/Plugins/Runtime/Harmonix/Source/HarmonixDspEditor/Private/Customization/FusionPatchDetailCustomization.cpp b/Engine/Plugins/Runtime/Harmonix/Source/HarmonixDspEditor/Private/Customization/FusionPatchDetailCustomization.cpp
--- a/Engine/Plugins/Runtime/Harmonix/Source/HarmonixDspEditor/Private/Customization/FusionPatchDetailCustomization.cpp
+++ b/Engine/Plugins/Runtime/Harmonix/Source/HarmonixDspEditor/Private/Customization/FusionPatchDetailCustomization.cpp
@@ -191,22 +191,25 @@ void FFusionPatchDetailCustomization::AddKeyzonesNameToMenuArray(TSharedPtr<IPro
 {
 	for (int32 KeyzoneIndex = 0; KeyzoneIndex < static_cast<int32>(NumKeyzones); ++KeyzoneIndex)
 	{
-		TSharedPtr<IPropertyHandle> SamplePathHandle = KeyzonesHandle->GetElement(KeyzoneIndex)->GetChildHandle(GET_MEMBER_NAME_CHECKED(FKeyzoneSettings, SamplePath));
-		FString CurrentSamplePath;
-		if (SamplePathHandle->GetValue(CurrentSamplePath) == FPropertyAccess::Success)
+		TSharedPtr<IPropertyHandle> SamplePathHandle = KeyzonesHandle->GetElement(KeyzoneIndex)->GetChildHandle(GET_MEMBER_NAME_CHECKED(FKeyzoneSettings, SoundWave));
+		FString SoundWaveName = FString::Printf(TEXT("KEYZONE_%d"), KeyzoneIndex);
+		
+		UObject* SoundWave;
+		if (SamplePathHandle->GetValue(SoundWave) == FPropertyAccess::Success && SoundWave != nullptr)
 		{
-			FString Path;
-			FString Filename;
-			FString Ext;
-			FPaths::Split(CurrentSamplePath, Path, Filename, Ext);
-			if (Filename.IsEmpty())
-			{
-				KeyzonesNameMenu.Add(MakeShareable(new FString(CurrentSamplePath)));
-			}
-			else
-			{
-				KeyzonesNameMenu.Add(MakeShareable(new FString(Filename)));
-			}
+			SoundWaveName = SoundWave->GetName();
+		}
+		FString Path;
+		FString Filename;
+		FString Ext;
+		FPaths::Split(SoundWaveName, Path, Filename, Ext);
+		if (Filename.IsEmpty())
+		{
+			KeyzonesNameMenu.Add(MakeShareable(new FString(SoundWaveName)));
+		}
+		else
+		{
+			KeyzonesNameMenu.Add(MakeShareable(new FString(Filename)));
 		}
 	}
 }
