[Harmonix] MIDI Pulse Node

#tests automated tests

#virtualized

diff --git a/Engine/Plugins/Runtime/Harmonix/Source/HarmonixMetasound/Private/HarmonixMetasound/DataTypes/MidiClock.cpp b/Engine/Plugins/Runtime/Harmonix/Source/HarmonixMetasound/Private/HarmonixMetasound/DataTypes/MidiClock.cpp
--- a/Engine/Plugins/Runtime/Harmonix/Source/HarmonixMetasound/Private/HarmonixMetasound/DataTypes/MidiClock.cpp
+++ b/Engine/Plugins/Runtime/Harmonix/Source/HarmonixMetasound/Private/HarmonixMetasound/DataTypes/MidiClock.cpp
@@ -32,6 +32,59 @@ namespace Metasound
 	DEFINE_METASOUND_ENUM_END()
 }
 
+float SubdivisionToBeats(EMidiClockSubdivisionQuantization Subdivision, const FTimeSignature& TimeSignature)
+{
+	// Easy cases first
+	if (Subdivision == EMidiClockSubdivisionQuantization::Bar)
+	{
+		return TimeSignature.Numerator;
+	}
+
+	if (Subdivision == EMidiClockSubdivisionQuantization::Beat)
+	{
+		return 1;
+	}
+
+	const float BeatsPerQuarter = TimeSignature.Denominator / 4.0f;
+	
+	switch (Subdivision)
+	{
+	case EMidiClockSubdivisionQuantization::ThirtySecondNote:
+		return BeatsPerQuarter / 8;
+	case EMidiClockSubdivisionQuantization::SixteenthNote:
+		return BeatsPerQuarter / 4;
+	case EMidiClockSubdivisionQuantization::EighthNote:
+		return BeatsPerQuarter / 2;
+	case EMidiClockSubdivisionQuantization::QuarterNote:
+		return BeatsPerQuarter;
+	case EMidiClockSubdivisionQuantization::HalfNote:
+		return BeatsPerQuarter * 2;
+	case EMidiClockSubdivisionQuantization::WholeNote:
+		return BeatsPerQuarter * 4;
+	case EMidiClockSubdivisionQuantization::DottedSixteenthNote:
+		return BeatsPerQuarter / 4 + BeatsPerQuarter / 8; 
+	case EMidiClockSubdivisionQuantization::DottedEighthNote:
+		return BeatsPerQuarter / 2 + BeatsPerQuarter / 4;
+	case EMidiClockSubdivisionQuantization::DottedQuarterNote:
+		return BeatsPerQuarter + BeatsPerQuarter / 2;
+	case EMidiClockSubdivisionQuantization::DottedHalfNote:
+		return BeatsPerQuarter * 3;
+	case EMidiClockSubdivisionQuantization::DottedWholeNote:
+		return BeatsPerQuarter * 6;
+	case EMidiClockSubdivisionQuantization::SixteenthNoteTriplet:
+		return (BeatsPerQuarter / 4) * 2 / 3;
+	case EMidiClockSubdivisionQuantization::EighthNoteTriplet:
+		return (BeatsPerQuarter / 2) * 2 / 3;
+	case EMidiClockSubdivisionQuantization::QuarterNoteTriplet:
+		return BeatsPerQuarter * 2 / 3;
+	case EMidiClockSubdivisionQuantization::HalfNoteTriplet:
+		return BeatsPerQuarter * 4 / 3;
+	default:
+		checkNoEntry();
+		return 0;
+	}
+}
+
 int32 SubdivisionToMidiTicks(const EMidiClockSubdivisionQuantization Division, const int32 CurrentTick, const FSongMaps& SongMap)
 {
 	int32 BarMapPointIndex = SongMap.GetBarMap().GetPointIndexForTick(CurrentTick);
diff --git a/Engine/Plugins/Runtime/Harmonix/Source/HarmonixMetasound/Private/HarmonixMetasound/DataTypes/MusicTimeInterval.cpp b/Engine/Plugins/Runtime/Harmonix/Source/HarmonixMetasound/Private/HarmonixMetasound/DataTypes/MusicTimeInterval.cpp
new file mode 100644
--- /dev/null
+++ b/Engine/Plugins/Runtime/Harmonix/Source/HarmonixMetasound/Private/HarmonixMetasound/DataTypes/MusicTimeInterval.cpp
@@ -0,0 +1,64 @@
+// Copyright Epic Games, Inc. All Rights Reserved.
+
+#include "HarmonixMetasound/DataTypes/MusicTimeInterval.h"
+
+namespace Harmonix
+{
+	float GetIntervalInBeats(const FMusicTimeInterval& Interval, const FTimeSignature& TimeSignature)
+	{
+		return SubdivisionToBeats(Interval.Interval, TimeSignature) * Interval.IntervalMultiplier;
+	}
+
+	void IncrementTimestampByBeats(FMusicTimestamp& Timestamp, float Beats, const FTimeSignature& TimeSignature)
+	{
+		// Add the beats
+		Timestamp.Beat += Beats;
+
+		// Add bars until the beat is within the bar
+		// Beats and bars are 1-indexed, so, for example, in 4/4, Beat 4.999... is within the bar.
+		while (Timestamp.Beat >= static_cast<float>(TimeSignature.Numerator + 1))
+		{
+			++Timestamp.Bar;
+			Timestamp.Beat -= TimeSignature.Numerator;
+		}
+	}
+	
+	void IncrementTimestampByInterval(
+		FMusicTimestamp& Timestamp,
+		const FMusicTimeInterval& Interval,
+		const FTimeSignature& TimeSignature)
+	{
+		// Get the pulse interval in beats
+		const float IntervalBeats = GetIntervalInBeats(Interval, TimeSignature);
+
+		if (FMath::IsNearlyZero(IntervalBeats))
+		{
+			return;
+		}
+
+		IncrementTimestampByBeats(Timestamp, IntervalBeats, TimeSignature);
+	}
+
+	float GetOffsetInBeats(const FMusicTimeInterval& Interval, const FTimeSignature& TimeSignature)
+	{
+		const float IntervalBeats = GetIntervalInBeats(Interval, TimeSignature);
+		const float OffsetBeats = SubdivisionToBeats(Interval.Offset, TimeSignature) * Interval.OffsetMultiplier;
+		return FMath::Fmod(OffsetBeats, IntervalBeats);
+	}
+
+	void IncrementTimestampByOffset(
+		FMusicTimestamp& Timestamp,
+		const FMusicTimeInterval& Interval,
+		const FTimeSignature& TimeSignature)
+	{
+		// Get the offset in beats
+		const float OffsetBeats = GetOffsetInBeats(Interval, TimeSignature);
+
+		if (FMath::IsNearlyZero(OffsetBeats))
+		{
+			return;
+		}
+
+		IncrementTimestampByBeats(Timestamp, OffsetBeats, TimeSignature);
+	}
+}
diff --git a/Engine/Plugins/Runtime/Harmonix/Source/HarmonixMetasound/Private/HarmonixMetasound/MidiOps/PulseGenerator.cpp b/Engine/Plugins/Runtime/Harmonix/Source/HarmonixMetasound/Private/HarmonixMetasound/MidiOps/PulseGenerator.cpp
new file mode 100644
--- /dev/null
+++ b/Engine/Plugins/Runtime/Harmonix/Source/HarmonixMetasound/Private/HarmonixMetasound/MidiOps/PulseGenerator.cpp
@@ -0,0 +1,159 @@
+// Copyright Epic Games, Inc. All Rights Reserved.
+
+#include "HarmonixMetasound/MidiOps/PulseGenerator.h"
+
+namespace Harmonix::Midi::Ops
+{
+	void FPulseGenerator::Enable(bool bEnable)
+	{
+		Enabled = bEnable;
+	}
+
+	void FPulseGenerator::SetClock(const HarmonixMetasound::FMidiClock* Clock)
+	{
+		Cursor.SetClock(Clock);
+	}
+
+	void FPulseGenerator::SetInterval(const FMusicTimeInterval& NewInterval)
+	{
+		Cursor.SetInterval(NewInterval);
+	}
+
+	void FPulseGenerator::Process(HarmonixMetasound::FMidiStream& OutStream)
+	{
+		// Keep draining the queue if disabled, so we get the next note off,
+		// and so we stay in phase if we toggle off and back on
+		FCursor::FPulseTime NextPulse;
+		while (Cursor.Pop(NextPulse))
+		{
+			int32 NoteOnSample = NextPulse.AudioBlockFrame;
+
+			// Note off if there was a previous note on
+			if (LastNoteOn.IsSet())
+			{
+				check(LastNoteOn->MidiMessage.IsNoteOn());
+				
+				// Trigger the note off one sample before the note on
+				const int32 NoteOffSample = NextPulse.AudioBlockFrame > 0 ? NextPulse.AudioBlockFrame - 1 : NextPulse.AudioBlockFrame;
+				NoteOnSample = NoteOffSample + 1;
+
+				// Trigger the note off one tick before the note on
+				const int32 NoteOffTick = NextPulse.MidiTick - 1;
+
+				FMidiMsg Msg{ FMidiMsg::CreateNoteOff(LastNoteOn->MidiMessage.GetStdChannel(), LastNoteOn->MidiMessage.GetStdData1()) };
+				HarmonixMetasound::FMidiStreamEvent Event{ &VoiceGenerator, Msg };
+				Event.BlockSampleFrameIndex = NoteOffSample;
+				Event.AuthoredMidiTick = NoteOffTick;
+				Event.CurrentMidiTick = NoteOffTick;
+				Event.TrackIndex = LastNoteOn->TrackIndex;
+				OutStream.InsertMidiEvent(Event);
+
+				LastNoteOn.Reset();
+			}
+			
+			// Note on
+			if (Enabled)
+			{
+				FMidiMsg Msg{ FMidiMsg::CreateNoteOn(Channel - 1, NoteNumber, Velocity) };
+				HarmonixMetasound::FMidiStreamEvent Event{ &VoiceGenerator, Msg };
+				Event.BlockSampleFrameIndex = NoteOnSample;
+				Event.AuthoredMidiTick = NextPulse.MidiTick;
+				Event.CurrentMidiTick = NextPulse.MidiTick;
+				Event.TrackIndex = Track;
+				OutStream.InsertMidiEvent(Event);
+
+				LastNoteOn.Emplace(MoveTemp(Event));
+			}
+		}
+	}
+
+	FPulseGenerator::FCursor::FCursor()
+	{
+		SetMessageFilter(EFilterPassFlags::TimeSig);
+	}
+
+	void FPulseGenerator::FCursor::SetClock(const HarmonixMetasound::FMidiClock* NewClock)
+	{
+		if (Clock == NewClock)
+		{
+			return;
+		}
+		
+		if (nullptr == NewClock)
+		{
+			if (nullptr != Owner)
+			{
+				Owner->UnregisterPlayCursor(this);
+			}
+		}
+		else
+		{
+			NewClock->RegisterHiResPlayCursor(this);
+		}
+
+		Clock = NewClock;
+	}
+
+	void FPulseGenerator::FCursor::SetInterval(const FMusicTimeInterval& NewInterval)
+	{
+		Interval = NewInterval;
+
+		// Multiplier should be >= 1
+		Interval.IntervalMultiplier = FMath::Max(Interval.IntervalMultiplier, static_cast<uint16>(1));
+	}
+
+	bool FPulseGenerator::FCursor::Pop(FPulseTime& PulseTime)
+	{
+		if (PulsesSinceLastProcess.Dequeue(PulseTime))
+		{
+			--QueueSize;
+			return true;
+		}
+
+		return false;
+	}
+
+	void FPulseGenerator::FCursor::Push(FPulseTime&& PulseTime)
+	{
+		constexpr int32 MaxQueueSize = 1024; // just to keep the queue from infinitely growing if we stop draining it
+		if (QueueSize < MaxQueueSize)
+		{
+			PulsesSinceLastProcess.Enqueue(MoveTemp(PulseTime));
+			++QueueSize;
+		}
+	}
+
+	void FPulseGenerator::FCursor::AdvanceThruTick(int32 Tick, bool IsPreRoll)
+	{
+		FMidiPlayCursor::AdvanceThruTick(Tick, IsPreRoll);
+
+		if (!NextPulseTimestamp.IsValid())
+		{
+			return;
+		}
+
+		check(nullptr != Clock);
+		int32 NextPulseTick = Clock->GetBarMap().MusicTimestampToTick(NextPulseTimestamp);
+
+		while (CurrentTick >= NextPulseTick)
+		{
+			Push({ Clock->GetCurrentBlockFrameIndex(), NextPulseTick });
+			
+			IncrementTimestampByInterval(NextPulseTimestamp, Interval, CurrentTimeSignature);
+
+			NextPulseTick = Clock->GetBarMap().MusicTimestampToTick(NextPulseTimestamp);
+		}
+	}
+
+	void FPulseGenerator::FCursor::OnTimeSig(int32 TrackIndex, int32 Tick, int32 Numerator, int32 Denominator, bool IsPreroll)
+	{
+		CurrentTimeSignature.Numerator = Numerator;
+		CurrentTimeSignature.Denominator = Denominator;
+		
+		check(nullptr != Owner);
+		// Time sig changes will come on the downbeat, and if we change time signature,
+		// we want to reset the pulse, so the next pulse is now plus the offset
+		NextPulseTimestamp = Owner->GetBarMap().TickToMusicTimestamp(Tick);
+		IncrementTimestampByOffset(NextPulseTimestamp, Interval, CurrentTimeSignature);
+	}
+}
diff --git a/Engine/Plugins/Runtime/Harmonix/Source/HarmonixMetasound/Private/HarmonixMetasound/Nodes/MidiPulseGeneratorNode.cpp b/Engine/Plugins/Runtime/Harmonix/Source/HarmonixMetasound/Private/HarmonixMetasound/Nodes/MidiPulseGeneratorNode.cpp
new file mode 100644
--- /dev/null
+++ b/Engine/Plugins/Runtime/Harmonix/Source/HarmonixMetasound/Private/HarmonixMetasound/Nodes/MidiPulseGeneratorNode.cpp
@@ -0,0 +1,226 @@
+// Copyright Epic Games, Inc. All Rights Reserved.
+
+#include "HarmonixMetasound/Nodes/MidiPulseGeneratorNode.h"
+
+#include "MetasoundExecutableOperator.h"
+#include "MetasoundFacade.h"
+#include "MetasoundNodeRegistrationMacro.h"
+#include "MetasoundStandardNodesCategories.h"
+
+#include "HarmonixMetasound/DataTypes/MidiClock.h"
+#include "HarmonixMetasound/MidiOps/PulseGenerator.h"
+
+#define LOCTEXT_NAMESPACE "HarmonixMetaSound_MidiPulseGeneratorNode"
+
+namespace HarmonixMetasound::Nodes::MidiPulseGeneratorNode
+{
+	const Metasound::FNodeClassName& GetClassName()
+	{
+		static const Metasound::FNodeClassName ClassName { HarmonixNodeNamespace, "MidiPulseGenerator", "" };
+		return ClassName;
+	}
+
+	int32 GetCurrentMajorVersion()
+	{
+		return 0;
+	}
+
+	namespace Inputs
+	{
+		DEFINE_METASOUND_PARAM_ALIAS(MidiClock, CommonPinNames::Inputs::MidiClock);
+		DEFINE_INPUT_METASOUND_PARAM(Interval, "Interval", "The musical time interval at which to send out the pulse");
+		DEFINE_INPUT_METASOUND_PARAM(IntervalMultiplier, "Interval Multiplier", "Multiplies the interval, 1 to use just the value of Interval");
+		DEFINE_INPUT_METASOUND_PARAM(Offset, "Offset", "Offsets the pulse by a musical time");
+		DEFINE_INPUT_METASOUND_PARAM(OffsetMultiplier, "Offset Multiplier", "Multiplies the offset, 0 for no offset");
+		DEFINE_METASOUND_PARAM_ALIAS(MidiTrack, CommonPinNames::Inputs::MidiTrackNumber);
+		DEFINE_METASOUND_PARAM_ALIAS(MidiChannel, CommonPinNames::Inputs::MidiChannelNumber);
+		DEFINE_INPUT_METASOUND_PARAM(MidiNoteNumber, "Note Number", "The note number to play");
+		DEFINE_INPUT_METASOUND_PARAM(MidiVelocity, "Velocity", "The velocity at which to play the note");
+	}
+
+	namespace Outputs
+	{
+		DEFINE_METASOUND_PARAM_ALIAS(MidiStream, CommonPinNames::Outputs::MidiStream);
+	}
+
+	class FMidiPulseGeneratorOperator final : public Metasound::TExecutableOperator<FMidiPulseGeneratorOperator>
+	{
+	public:
+		static const Metasound::FNodeClassMetadata& GetNodeInfo()
+		{
+			using namespace Metasound;
+			
+			auto InitNodeInfo = []() -> FNodeClassMetadata
+			{
+				FNodeClassMetadata Info;
+				Info.ClassName        = GetClassName();
+				Info.MajorVersion     = 0;
+				Info.MinorVersion     = 1;
+				Info.DisplayName      = METASOUND_LOCTEXT("MidiPulseGeneratorNode_DisplayName", "MIDI Pulse Generator");
+				Info.Description      = METASOUND_LOCTEXT("MidiPulseGeneratorNode_Description", "Outputs a repeated MIDI note at the specified musical time interval");
+				Info.Author           = PluginAuthor;
+				Info.PromptIfMissing  = PluginNodeMissingPrompt;
+				Info.DefaultInterface = GetVertexInterface();
+				Info.CategoryHierarchy.Emplace(NodeCategories::Music);
+				return Info;
+			};
+
+			static const FNodeClassMetadata Info = InitNodeInfo();
+
+			return Info;
+		}
+		
+		static const Metasound::FVertexInterface& GetVertexInterface()
+		{
+			using namespace Metasound;
+
+			static const Harmonix::Midi::Ops::FPulseGenerator PulseGeneratorForDefaults;
+			const auto DefaultInterval = PulseGeneratorForDefaults.GetInterval();
+
+			static const FVertexInterface Interface
+			{
+				FInputVertexInterface {
+					TInputDataVertex<FMidiClock>(METASOUND_GET_PARAM_NAME_AND_METADATA(Inputs::MidiClock)),
+					TInputDataVertex<FEnumMidiClockSubdivisionQuantizationType>(
+						METASOUND_GET_PARAM_NAME_AND_METADATA(Inputs::Interval), static_cast<int32>(DefaultInterval.Interval)),
+					TInputDataVertex<int32>(
+						METASOUND_GET_PARAM_NAME_AND_METADATA(Inputs::IntervalMultiplier), static_cast<int32>(DefaultInterval.IntervalMultiplier)),
+					TInputDataVertex<FEnumMidiClockSubdivisionQuantizationType>(
+						METASOUND_GET_PARAM_NAME_AND_METADATA(Inputs::Offset), static_cast<int32>(DefaultInterval.Offset)),
+					TInputDataVertex<int32>(
+						METASOUND_GET_PARAM_NAME_AND_METADATA(Inputs::OffsetMultiplier), static_cast<int32>(DefaultInterval.OffsetMultiplier)),
+					TInputDataVertex<int32>(
+						METASOUND_GET_PARAM_NAME_AND_METADATA(Inputs::MidiTrack), static_cast<int32>(PulseGeneratorForDefaults.Track)),
+					TInputDataVertex<int32>(
+						METASOUND_GET_PARAM_NAME_AND_METADATA(Inputs::MidiChannel), static_cast<int32>(PulseGeneratorForDefaults.Channel)),
+					TInputDataVertex<int32>(
+						METASOUND_GET_PARAM_NAME_AND_METADATA(Inputs::MidiNoteNumber), static_cast<int32>(PulseGeneratorForDefaults.NoteNumber)),
+					TInputDataVertex<int32>(
+						METASOUND_GET_PARAM_NAME_AND_METADATA(Inputs::MidiVelocity), static_cast<int32>(PulseGeneratorForDefaults.Velocity))
+				},
+				FOutputVertexInterface
+				{
+					TOutputDataVertex<FMidiStream>(METASOUND_GET_PARAM_NAME_AND_METADATA(Outputs::MidiStream))
+				}
+			};
+
+			return Interface;
+		}
+
+		struct FInputs
+		{
+			FMidiClockReadRef Clock;
+			Metasound::FEnumMidiClockSubdivisionQuantizationReadRef Interval;
+			Metasound::FInt32ReadRef IntervalMultiplier;
+			Metasound::FEnumMidiClockSubdivisionQuantizationReadRef Offset;
+			Metasound::FInt32ReadRef OffsetMultiplier;
+			Metasound::FInt32ReadRef Track;
+			Metasound::FInt32ReadRef Channel;
+			Metasound::FInt32ReadRef NoteNumber;
+			Metasound::FInt32ReadRef Velocity;
+		};
+
+		struct FOutputs
+		{
+			FMidiStreamWriteRef MidiStream;
+		};
+		
+		static TUniquePtr<IOperator> CreateOperator(const Metasound::FBuildOperatorParams& InParams, Metasound::FBuildResults& OutResults)
+		{
+			FInputs Inputs
+			{
+				InParams.InputData.GetOrConstructDataReadReference<FMidiClock>(Inputs::MidiClockName, InParams.OperatorSettings),
+				InParams.InputData.GetOrCreateDefaultDataReadReference<Metasound::FEnumMidiClockSubdivisionQuantizationType>(
+					Inputs::IntervalName, InParams.OperatorSettings),
+				InParams.InputData.GetOrCreateDefaultDataReadReference<int32>(Inputs::IntervalMultiplierName, InParams.OperatorSettings),
+				InParams.InputData.GetOrCreateDefaultDataReadReference<Metasound::FEnumMidiClockSubdivisionQuantizationType>(
+					Inputs::OffsetName, InParams.OperatorSettings),
+				InParams.InputData.GetOrCreateDefaultDataReadReference<int32>(Inputs::OffsetMultiplierName, InParams.OperatorSettings),
+				InParams.InputData.GetOrCreateDefaultDataReadReference<int32>(Inputs::MidiTrackName, InParams.OperatorSettings),
+				InParams.InputData.GetOrCreateDefaultDataReadReference<int32>(Inputs::MidiChannelName, InParams.OperatorSettings),
+				InParams.InputData.GetOrCreateDefaultDataReadReference<int32>(Inputs::MidiNoteNumberName, InParams.OperatorSettings),
+				InParams.InputData.GetOrCreateDefaultDataReadReference<int32>(Inputs::MidiVelocityName, InParams.OperatorSettings),
+			};
+
+			FOutputs Outputs
+			{
+				FMidiStreamWriteRef::CreateNew(InParams.OperatorSettings)
+			};
+
+			return MakeUnique<FMidiPulseGeneratorOperator>(InParams, MoveTemp(Inputs), MoveTemp(Outputs));
+		}
+
+		FMidiPulseGeneratorOperator(const Metasound::FBuildOperatorParams& InParams, FInputs&& InInputs, FOutputs&& InOutputs)
+			: Inputs(MoveTemp(InInputs))
+			, Outputs(MoveTemp(InOutputs))
+		{
+			Reset(InParams);
+		}
+
+		void Reset(const FResetParams&)
+		{
+			PulseGenerator.SetClock(Inputs.Clock.Get());
+			ApplyParameters();
+		}
+
+		virtual void BindInputs(Metasound::FInputVertexInterfaceData& InVertexData) override
+		{
+			InVertexData.BindReadVertex(Inputs::MidiClockName, Inputs.Clock);
+			InVertexData.BindReadVertex(Inputs::IntervalName, Inputs.Interval);
+			InVertexData.BindReadVertex(Inputs::IntervalMultiplierName, Inputs.IntervalMultiplier);
+			InVertexData.BindReadVertex(Inputs::OffsetName, Inputs.Offset);
+			InVertexData.BindReadVertex(Inputs::OffsetMultiplierName, Inputs.OffsetMultiplier);
+			InVertexData.BindReadVertex(Inputs::MidiTrackName, Inputs.Track);
+			InVertexData.BindReadVertex(Inputs::MidiChannelName, Inputs.Channel);
+			InVertexData.BindReadVertex(Inputs::MidiNoteNumberName, Inputs.NoteNumber);
+			InVertexData.BindReadVertex(Inputs::MidiVelocityName, Inputs.Velocity);
+
+			PulseGenerator.SetClock(Inputs.Clock.Get());
+			ApplyParameters();
+		}
+
+		virtual void BindOutputs(Metasound::FOutputVertexInterfaceData& InVertexData) override
+		{
+			InVertexData.BindReadVertex(Outputs::MidiStreamName, Outputs.MidiStream);
+		}
+
+		void Execute()
+		{
+			ApplyParameters();
+
+			Outputs.MidiStream->PrepareBlock();
+			PulseGenerator.Process(*Outputs.MidiStream);
+		}
+	private:
+		void ApplyParameters()
+		{
+			PulseGenerator.Track = *Inputs.Track;
+			PulseGenerator.Channel = *Inputs.Channel;
+			PulseGenerator.NoteNumber = *Inputs.NoteNumber;
+			PulseGenerator.Velocity = *Inputs.Velocity;
+			PulseGenerator.SetInterval(
+				{
+					*Inputs.Interval,
+					*Inputs.Offset,
+					static_cast<uint16>(*Inputs.IntervalMultiplier),
+					static_cast<uint16>(*Inputs.OffsetMultiplier)
+				});
+		}
+		
+		FInputs Inputs;
+		FOutputs Outputs;
+		Harmonix::Midi::Ops::FPulseGenerator PulseGenerator;
+	};
+
+	class FMidiPulseGeneratorNode final : public Metasound::FNodeFacade
+	{
+	public:
+		explicit FMidiPulseGeneratorNode(const Metasound::FNodeInitData& InInitData)
+			: FNodeFacade(InInitData.InstanceName, InInitData.InstanceID, Metasound::TFacadeOperatorClass<FMidiPulseGeneratorOperator>())
+		{}
+	};
+
+	METASOUND_REGISTER_NODE(FMidiPulseGeneratorNode)
+}
+
+#undef LOCTEXT_NAMESPACE
diff --git a/Engine/Plugins/Runtime/Harmonix/Source/HarmonixMetasound/Public/HarmonixMetasound/DataTypes/MidiClock.h b/Engine/Plugins/Runtime/Harmonix/Source/HarmonixMetasound/Public/HarmonixMetasound/DataTypes/MidiClock.h
--- a/Engine/Plugins/Runtime/Harmonix/Source/HarmonixMetasound/Public/HarmonixMetasound/DataTypes/MidiClock.h
+++ b/Engine/Plugins/Runtime/Harmonix/Source/HarmonixMetasound/Public/HarmonixMetasound/DataTypes/MidiClock.h
@@ -20,25 +20,25 @@
 #include "HarmonixMidi/MidiPlayCursor.h"
 #include "Sound/QuartzQuantizationUtilities.h"
 
-enum class EMidiClockSubdivisionQuantization : int32
+enum class EMidiClockSubdivisionQuantization : uint8
 {
-	Bar						= static_cast<int32>(EQuartzCommandQuantization::Bar),
-	Beat					= static_cast<int32>(EQuartzCommandQuantization::Beat),
-	ThirtySecondNote		= static_cast<int32>(EQuartzCommandQuantization::ThirtySecondNote),
-	SixteenthNote			= static_cast<int32>(EQuartzCommandQuantization::SixteenthNote),
-	EighthNote				= static_cast<int32>(EQuartzCommandQuantization::EighthNote),
-	QuarterNote				= static_cast<int32>(EQuartzCommandQuantization::QuarterNote),
-	HalfNote				= static_cast<int32>(EQuartzCommandQuantization::HalfNote),
-	WholeNote				= static_cast<int32>(EQuartzCommandQuantization::WholeNote),
-	DottedSixteenthNote		= static_cast<int32>(EQuartzCommandQuantization::DottedSixteenthNote),
-	DottedEighthNote		= static_cast<int32>(EQuartzCommandQuantization::DottedEighthNote),
-	DottedQuarterNote		= static_cast<int32>(EQuartzCommandQuantization::DottedQuarterNote),
-	DottedHalfNote			= static_cast<int32>(EQuartzCommandQuantization::DottedHalfNote),
-	DottedWholeNote			= static_cast<int32>(EQuartzCommandQuantization::DottedWholeNote),
-	SixteenthNoteTriplet	= static_cast<int32>(EQuartzCommandQuantization::SixteenthNoteTriplet),
-	EighthNoteTriplet		= static_cast<int32>(EQuartzCommandQuantization::EighthNoteTriplet),
-	QuarterNoteTriplet		= static_cast<int32>(EQuartzCommandQuantization::QuarterNoteTriplet),
-	HalfNoteTriplet			= static_cast<int32>(EQuartzCommandQuantization::HalfNoteTriplet),
+	Bar						= static_cast<uint8>(EQuartzCommandQuantization::Bar),
+	Beat					= static_cast<uint8>(EQuartzCommandQuantization::Beat),
+	ThirtySecondNote		= static_cast<uint8>(EQuartzCommandQuantization::ThirtySecondNote),
+	SixteenthNote			= static_cast<uint8>(EQuartzCommandQuantization::SixteenthNote),
+	EighthNote				= static_cast<uint8>(EQuartzCommandQuantization::EighthNote),
+	QuarterNote				= static_cast<uint8>(EQuartzCommandQuantization::QuarterNote),
+	HalfNote				= static_cast<uint8>(EQuartzCommandQuantization::HalfNote),
+	WholeNote				= static_cast<uint8>(EQuartzCommandQuantization::WholeNote),
+	DottedSixteenthNote		= static_cast<uint8>(EQuartzCommandQuantization::DottedSixteenthNote),
+	DottedEighthNote		= static_cast<uint8>(EQuartzCommandQuantization::DottedEighthNote),
+	DottedQuarterNote		= static_cast<uint8>(EQuartzCommandQuantization::DottedQuarterNote),
+	DottedHalfNote			= static_cast<uint8>(EQuartzCommandQuantization::DottedHalfNote),
+	DottedWholeNote			= static_cast<uint8>(EQuartzCommandQuantization::DottedWholeNote),
+	SixteenthNoteTriplet	= static_cast<uint8>(EQuartzCommandQuantization::SixteenthNoteTriplet),
+	EighthNoteTriplet		= static_cast<uint8>(EQuartzCommandQuantization::EighthNoteTriplet),
+	QuarterNoteTriplet		= static_cast<uint8>(EQuartzCommandQuantization::QuarterNoteTriplet),
+	HalfNoteTriplet			= static_cast<uint8>(EQuartzCommandQuantization::HalfNoteTriplet),
 };
 
 namespace Metasound
@@ -54,6 +54,8 @@ namespace Metasound
 	);
 }
 
+HARMONIXMETASOUND_API float SubdivisionToBeats(EMidiClockSubdivisionQuantization Subdivision, const FTimeSignature& TimeSignature);
+
 HARMONIXMETASOUND_API int32 SubdivisionToMidiTicks(EMidiClockSubdivisionQuantization Division, int32 CurrentTick, const FSongMaps& SongMap);
 
 namespace HarmonixMetasound
@@ -84,6 +86,8 @@ namespace HarmonixMetasound
 	class HARMONIXMETASOUND_API FMidiClock
 	{
 	public:
+		static constexpr int32 kMidiGranularity = 128;
+		
 		explicit FMidiClock(const Metasound::FOperatorSettings& InSettings);
 		FMidiClock(const FMidiClock& Other);
 		FMidiClock(FMidiClock&& Other);
@@ -242,7 +246,6 @@ namespace HarmonixMetasound
 		float SampleRate;
 		Metasound::FSampleCount SampleCount;
 		int32 FramesUntilNextProcess = 0;
-		static const int32 kMidiGranularity = 128;
 		FMidiTimestampTransportState CurrentTransportState;
 
 		TArray<FMidiTimestampTransportState> TransportChangesInBlock;
diff --git a/Engine/Plugins/Runtime/Harmonix/Source/HarmonixMetasound/Public/HarmonixMetasound/DataTypes/MusicTimeInterval.h b/Engine/Plugins/Runtime/Harmonix/Source/HarmonixMetasound/Public/HarmonixMetasound/DataTypes/MusicTimeInterval.h
new file mode 100644
--- /dev/null
+++ b/Engine/Plugins/Runtime/Harmonix/Source/HarmonixMetasound/Public/HarmonixMetasound/DataTypes/MusicTimeInterval.h
@@ -0,0 +1,26 @@
+// Copyright Epic Games, Inc. All Rights Reserved.
+
+#pragma once
+
+#include "MidiClock.h"
+
+namespace Harmonix
+{
+	struct HARMONIXMETASOUND_API FMusicTimeInterval
+	{
+		EMidiClockSubdivisionQuantization Interval = EMidiClockSubdivisionQuantization::Beat;
+		EMidiClockSubdivisionQuantization Offset = EMidiClockSubdivisionQuantization::Beat;
+		uint16 IntervalMultiplier = 1;
+		uint16 OffsetMultiplier = 0;
+	};
+
+	HARMONIXMETASOUND_API void IncrementTimestampByInterval(
+			FMusicTimestamp& Timestamp,
+			const FMusicTimeInterval& Interval,
+			const FTimeSignature& TimeSignature);
+
+	HARMONIXMETASOUND_API void IncrementTimestampByOffset(
+			FMusicTimestamp& Timestamp,
+			const FMusicTimeInterval& Interval,
+			const FTimeSignature& TimeSignature);
+}
diff --git a/Engine/Plugins/Runtime/Harmonix/Source/HarmonixMetasound/Public/HarmonixMetasound/MidiOps/PulseGenerator.h b/Engine/Plugins/Runtime/Harmonix/Source/HarmonixMetasound/Public/HarmonixMetasound/MidiOps/PulseGenerator.h
new file mode 100644
--- /dev/null
+++ b/Engine/Plugins/Runtime/Harmonix/Source/HarmonixMetasound/Public/HarmonixMetasound/MidiOps/PulseGenerator.h
@@ -0,0 +1,79 @@
+// Copyright Epic Games, Inc. All Rights Reserved.
+
+#pragma once
+
+#include "Containers/SpscQueue.h"
+
+#include "HarmonixDsp/Parameters/Parameter.h"
+
+#include "HarmonixMetasound/DataTypes/MidiClock.h"
+#include "HarmonixMetasound/DataTypes/MidiStream.h"
+#include "HarmonixMetasound/DataTypes/MusicTimeInterval.h"
+
+#include "HarmonixMidi/MidiVoiceId.h"
+
+namespace Harmonix::Midi::Ops
+{
+	class HARMONIXMETASOUND_API FPulseGenerator
+	{
+	public:
+		void Enable(bool bEnable);
+		
+		void SetClock(const HarmonixMetasound::FMidiClock* Clock);
+
+		Dsp::Parameters::TParameter<uint8> Channel{ 1, 16, 1 };
+
+		Dsp::Parameters::TParameter<uint16> Track{ 1, UINT16_MAX, 1 };
+		
+		Dsp::Parameters::TParameter<uint8> NoteNumber{ 0, 127, 60 };
+		
+		Dsp::Parameters::TParameter<uint8> Velocity{ 0, 127, 127 };
+		
+		void SetInterval(const FMusicTimeInterval& NewInterval);
+
+		FMusicTimeInterval GetInterval() const { return Cursor.GetInterval(); }
+
+		void Process(HarmonixMetasound::FMidiStream& OutStream);
+
+	private:
+		FMidiVoiceGeneratorBase VoiceGenerator{};
+		TOptional<HarmonixMetasound::FMidiStreamEvent> LastNoteOn;
+		bool Enabled{ true };
+
+		class FCursor final : public FMidiPlayCursor
+		{
+		public:
+			FCursor();
+
+			void SetClock(const HarmonixMetasound::FMidiClock* NewClock);
+
+			void SetInterval(const FMusicTimeInterval& NewInterval);
+
+			FMusicTimeInterval GetInterval() const { return Interval; }
+
+			struct FPulseTime
+			{
+				int32 AudioBlockFrame;
+				int32 MidiTick;
+			};
+
+			bool Pop(FPulseTime& PulseTime);
+
+		private:
+			void Push(FPulseTime&& PulseTime);
+			
+			virtual void AdvanceThruTick(int32 Tick, bool IsPreRoll) override;
+			
+			virtual void OnTimeSig(int32 TrackIndex, int32 Tick, int32 Numerator, int32 Denominator, bool IsPreroll) override;
+
+			const HarmonixMetasound::FMidiClock* Clock{ nullptr };
+			FMusicTimeInterval Interval{};
+			TSpscQueue<FPulseTime> PulsesSinceLastProcess;
+			int32 QueueSize{0};
+			FTimeSignature CurrentTimeSignature{};
+			FMusicTimestamp NextPulseTimestamp{ -1, -1 };
+		};
+		
+		FCursor Cursor;
+	};
+}
diff --git a/Engine/Plugins/Runtime/Harmonix/Source/HarmonixMetasound/Public/HarmonixMetasound/Nodes/MidiPulseGeneratorNode.h b/Engine/Plugins/Runtime/Harmonix/Source/HarmonixMetasound/Public/HarmonixMetasound/Nodes/MidiPulseGeneratorNode.h
new file mode 100644
--- /dev/null
+++ b/Engine/Plugins/Runtime/Harmonix/Source/HarmonixMetasound/Public/HarmonixMetasound/Nodes/MidiPulseGeneratorNode.h
@@ -0,0 +1,31 @@
+// Copyright Epic Games, Inc. All Rights Reserved.
+
+#pragma once
+
+#include "MetasoundNodeInterface.h"
+
+#include "HarmonixMetasound/Common.h"
+
+namespace HarmonixMetasound::Nodes::MidiPulseGeneratorNode
+{
+	const HARMONIXMETASOUND_API Metasound::FNodeClassName& GetClassName();
+	HARMONIXMETASOUND_API int32 GetCurrentMajorVersion();
+
+	namespace Inputs
+	{
+		DECLARE_METASOUND_PARAM_EXTERN(MidiClock);
+		DECLARE_METASOUND_PARAM_EXTERN(Interval);
+		DECLARE_METASOUND_PARAM_EXTERN(IntervalMultiplier);
+		DECLARE_METASOUND_PARAM_EXTERN(Offset);
+		DECLARE_METASOUND_PARAM_EXTERN(OffsetMultiplier);
+		DECLARE_METASOUND_PARAM_EXTERN(MidiTrack);
+		DECLARE_METASOUND_PARAM_EXTERN(MidiChannel);
+		DECLARE_METASOUND_PARAM_EXTERN(MidiNoteNumber);
+		DECLARE_METASOUND_PARAM_EXTERN(MidiVelocity);
+	}
+
+	namespace Outputs
+	{
+		DECLARE_METASOUND_PARAM_EXTERN(MidiStream);
+	}
+}
\ No newline at end of file
diff --git a/Engine/Plugins/Runtime/Harmonix/Source/HarmonixMetasoundTests/Private/MidiOps/PulseGeneratorTests.cpp b/Engine/Plugins/Runtime/Harmonix/Source/HarmonixMetasoundTests/Private/MidiOps/PulseGeneratorTests.cpp
new file mode 100644
--- /dev/null
+++ b/Engine/Plugins/Runtime/Harmonix/Source/HarmonixMetasoundTests/Private/MidiOps/PulseGeneratorTests.cpp
@@ -0,0 +1,127 @@
+// Copyright Epic Games, Inc. All Rights Reserved.
+
+#include "HarmonixMetasound/MidiOps/PulseGenerator.h"
+#include "Misc/AutomationTest.h"
+
+#if WITH_DEV_AUTOMATION_TESTS
+
+namespace Harmonix::Midi::Ops::Tests
+{
+	IMPLEMENT_SIMPLE_AUTOMATION_TEST(
+		FMidiPulseGeneratorBasicTest,
+		"Harmonix.Midi.Ops.PulseGenerator.Basic",
+		EAutomationTestFlags::EditorContext | EAutomationTestFlags::EngineFilter)
+	bool FMidiPulseGeneratorBasicTest::RunTest(const FString&)
+	{
+		FPulseGenerator PulseGenerator;
+		
+		Metasound::FOperatorSettings OperatorSettings{ 48000, 100 };
+
+		constexpr float Tempo = 123;
+		constexpr uint8 TimeSigNumerator = 4;
+		constexpr uint8 TimeSigDenominator = 4;
+		HarmonixMetasound::FMidiClock Clock{ OperatorSettings };
+		Clock.AttachToMidiResource(HarmonixMetasound::FMidiClock::MakeClockConductorMidiData(Tempo, TimeSigNumerator, TimeSigDenominator));
+		
+		HarmonixMetasound::FMidiStream OutputStream{ OperatorSettings };
+
+		PulseGenerator.SetClock(&Clock);
+
+		// Default: a pulse every beat
+		{
+			constexpr int32 NotesUntilWeAreSatisfiedThisWorks = 23;
+			int32 NumNotesReceived = 0;
+
+			Clock.ResetAndStart(0);
+
+			FTimeSignature TimeSignature = Clock.GetBarMap().GetTimeSignatureAtTick(0);
+			FMusicTimeInterval Interval = PulseGenerator.GetInterval();
+			FMusicTimestamp NextPulse{ 1, 1 };
+			IncrementTimestampByOffset(NextPulse, Interval, TimeSignature);
+
+			int32 EndTick = 0;
+			{
+				FMusicTimestamp EndTimestamp{ NextPulse };
+				for (int i = 0; i < NotesUntilWeAreSatisfiedThisWorks; ++i)
+				{
+					IncrementTimestampByInterval(EndTimestamp, Interval, TimeSignature);
+				}
+				EndTick = Clock.GetBarMap().MusicTimestampToTick(EndTimestamp);
+			}
+			
+			while (Clock.GetCurrentMidiTick() < EndTick)
+			{
+				// Advance the clock, which will advance the play cursor in the pulse generator
+				Clock.WriteAdvance(0, OperatorSettings.GetNumFramesPerBlock());
+
+				// Process, which will pop the next notes
+				OutputStream.PrepareBlock();
+				PulseGenerator.Process(OutputStream);
+				
+				// If this is a block where we should get a pulse, check that we got it
+				if (Clock.GetCurrentMidiTick() >= Clock.GetBarMap().MusicTimestampToTick(NextPulse))
+				{
+					const bool ShouldGetNoteOff = NumNotesReceived > 0;
+
+					bool GotNoteOn = false;
+					bool GotNoteOff = false;
+
+					const TArray<HarmonixMetasound::FMidiStreamEvent> Events = OutputStream.GetEventsInBlock();
+
+					if (Events.Num() != (ShouldGetNoteOff ? 2 : 1))
+					{
+						return false;
+					}
+
+					UTEST_EQUAL("Got the right number of events", Events.Num(), ShouldGetNoteOff ? 2 : 1);
+
+					for (const HarmonixMetasound::FMidiStreamEvent& Event : Events)
+					{
+						if (Event.MidiMessage.IsNoteOn())
+						{
+							GotNoteOn = true;
+
+							UTEST_EQUAL("Right track", Event.TrackIndex, PulseGenerator.Track);
+							UTEST_EQUAL("Right channel", Event.MidiMessage.GetStdChannel(), PulseGenerator.Channel - 1);
+							UTEST_EQUAL("Right note number", Event.MidiMessage.GetStdData1(), PulseGenerator.NoteNumber);
+							UTEST_EQUAL("Right velocity", Event.MidiMessage.GetStdData2(), PulseGenerator.Velocity);
+
+							++NumNotesReceived;
+						}
+						else if (Event.MidiMessage.IsNoteOff())
+						{
+							GotNoteOff = true;
+
+							UTEST_EQUAL("Right track", Event.TrackIndex, PulseGenerator.Track);
+							UTEST_EQUAL("Right channel", Event.MidiMessage.GetStdChannel(), PulseGenerator.Channel - 1);
+							UTEST_EQUAL("Right note number", Event.MidiMessage.GetStdData1(), PulseGenerator.NoteNumber);
+						}
+						else
+						{
+							UTEST_TRUE("Unexpected event", false);
+						}
+					}
+
+					UTEST_TRUE("Got note on", GotNoteOn);
+
+					if (ShouldGetNoteOff)
+					{
+						UTEST_TRUE("Got note off", GotNoteOff);
+					}
+					else
+					{
+						UTEST_FALSE("Did not get note off", GotNoteOff);
+					}
+
+					IncrementTimestampByInterval(NextPulse, Interval, TimeSignature);
+				}
+			}
+			
+			UTEST_TRUE("Got all the notes at the right time", NumNotesReceived >= NotesUntilWeAreSatisfiedThisWorks);
+		}
+
+		return true;
+	}
+}
+
+#endif
\ No newline at end of file
diff --git a/Engine/Plugins/Runtime/Harmonix/Source/HarmonixMetasoundTests/Private/Nodes/MidiPulseGeneratorNodeTests.cpp b/Engine/Plugins/Runtime/Harmonix/Source/HarmonixMetasoundTests/Private/Nodes/MidiPulseGeneratorNodeTests.cpp
new file mode 100644
--- /dev/null
+++ b/Engine/Plugins/Runtime/Harmonix/Source/HarmonixMetasoundTests/Private/Nodes/MidiPulseGeneratorNodeTests.cpp
@@ -0,0 +1,65 @@
+// Copyright Epic Games, Inc. All Rights Reserved.
+
+#include "NodeTestGraphBuilder.h"
+#include "HarmonixMetasound/MidiOps/PulseGenerator.h"
+#include "HarmonixMetasound/Nodes/MidiPulseGeneratorNode.h"
+#include "Misc/AutomationTest.h"
+
+#if WITH_DEV_AUTOMATION_TESTS
+
+namespace HarmonixMetasound::Nodes::MidiPulseGeneratorNode::Tests
+{
+	IMPLEMENT_SIMPLE_AUTOMATION_TEST(
+	FMidiPulseGeneratorNodeParityTest,
+	"Harmonix.Metasound.Nodes.MidiPulseGenerator.Parity",
+	EAutomationTestFlags::EditorContext | EAutomationTestFlags::EngineFilter)
+	bool FMidiPulseGeneratorNodeParityTest::RunTest(const FString&)
+	{
+		const TUniquePtr<Metasound::FMetasoundGenerator> Generator =
+			Metasound::Test::FNodeTestGraphBuilder::MakeSingleNodeGraph(GetClassName(), GetCurrentMajorVersion());
+		TOptional<Metasound::TDataReadReference<FMidiStream>> NodeMidiOutput = Generator->GetOutputReadReference<FMidiStream>(Outputs::MidiStreamName);
+		UTEST_TRUE("Got node MIDI output", NodeMidiOutput.IsSet());
+
+		TOptional<FMidiClockWriteRef> Clock = Generator->GetInputWriteReference<FMidiClock>(Inputs::MidiClockName);
+		UTEST_TRUE("Got clock", Clock.IsSet());
+
+		Harmonix::Midi::Ops::FPulseGenerator PulseGenerator;
+		FMidiStream PulseGeneratorMidiOutput{ Generator->OperatorSettings };
+		PulseGenerator.SetClock(Clock->Get());
+
+		// Render for a bit and expect the same output from both the node and the raw processor
+		constexpr int32 NumBlocks = 1000;
+		Metasound::FAudioBuffer Buffer{ Generator->OperatorSettings };
+		
+		for (int32 BlockIdx = 0; BlockIdx < NumBlocks; ++BlockIdx)
+		{
+			// Advance the clock, which will advance the play cursor in the pulse generators
+			(*Clock)->WriteAdvance(0, Generator->OperatorSettings.GetNumFramesPerBlock());
+
+			// Process
+			Generator->OnGenerateAudio(Buffer.GetData(), Buffer.Num());
+			PulseGeneratorMidiOutput.PrepareBlock();
+			PulseGenerator.Process(PulseGeneratorMidiOutput);
+
+			// If there are notes in the pulse generator output, expect them in the node output
+			const TArray<FMidiStreamEvent> NodeEvents = (*NodeMidiOutput)->GetEventsInBlock();
+			const TArray<FMidiStreamEvent> PulseGeneratorEvents = PulseGeneratorMidiOutput.GetEventsInBlock();
+
+			for (const FMidiStreamEvent& PulseGeneratorEvent : PulseGeneratorEvents)
+			{
+				UTEST_TRUE("Event is in node output", NodeEvents.ContainsByPredicate([&PulseGeneratorEvent](const FMidiStreamEvent& NodeEvent)
+				{
+					return NodeEvent.BlockSampleFrameIndex == PulseGeneratorEvent.BlockSampleFrameIndex
+					&& NodeEvent.AuthoredMidiTick == PulseGeneratorEvent.AuthoredMidiTick
+					&& NodeEvent.CurrentMidiTick == PulseGeneratorEvent.CurrentMidiTick
+					&& NodeEvent.TrackIndex == PulseGeneratorEvent.TrackIndex
+					&& NodeEvent.MidiMessage == PulseGeneratorEvent.MidiMessage;
+				}));
+			}
+		}
+		
+		return true;
+	}
+}
+
+#endif
\ No newline at end of file
diff --git a/Engine/Plugins/Runtime/Harmonix/Source/HarmonixMidi/Public/HarmonixMidi/BarMap.h b/Engine/Plugins/Runtime/Harmonix/Source/HarmonixMidi/Public/HarmonixMidi/BarMap.h
--- a/Engine/Plugins/Runtime/Harmonix/Source/HarmonixMidi/Public/HarmonixMidi/BarMap.h
+++ b/Engine/Plugins/Runtime/Harmonix/Source/HarmonixMidi/Public/HarmonixMidi/BarMap.h
@@ -9,7 +9,7 @@
 class UTempoMap;
 
 USTRUCT(BlueprintType)
-struct FMusicTimestamp
+struct HARMONIXMIDI_API FMusicTimestamp
 {
 	GENERATED_BODY()
 
