[Harmonix] Set IsBetaVersion to false, in case it collides with IsExperimentalVersion

diff --git a/Engine/Plugins/Runtime/Harmonix/Harmonix.uplugin b/Engine/Plugins/Runtime/Harmonix/Harmonix.uplugin
--- a/Engine/Plugins/Runtime/Harmonix/Harmonix.uplugin
+++ b/Engine/Plugins/Runtime/Harmonix/Harmonix.uplugin
@@ -12,7 +12,7 @@
 	"SupportURL": "",
 	"EnabledByDefault": false,
 	"CanContainContent": true,
-	"IsBetaVersion": true,
+	"IsBetaVersion": false,
 	"IsExperimentalVersion": true,
 	"Installed": false,
 	"Modules": [
