[Harmonix] MIDI Track Filter node: move old version to internal

diff --git a/Engine/Plugins/Runtime/Harmonix/Source/HarmonixMetasound/Private/HarmonixMetasound/Nodes/MidiStreamTrackFilterNode.cpp b/Engine/Plugins/Runtime/Harmonix/Source/HarmonixMetasound/Private/HarmonixMetasound/Nodes/MidiStreamTrackFilterNode.cpp
--- a/Engine/Plugins/Runtime/Harmonix/Source/HarmonixMetasound/Private/HarmonixMetasound/Nodes/MidiStreamTrackFilterNode.cpp
+++ b/Engine/Plugins/Runtime/Harmonix/Source/HarmonixMetasound/Private/HarmonixMetasound/Nodes/MidiStreamTrackFilterNode.cpp
@@ -14,8 +14,6 @@
 
 #include "HarmonixMetasound/MidiOps/MidiTrackFilter.h"
 
-DEFINE_LOG_CATEGORY_STATIC(LogMidiStreamTrackFilter, Log, All);
-
 #define LOCTEXT_NAMESPACE "HarmonixMetaSound"
 
 namespace HarmonixMetasound::Nodes::MidiTrackFilter
@@ -198,173 +196,6 @@ namespace HarmonixMetasound::Nodes::MidiTrackFilter
 	};
 
 	METASOUND_REGISTER_NODE(FMidiStreamTrackFilterNode_V1)
-	
-	class FMidiStreamTrackFilterOperator_V0 : public TExecutableOperator<FMidiStreamTrackFilterOperator_V0>
-	{
-	public:
-		static const FNodeClassMetadata& GetNodeInfo();
-		static const FVertexInterface& GetVertexInterface();
-		static TUniquePtr<IOperator> CreateOperator(const FBuildOperatorParams& InParams, FBuildResults& OutResults);
-
-		FMidiStreamTrackFilterOperator_V0(const FBuildOperatorParams& InParams,
-								  const FBoolReadRef&       InEnabled,
-								  const FStringReadRef&     InTrackSelect,
-								  const FMidiStreamReadRef& InMidiStream);
-
-		virtual void BindInputs(FInputVertexInterfaceData& InVertexData) override;
-		virtual void BindOutputs(FOutputVertexInterfaceData& InVertexData) override;
-
-		void Reset(const FResetParams& ResetParams);
-
-		void Execute();
-
-	private:
-		//** INPUTS
-		FBoolReadRef       EnableInPin;
-		FStringReadRef     TrackSelectInPin;
-		FMidiStreamReadRef MidiStreamInPin;
-
-		//** OUTPUTS
-		FMidiStreamWriteRef MidiStreamOutPin;
-
-		//** DATA
-		FString CurrentTrackIndexFilter;
-		FMidiStreamEventTrackChannelFilter Filter;
-	};
-
-	class FMidiStreamTrackFilterNode_V0 : public FNodeFacade
-	{
-	public:
-		FMidiStreamTrackFilterNode_V0(const FNodeInitData& InInitData)
-			: FNodeFacade(InInitData.InstanceName, InInitData.InstanceID, TFacadeOperatorClass<FMidiStreamTrackFilterOperator_V0>())
-		{}
-		virtual ~FMidiStreamTrackFilterNode_V0() = default;
-	};
-
-	METASOUND_REGISTER_NODE(FMidiStreamTrackFilterNode_V0)
-		
-	const FNodeClassMetadata& FMidiStreamTrackFilterOperator_V0::GetNodeInfo()
-	{
-		auto InitNodeInfo = []() -> FNodeClassMetadata
-		{
-			FNodeClassMetadata Info;
-			Info.ClassName        = GetClassName();
-			Info.MajorVersion     = 0;
-			Info.MinorVersion     = 1;
-			Info.DisplayName      = METASOUND_LOCTEXT("MidiStreamTrackFilterNodeV0_DisplayName", "Midi Stream Track Filter");
-			Info.Description      = METASOUND_LOCTEXT("MidiStreamTrackFilterNodeV0_Description", "Duplicates the incoming midi stream to its output after filtering tracks from the input Midi Stream");
-			Info.Author           = PluginAuthor;
-			Info.PromptIfMissing  = PluginNodeMissingPrompt;
-			Info.DefaultInterface = GetVertexInterface();
-			Info.CategoryHierarchy.Emplace(NodeCategories::Music);
-			Info.bDeprecated = true;
-			return Info;
-		};
-
-		static const FNodeClassMetadata Info = InitNodeInfo();
-
-		return Info;
-	}
-
-	const FVertexInterface& FMidiStreamTrackFilterOperator_V0::GetVertexInterface()
-	{
-		static const FVertexInterface Interface(
-			FInputVertexInterface(
-				TInputDataVertex<bool>(METASOUND_GET_PARAM_NAME_AND_METADATA(Inputs::Enable), true),
-				TInputDataVertex<FString>(METASOUND_GET_PARAM_NAME_AND_METADATA(CommonPinNames::Inputs::MidiTrackIndexFilterSpecifier), FString(TEXT("*"))),
-				TInputDataVertex<FMidiStream>(METASOUND_GET_PARAM_NAME_AND_METADATA(Inputs::MidiStream))
-			),
-			FOutputVertexInterface(
-				TOutputDataVertex<FMidiStream>(METASOUND_GET_PARAM_NAME_AND_METADATA(Outputs::MidiStream))
-			)
-		);
-
-		return Interface;
-	}
-
-	TUniquePtr<IOperator> FMidiStreamTrackFilterOperator_V0::CreateOperator(const FBuildOperatorParams& InParams, FBuildResults& OutResults)
-	{
-		const FMidiStreamTrackFilterNode_V0& LoggerNode = static_cast<const FMidiStreamTrackFilterNode_V0&>(InParams.Node);
-
-		const FInputVertexInterfaceData& InputData = InParams.InputData;
-		FBoolReadRef InEnabled = InputData.GetOrCreateDefaultDataReadReference<bool>(METASOUND_GET_PARAM_NAME(Inputs::Enable), InParams.OperatorSettings);
-		FStringReadRef InTrackIndexFilter = InputData.GetOrCreateDefaultDataReadReference<FString>(METASOUND_GET_PARAM_NAME(CommonPinNames::Inputs::MidiTrackIndexFilterSpecifier), InParams.OperatorSettings);
-		FMidiStreamReadRef InMidiStream = InputData.GetOrConstructDataReadReference<FMidiStream>(METASOUND_GET_PARAM_NAME(Inputs::MidiStream), InParams.OperatorSettings);
-
-		return MakeUnique<FMidiStreamTrackFilterOperator_V0>(InParams, 
-			InEnabled,
-			InTrackIndexFilter,
-			InMidiStream);
-	}
-
-	FMidiStreamTrackFilterOperator_V0::FMidiStreamTrackFilterOperator_V0(const FBuildOperatorParams& InParams,
-											   const FBoolReadRef& InEnabled,
-											   const FStringReadRef& InTrackSelect,
-											   const FMidiStreamReadRef& InMidiStream)
-		: EnableInPin(InEnabled)
-		, TrackSelectInPin(InTrackSelect)
-		, MidiStreamInPin(InMidiStream)
-		, MidiStreamOutPin(FMidiStreamWriteRef::CreateNew(InParams.OperatorSettings))
-	{
-		Reset(InParams);
-	}
-
-	void FMidiStreamTrackFilterOperator_V0::BindInputs(FInputVertexInterfaceData& InVertexData)
-	{
-		InVertexData.BindReadVertex(METASOUND_GET_PARAM_NAME(Inputs::Enable), EnableInPin);
-		InVertexData.BindReadVertex(METASOUND_GET_PARAM_NAME(CommonPinNames::Inputs::MidiTrackIndexFilterSpecifier), TrackSelectInPin);
-		InVertexData.BindReadVertex(METASOUND_GET_PARAM_NAME(Inputs::MidiStream), MidiStreamInPin);
-	}
-
-	void FMidiStreamTrackFilterOperator_V0::BindOutputs(FOutputVertexInterfaceData& InVertexData)
-	{
-		using namespace CommonPinNames;
-
-		InVertexData.BindReadVertex(METASOUND_GET_PARAM_NAME(Outputs::MidiStream), MidiStreamOutPin);
-	}
-
-	void FMidiStreamTrackFilterOperator_V0::Reset(const FResetParams& ResetParams)
-	{
-		MidiStreamOutPin->PrepareBlock();
-		
-		if (const FMidiClockReadRef* Clock = MidiStreamInPin->GetMidiClockSource())
-		{
-			MidiStreamOutPin->SetClockSource(*Clock);
-		}
-
-		CurrentTrackIndexFilter = *TrackSelectInPin;
-		
-		FString ParseTrackIndexFilterErrorMessage;
-		if (!Filter.SetTrackFilterFromString(CurrentTrackIndexFilter, ParseTrackIndexFilterErrorMessage))
-		{
-			UE_LOG(LogMidiStreamTrackFilter, Warning, TEXT("Error parsing Track Filter string: %s"), *ParseTrackIndexFilterErrorMessage);
-		}
-
-		FString ParseMidiChannelFilterErrorMessage;
-		if (!Filter.SetMidiChannelFilterFromString(TEXT("*"), ParseMidiChannelFilterErrorMessage))
-		{
-			UE_LOG(LogMidiStreamTrackFilter, Warning, TEXT("Error parsing Midi Channel Filter string: %s"), *ParseMidiChannelFilterErrorMessage);
-		}
-	}
-
-	void FMidiStreamTrackFilterOperator_V0::Execute()
-	{
-		MidiStreamOutPin->PrepareBlock();
-
-		if (CurrentTrackIndexFilter != *TrackSelectInPin)
-		{
-			CurrentTrackIndexFilter = *TrackSelectInPin;
-
-			FString ParseTrackIndexFilterErrorMessage;
-			if (!Filter.SetTrackFilterFromString(CurrentTrackIndexFilter, ParseTrackIndexFilterErrorMessage))
-			{
-				UE_LOG(LogMidiStreamTrackFilter, Warning, TEXT("Error parsing Track Filter string: %s"), *ParseTrackIndexFilterErrorMessage);
-			}
-		}
-
-		MidiStreamOutPin->PrepareBlock();
-		MidiStreamOutPin->Copy(MidiStreamInPin, [&](const FMidiStreamEvent& InEvent) {return Filter(InEvent); });
-	}
 }
 
 #undef LOCTEXT_NAMESPACE // "HarmonixMetaSound"
diff --git a/Engine/Plugins/Runtime/Harmonix/Source/HarmonixMetasound/Public/HarmonixMetasound/DataTypes/MidiStream.h b/Engine/Plugins/Runtime/Harmonix/Source/HarmonixMetasound/Public/HarmonixMetasound/DataTypes/MidiStream.h
--- a/Engine/Plugins/Runtime/Harmonix/Source/HarmonixMetasound/Public/HarmonixMetasound/DataTypes/MidiStream.h
+++ b/Engine/Plugins/Runtime/Harmonix/Source/HarmonixMetasound/Public/HarmonixMetasound/DataTypes/MidiStream.h
@@ -18,7 +18,7 @@
 
 namespace HarmonixMetasound
 {
-	DECLARE_LOG_CATEGORY_EXTERN(LogMidiStreamDataType, Log, All)
+	HARMONIXMETASOUND_API DECLARE_LOG_CATEGORY_EXTERN(LogMidiStreamDataType, Log, All)
 
 	struct HARMONIXMETASOUND_API FMidiStreamEvent
 	{
