[Harmonix] MIDI Note Trigger node: move old version to internal

diff --git a/Engine/Plugins/Runtime/Harmonix/Source/HarmonixMetasound/Private/HarmonixMetasound/Nodes/MidiNoteTriggerNode.cpp b/Engine/Plugins/Runtime/Harmonix/Source/HarmonixMetasound/Private/HarmonixMetasound/Nodes/MidiNoteTriggerNode.cpp
--- a/Engine/Plugins/Runtime/Harmonix/Source/HarmonixMetasound/Private/HarmonixMetasound/Nodes/MidiNoteTriggerNode.cpp
+++ b/Engine/Plugins/Runtime/Harmonix/Source/HarmonixMetasound/Private/HarmonixMetasound/Nodes/MidiNoteTriggerNode.cpp
@@ -228,327 +228,6 @@ namespace HarmonixMetasound::Nodes::MidiNoteTriggerNode
 	};
 
 	METASOUND_REGISTER_NODE(FMidiNoteTriggerNode_V1)
-
-	class FMidiNoteTriggerOperator_V0 : public TExecutableOperator<FMidiNoteTriggerOperator_V0>
-	{
-	public:
-		static const FNodeClassMetadata& GetNodeInfo();
-		static const FVertexInterface& GetVertexInterface();
-		static TUniquePtr<IOperator> CreateOperator(const FBuildOperatorParams& InParams, FBuildResults& OutResults);
-
-		FMidiNoteTriggerOperator_V0(const FBuildOperatorParams& InParams,
-								 const FBoolReadRef&       InEnabled,
-								 const FMidiStreamReadRef& InMidiStream,
-								 const FInt32ReadRef&      InTrackNumber,
-								 const FInt32ReadRef&      InChannelNumber,
-								 const FInt32ReadRef&      InMinNoteNumber,
-								 const FInt32ReadRef&      InMaxNoteNumber,
-								 const FInt32ReadRef&      InMinVelocity,
-								 const FInt32ReadRef&      InMaxVelocity,
-								 const FBoolReadRef&       InMidiClockAffectsFrequency);
-
-		virtual void BindInputs(FInputVertexInterfaceData& InVertexData) override;
-		virtual void BindOutputs(FOutputVertexInterfaceData& InVertexData) override;
-		virtual FDataReferenceCollection GetInputs() const override;
-		virtual FDataReferenceCollection GetOutputs() const override;
-
-		void Reset(const FResetParams& ResetParams);
-		
-		void Execute();
-
-	private:
-		//** INPUTS
-		FBoolReadRef       EnableInPin;
-		FMidiStreamReadRef MidiStreamInPin;
-		FInt32ReadRef      TrackNumberInPin;
-		FInt32ReadRef      ChannelNumInPin;
-		FInt32ReadRef      MinNoteInPin;
-		FInt32ReadRef      MaxNoteInPin;
-		FInt32ReadRef      MinVelInPin;
-		FInt32ReadRef      MaxVelInPin;
-		FBoolReadRef       ClockAffectsFrequencyInPin;
-
-		//** OUTPUTS
-		FTriggerWriteRef NoteOnOutPin;
-		FTriggerWriteRef NoteOffOutPin;
-		FInt32WriteRef   NoteNumOutPin;
-		FFloatWriteRef   FreqOutPin;
-		FInt32WriteRef   VelOutPin;
-		FFloatWriteRef   NormVelOutPin;
-
-		//** DATA
-		int8 SoundingNote = -1;
-		FMidiVoiceId PlayingId;
-	};
-
-	class FMidiNoteTriggerNode_V0 : public FNodeFacade
-	{
-	public:
-		FMidiNoteTriggerNode_V0(const FNodeInitData& InInitData)
-			: FNodeFacade(InInitData.InstanceName, InInitData.InstanceID, TFacadeOperatorClass<FMidiNoteTriggerOperator_V0>())
-		{}
-		virtual ~FMidiNoteTriggerNode_V0() = default;
-	};
-
-	METASOUND_REGISTER_NODE(FMidiNoteTriggerNode_V0)
-
-	const FNodeClassMetadata& FMidiNoteTriggerOperator_V0::GetNodeInfo()
-	{
-		auto InitNodeInfo = []() -> FNodeClassMetadata
-		{
-			FNodeClassMetadata Info;
-			Info.ClassName        = GetClassName();
-			Info.MajorVersion     = 0;
-			Info.MinorVersion     = 1;
-			Info.DisplayName      = METASOUND_LOCTEXT("MidiNoteTriggerNodeV0_DisplayName", "Midi Note Trigger");
-			Info.Description      = METASOUND_LOCTEXT("MidiNoteTriggerNodeV0_Description", "Receives a midi stream, filters for the desired messages, and outputs triggers.");
-			Info.Author           = PluginAuthor;
-			Info.PromptIfMissing  = PluginNodeMissingPrompt;
-			Info.DefaultInterface = GetVertexInterface();
-			Info.CategoryHierarchy.Emplace(NodeCategories::Music);
-			Info.bDeprecated = true;
-			return Info;
-		};
-
-		static const FNodeClassMetadata Info = InitNodeInfo();
-
-		return Info;
-	}
-
-	const FVertexInterface& FMidiNoteTriggerOperator_V0::GetVertexInterface()
-	{
-		static const FVertexInterface Interface(
-			FInputVertexInterface(
-				TInputDataVertex<bool>(METASOUND_GET_PARAM_NAME_AND_METADATA(Inputs::Enable), true),
-				TInputDataVertex<FMidiStream>(METASOUND_GET_PARAM_NAME_AND_METADATA(Inputs::MidiStream)),
-				TInputDataVertex<int32>(METASOUND_GET_PARAM_NAME_AND_METADATA(CommonPinNames::Inputs::MidiTrackNumber),1),
-				TInputDataVertex<int32>(METASOUND_GET_PARAM_NAME_AND_METADATA(CommonPinNames::Inputs::MidiChannelNumber), 1),
-				TInputDataVertex<int32>(METASOUND_GET_PARAM_NAME_AND_METADATA(CommonPinNames::Inputs::MinMidiNote), 0),
-				TInputDataVertex<int32>(METASOUND_GET_PARAM_NAME_AND_METADATA(CommonPinNames::Inputs::MaxMidiNote), 127),
-				TInputDataVertex<int32>(METASOUND_GET_PARAM_NAME_AND_METADATA(CommonPinNames::Inputs::MinMidiVelocity), 0),
-				TInputDataVertex<int32>(METASOUND_GET_PARAM_NAME_AND_METADATA(CommonPinNames::Inputs::MaxMidiVelocity), 127),
-				TInputDataVertex<bool>(METASOUND_GET_PARAM_NAME_AND_METADATA(CommonPinNames::Inputs::ClockSpeedToFrequency), false)
-				),
-			FOutputVertexInterface(
-				TOutputDataVertex<FTrigger>(METASOUND_GET_PARAM_NAME_AND_METADATA(Outputs::NoteOnTrigger)),
-				TOutputDataVertex<FTrigger>(METASOUND_GET_PARAM_NAME_AND_METADATA(Outputs::NoteOffTrigger)),
-				TOutputDataVertex<int32>(METASOUND_GET_PARAM_NAME_AND_METADATA(Outputs::MidiNoteNumber)),
-				TOutputDataVertex<float>(METASOUND_GET_PARAM_NAME_AND_METADATA(CommonPinNames::Outputs::Frequency)),
-				TOutputDataVertex<int32>(METASOUND_GET_PARAM_NAME_AND_METADATA(Outputs::MidiVelocity)),
-				TOutputDataVertex<float>(METASOUND_GET_PARAM_NAME_AND_METADATA(CommonPinNames::Outputs::NormalizedVelocity))
-				)
-		);
-
-		return Interface;
-	}
-
-	TUniquePtr<IOperator> FMidiNoteTriggerOperator_V0::CreateOperator(const FBuildOperatorParams& InParams, FBuildResults& OutResults)
-	{
-		using namespace CommonPinNames;
-
-		const FInputVertexInterfaceData& InputData = InParams.InputData;
-		FBoolReadRef InEnabled          = InputData.GetOrCreateDefaultDataReadReference<bool>(METASOUND_GET_PARAM_NAME(Inputs::Enable), InParams.OperatorSettings);
-		FMidiStreamReadRef InMidiStream = InputData.GetOrConstructDataReadReference<FMidiStream>(METASOUND_GET_PARAM_NAME(Inputs::MidiStream), InParams.OperatorSettings);
-		FInt32ReadRef InTrackNumber     = InputData.GetOrCreateDefaultDataReadReference<int32>(METASOUND_GET_PARAM_NAME(CommonPinNames::Inputs::MidiTrackNumber), InParams.OperatorSettings);
-		FInt32ReadRef InMidiChannel     = InputData.GetOrCreateDefaultDataReadReference<int32>(METASOUND_GET_PARAM_NAME(CommonPinNames::Inputs::MidiChannelNumber), InParams.OperatorSettings);
-		FInt32ReadRef InMinNote         = InputData.GetOrCreateDefaultDataReadReference<int32>(METASOUND_GET_PARAM_NAME(CommonPinNames::Inputs::MinMidiNote), InParams.OperatorSettings);
-		FInt32ReadRef InMaxNote         = InputData.GetOrCreateDefaultDataReadReference<int32>(METASOUND_GET_PARAM_NAME(CommonPinNames::Inputs::MaxMidiNote), InParams.OperatorSettings);
-		FInt32ReadRef InMinVelocity     = InputData.GetOrCreateDefaultDataReadReference<int32>(METASOUND_GET_PARAM_NAME(CommonPinNames::Inputs::MinMidiVelocity), InParams.OperatorSettings);
-		FInt32ReadRef InMaxVelocity     = InputData.GetOrCreateDefaultDataReadReference<int32>(METASOUND_GET_PARAM_NAME(CommonPinNames::Inputs::MaxMidiVelocity), InParams.OperatorSettings);
-		FBoolReadRef  InClockToFreq     = InputData.GetOrCreateDefaultDataReadReference<bool>(METASOUND_GET_PARAM_NAME(CommonPinNames::Inputs::ClockSpeedToFrequency), InParams.OperatorSettings);
-
-		return MakeUnique<FMidiNoteTriggerOperator_V0>(InParams, InEnabled, InMidiStream, InTrackNumber, InMidiChannel, InMinNote, InMaxNote, InMinVelocity, InMaxVelocity, InClockToFreq);
-	}
-
-	FMidiNoteTriggerOperator_V0::FMidiNoteTriggerOperator_V0(const FBuildOperatorParams& InParams,
-													   const FBoolReadRef&       InEnabled,
-													   const FMidiStreamReadRef& InMidiStream,
-													   const FInt32ReadRef&      InTrackNumber,
-													   const FInt32ReadRef&      InChannelNumber,
-													   const FInt32ReadRef&      InMinNoteNumber,
-													   const FInt32ReadRef&      InMaxNoteNumber,
-													   const FInt32ReadRef&      InMinVelocity,
-													   const FInt32ReadRef&      InMaxVelocity,
-													   const FBoolReadRef&       InClockToFreq)
-		: EnableInPin(InEnabled)
-		, MidiStreamInPin(InMidiStream)
-		, TrackNumberInPin(InTrackNumber)
-		, ChannelNumInPin(InChannelNumber)
-		, MinNoteInPin(InMinNoteNumber)
-		, MaxNoteInPin(InMaxNoteNumber)
-		, MinVelInPin(InMinVelocity)
-		, MaxVelInPin(InMaxVelocity)
-		, ClockAffectsFrequencyInPin(InClockToFreq)
-		, NoteOnOutPin(FTriggerWriteRef::CreateNew(InParams.OperatorSettings))
-		, NoteOffOutPin(FTriggerWriteRef::CreateNew(InParams.OperatorSettings))
-		, NoteNumOutPin(FInt32WriteRef::CreateNew(0))
-		, FreqOutPin(FFloatWriteRef::CreateNew(0.0f))
-		, VelOutPin(FInt32WriteRef::CreateNew(0))
-		, NormVelOutPin(FFloatWriteRef::CreateNew(0.0f))
-	{
-		Reset(InParams);
-	}
-
-	void FMidiNoteTriggerOperator_V0::BindInputs(FInputVertexInterfaceData& InVertexData)
-	{
-		using namespace CommonPinNames;
-
-		InVertexData.BindReadVertex(METASOUND_GET_PARAM_NAME(Inputs::Enable), EnableInPin);
-		InVertexData.BindReadVertex(METASOUND_GET_PARAM_NAME(Inputs::MidiStream),   MidiStreamInPin);
-		InVertexData.BindReadVertex(METASOUND_GET_PARAM_NAME(CommonPinNames::Inputs::MidiTrackNumber),  TrackNumberInPin);
-		InVertexData.BindReadVertex(METASOUND_GET_PARAM_NAME(CommonPinNames::Inputs::MidiChannelNumber), ChannelNumInPin);
-		InVertexData.BindReadVertex(METASOUND_GET_PARAM_NAME(CommonPinNames::Inputs::MinMidiNote), MinNoteInPin);
-		InVertexData.BindReadVertex(METASOUND_GET_PARAM_NAME(CommonPinNames::Inputs::MaxMidiNote), MaxNoteInPin);
-		InVertexData.BindReadVertex(METASOUND_GET_PARAM_NAME(CommonPinNames::Inputs::MinMidiVelocity), MinVelInPin);
-		InVertexData.BindReadVertex(METASOUND_GET_PARAM_NAME(CommonPinNames::Inputs::MaxMidiVelocity), MaxVelInPin);
-		InVertexData.BindReadVertex(METASOUND_GET_PARAM_NAME(CommonPinNames::Inputs::ClockSpeedToFrequency), ClockAffectsFrequencyInPin);
-	}
-
-	void FMidiNoteTriggerOperator_V0::BindOutputs(FOutputVertexInterfaceData& InVertexData)
-	{
-		using namespace CommonPinNames;
-
-		InVertexData.BindReadVertex(METASOUND_GET_PARAM_NAME(CommonPinNames::Outputs::NoteOn), NoteOnOutPin);
-		InVertexData.BindReadVertex(METASOUND_GET_PARAM_NAME(CommonPinNames::Outputs::NoteOff), NoteOffOutPin);
-		InVertexData.BindReadVertex(METASOUND_GET_PARAM_NAME(Outputs::MidiNoteNumber), NoteNumOutPin);
-		InVertexData.BindReadVertex(METASOUND_GET_PARAM_NAME(CommonPinNames::Outputs::Frequency), FreqOutPin);
-		InVertexData.BindReadVertex(METASOUND_GET_PARAM_NAME(Outputs::MidiVelocity), VelOutPin);
-		InVertexData.BindReadVertex(METASOUND_GET_PARAM_NAME(CommonPinNames::Outputs::NormalizedVelocity), NormVelOutPin);
-	}
-	
-	FDataReferenceCollection FMidiNoteTriggerOperator_V0::GetInputs() const
-	{
-		// This should never be called. Bind(...) is called instead. This method
-		// exists as a stop-gap until the API can be deprecated and removed.
-		checkNoEntry();
-		return {};
-	}
-
-	FDataReferenceCollection FMidiNoteTriggerOperator_V0::GetOutputs() const
-	{
-		// This should never be called. Bind(...) is called instead. This method
-		// exists as a stop-gap until the API can be deprecated and removed.
-		checkNoEntry();
-		return {};
-	}
-
-	void FMidiNoteTriggerOperator_V0::Reset(const FResetParams& ResetParams)
-	{
-		NoteOnOutPin->Reset();
-		NoteOffOutPin->Reset();
-		*NoteNumOutPin = 0;
-		*FreqOutPin = 0;
-		*VelOutPin = 0;
-		*NormVelOutPin = 0;
-		SoundingNote = -1;
-		PlayingId = FMidiVoiceId::None();
-	}
-
-	void FMidiNoteTriggerOperator_V0::Execute()
-	{
-		NoteOnOutPin->AdvanceBlock();
-		NoteOffOutPin->AdvanceBlock();
-
-		if (!*EnableInPin)
-		{
-			if (SoundingNote >= 0)
-			{
-				*VelOutPin     = 0;
-				*NormVelOutPin = 0.0f;
-				NoteOffOutPin->TriggerFrame(0);
-				SoundingNote   = -1;
-			}
-			return;
-		}
-
-		auto TransportChanges = MidiStreamInPin->GetTransportChangesInBlock();
-		if (!TransportChanges.IsEmpty())
-		{
-			if (TransportChanges.Last().TransportState != EMusicPlayerTransportState::Playing)
-			{
-				if (SoundingNote >= 0)
-				{
-					*VelOutPin     = 0;
-					*NormVelOutPin = 0.0f;
-					NoteOffOutPin->TriggerFrame(0);
-					SoundingNote   = -1;
-				}
-				return;
-			}
-		}
-
-		int32 NoteOffTriggerFrame = -1;
-		const TArray<FMidiStreamEvent>& MidiEvents = MidiStreamInPin->GetEventsInBlock();
-		for (auto& Event : MidiEvents)
-		{
-			int32 trackPin   = *TrackNumberInPin;
-			int32 channelPin = *ChannelNumInPin;
-			bool IsNoteOn    = Event.MidiMessage.IsNoteOn();
-			bool IsNoteOff   = Event.MidiMessage.IsNoteOff();
-		
-			if (Event.TrackIndex != *TrackNumberInPin ||
-				(!IsNoteOn && !IsNoteOff) ||
-				Event.MidiMessage.GetStdChannel() + 1 != *ChannelNumInPin)
-			{
-				continue;
-			}
-
-			if (Event.MidiMessage.GetStdData1() < *MinNoteInPin ||
-				Event.MidiMessage.GetStdData1() > *MaxNoteInPin)
-			{
-				continue;
-			}
-
-			if (Event.MidiMessage.IsNoteOn()) 
-			{
-				if (Event.MidiMessage.GetStdData2() < *MinVelInPin ||
-					Event.MidiMessage.GetStdData2() > *MaxVelInPin)
-				{
-					continue;
-				}
-
-				if (SoundingNote > -1)
-				{
-					// stopping sounding note!
-					NoteOffOutPin->TriggerFrame(Event.BlockSampleFrameIndex);
-					NoteOffTriggerFrame = Event.BlockSampleFrameIndex;
-					SoundingNote = -1;
-				}
-				PlayingId      = Event.GetVoiceId();
-				*VelOutPin     = Event.MidiMessage.GetStdData2();
-				*NormVelOutPin = (float)Event.MidiMessage.GetStdData2() / 127.0f;
-				*NoteNumOutPin = Event.MidiMessage.GetStdData1();
-				NoteOnOutPin->TriggerFrame(NoteOffTriggerFrame == Event.BlockSampleFrameIndex ? NoteOffTriggerFrame + 1 : Event.BlockSampleFrameIndex);
-				SoundingNote   = Event.MidiMessage.GetStdData1();;
-			}
-			else
-			{
-				if (Event.GetVoiceId() == PlayingId)
-				{
-					*VelOutPin = 0;
-					*NormVelOutPin = 0.0f;
-					*NoteNumOutPin = Event.MidiMessage.GetStdData1();
-					NoteOffOutPin->TriggerFrame(Event.BlockSampleFrameIndex);
-					NoteOffTriggerFrame = Event.BlockSampleFrameIndex;
-					SoundingNote = -1;
-					PlayingId = FMidiVoiceId::None();
-				}
-			}
-		}
-		
-		// We do this here so that even during the playback of a note we can
-		// adjust the pitch if there is a clock source and it has a changing speed. 
-		if (SoundingNote != -1)
-		{
-			float Frequency = powf(2.0f, (float)(SoundingNote - 69) / 12.0f) * 440.0f;
-			const FMidiClockReadRef* ClockSource = MidiStreamInPin->GetMidiClockSource();
-			if (ClockSource && *ClockAffectsFrequencyInPin)
-			{
-				Frequency *= (*ClockSource)->GetSpeedAtEndOfBlock();
-			}
-			*FreqOutPin = Frequency;
-		}
-	}
 }
 
 #undef LOCTEXT_NAMESPACE // "HarmonixMetaSound"
