[Harmonix] MIDI Transpose node: move old node to internal

diff --git a/Engine/Plugins/Runtime/Harmonix/Source/HarmonixMetasound/Private/HarmonixMetasound/Nodes/MidiStreamTransposerNode.cpp b/Engine/Plugins/Runtime/Harmonix/Source/HarmonixMetasound/Private/HarmonixMetasound/Nodes/MidiStreamTransposerNode.cpp
--- a/Engine/Plugins/Runtime/Harmonix/Source/HarmonixMetasound/Private/HarmonixMetasound/Nodes/MidiStreamTransposerNode.cpp
+++ b/Engine/Plugins/Runtime/Harmonix/Source/HarmonixMetasound/Private/HarmonixMetasound/Nodes/MidiStreamTransposerNode.cpp
@@ -14,8 +14,6 @@
 #include "HarmonixMetasound/DataTypes/MidiStream.h"
 #include "HarmonixMetasound/DataTypes/MusicTransport.h"
 
-DEFINE_LOG_CATEGORY_STATIC(LogMidiStreamTransposer, Log, All);
-
 #define LOCTEXT_NAMESPACE "HarmonixMetaSound"
 
 namespace HarmonixMetasound::Nodes::MidiNoteTranspose
@@ -195,204 +193,6 @@ namespace HarmonixMetasound::Nodes::MidiNoteTranspose
 	};
 
 	METASOUND_REGISTER_NODE(FMidiNoteTransposeNode)
-
-	class FMidiStreamTransposerOperator_V0 : public TExecutableOperator<FMidiStreamTransposerOperator_V0>
-	{
-	public:
-		static const FNodeClassMetadata& GetNodeInfo();
-		static const FVertexInterface& GetVertexInterface();
-		static TUniquePtr<IOperator> CreateOperator(const FBuildOperatorParams& InParams, FBuildResults& OutResults);
-
-		FMidiStreamTransposerOperator_V0(const FBuildOperatorParams& InParams,
-								  const FBoolReadRef&       InEnabled,
-								  const FStringReadRef&     InTrackSelect,
-								  const FStringReadRef&     InChannelSelect,
-								  const FInt32ReadRef&      InTransposition,
-								  const FMidiStreamReadRef& InMidiStream);
-
-		virtual void BindInputs(FInputVertexInterfaceData& InVertexData) override;
-		virtual void BindOutputs(FOutputVertexInterfaceData& InVertexData) override;
-
-		void Reset(const FResetParams& ResetParams);
-
-		void Execute();
-
-	private:
-		//** INPUTS
-		FBoolReadRef       EnableInPin;
-		FStringReadRef     TrackSelectInPin;
-		FStringReadRef     ChannelSelectInPin;
-		FInt32ReadRef      TranspositionInPin;
-		FMidiStreamReadRef MidiStreamInPin;
-
-		//** OUTPUTS
-		FMidiStreamWriteRef MidiStreamOutPin;
-
-		//** DATA
-		FString CurrentTrackFilter;
-		FString CurrentMidiChannelFilter;
-		FMidiStreamEventTrackChannelFilter Filter;
-	};
-
-	class FMidiStreamTransposerNode_v0 : public FNodeFacade
-	{
-	public:
-		FMidiStreamTransposerNode_v0(const FNodeInitData& InInitData)
-			: FNodeFacade(InInitData.InstanceName, InInitData.InstanceID, TFacadeOperatorClass<FMidiStreamTransposerOperator_V0>())
-		{}
-		virtual ~FMidiStreamTransposerNode_v0() = default;
-	};
-
-	METASOUND_REGISTER_NODE(FMidiStreamTransposerNode_v0)
-		
-	const FNodeClassMetadata& FMidiStreamTransposerOperator_V0::GetNodeInfo()
-	{
-		auto InitNodeInfo = []() -> FNodeClassMetadata
-		{
-			FNodeClassMetadata Info;
-			Info.ClassName        = { HarmonixNodeNamespace, TEXT("MidiStreamTransposer"), TEXT("")};
-			Info.MajorVersion     = 0;
-			Info.MinorVersion     = 1;
-			Info.DisplayName      = METASOUND_LOCTEXT("MidiStreamTransposerNode_DisplayName", "Midi Stream Transposer");
-			Info.Description      = METASOUND_LOCTEXT("MidiStreamTransposerNode_Description", "Duplicates the incoming midi stream to its output with the note on/off messages transposed by the specified number of semitones.");
-			Info.Author           = PluginAuthor;
-			Info.PromptIfMissing  = PluginNodeMissingPrompt;
-			Info.DefaultInterface = GetVertexInterface();
-			Info.CategoryHierarchy.Emplace(NodeCategories::Music);
-			Info.bDeprecated = true;
-			return Info;
-		};
-
-		static const FNodeClassMetadata Info = InitNodeInfo();
-
-		return Info;
-	}
-
-	const FVertexInterface& FMidiStreamTransposerOperator_V0::GetVertexInterface()
-	{
-		static const FVertexInterface Interface(
-			FInputVertexInterface(
-				TInputDataVertex<bool>(METASOUND_GET_PARAM_NAME_AND_METADATA(Inputs::Enable), true),
-				TInputDataVertex<FString>(METASOUND_GET_PARAM_NAME_AND_METADATA(CommonPinNames::Inputs::MidiTrackIndexFilterSpecifier), FString(TEXT("*"))),
-				TInputDataVertex<FString>(METASOUND_GET_PARAM_NAME_AND_METADATA(CommonPinNames::Inputs::MidiChannelFilterSpecifier), FString(TEXT("*"))),
-				TInputDataVertex<int32>(METASOUND_GET_PARAM_NAME_AND_METADATA(Inputs::Transposition), 0),
-				TInputDataVertex<FMidiStream>(METASOUND_GET_PARAM_NAME_AND_METADATA(Inputs::MidiStream))
-			),
-			FOutputVertexInterface(
-				TOutputDataVertex<FMidiStream>(METASOUND_GET_PARAM_NAME_AND_METADATA(Outputs::MidiStream))
-			)
-		);
-
-		return Interface;
-	}
-
-	TUniquePtr<IOperator> FMidiStreamTransposerOperator_V0::CreateOperator(const FBuildOperatorParams& InParams, FBuildResults& OutResults)
-	{
-		const FMidiStreamTransposerNode_v0& LoggerNode = static_cast<const FMidiStreamTransposerNode_v0&>(InParams.Node);
-
-		const FInputVertexInterfaceData& InputData = InParams.InputData;
-		FBoolReadRef InEnabled = InputData.GetOrCreateDefaultDataReadReference<bool>(METASOUND_GET_PARAM_NAME(Inputs::Enable), InParams.OperatorSettings);
-		FStringReadRef InTrackFilter = InputData.GetOrCreateDefaultDataReadReference<FString>(METASOUND_GET_PARAM_NAME(CommonPinNames::Inputs::MidiTrackIndexFilterSpecifier), InParams.OperatorSettings);
-		FStringReadRef InMidiChannelFilter = InputData.GetOrCreateDefaultDataReadReference<FString>(METASOUND_GET_PARAM_NAME(CommonPinNames::Inputs::MidiChannelFilterSpecifier), InParams.OperatorSettings);
-		FInt32ReadRef InTransposition = InputData.GetOrCreateDefaultDataReadReference<int32>(METASOUND_GET_PARAM_NAME(Inputs::Transposition), InParams.OperatorSettings);
-		FMidiStreamReadRef InMidiStream = InputData.GetOrConstructDataReadReference<FMidiStream>(METASOUND_GET_PARAM_NAME(Inputs::MidiStream), InParams.OperatorSettings);
-
-		return MakeUnique<FMidiStreamTransposerOperator_V0>(InParams, 
-			InEnabled,
-			InTrackFilter,
-			InMidiChannelFilter,
-			InTransposition,
-			InMidiStream);
-	}
-
-	FMidiStreamTransposerOperator_V0::FMidiStreamTransposerOperator_V0(const FBuildOperatorParams& InParams,
-											   const FBoolReadRef& InEnabled,
-											   const FStringReadRef& InTrackSelect,
-											   const FStringReadRef& InChannelSelect,
-											   const FInt32ReadRef& InTransposition,
-											   const FMidiStreamReadRef& InMidiStream)
-		: EnableInPin(InEnabled)
-		, TrackSelectInPin(InTrackSelect)
-		, ChannelSelectInPin(InChannelSelect)
-		, TranspositionInPin(InTransposition)
-		, MidiStreamInPin(InMidiStream)
-		, MidiStreamOutPin(FMidiStreamWriteRef::CreateNew(InParams.OperatorSettings))
-	{
-		Reset(InParams);
-	}
-
-	void FMidiStreamTransposerOperator_V0::BindInputs(FInputVertexInterfaceData& InVertexData)
-	{
-		InVertexData.BindReadVertex(METASOUND_GET_PARAM_NAME(Inputs::Enable), EnableInPin);
-		InVertexData.BindReadVertex(METASOUND_GET_PARAM_NAME(CommonPinNames::Inputs::MidiTrackIndexFilterSpecifier), TrackSelectInPin);
-		InVertexData.BindReadVertex(METASOUND_GET_PARAM_NAME(CommonPinNames::Inputs::MidiChannelFilterSpecifier), ChannelSelectInPin);
-		InVertexData.BindReadVertex(METASOUND_GET_PARAM_NAME(Inputs::Transposition), TranspositionInPin);
-		InVertexData.BindReadVertex(METASOUND_GET_PARAM_NAME(Inputs::MidiStream), MidiStreamInPin);
-	}
-
-	void FMidiStreamTransposerOperator_V0::BindOutputs(FOutputVertexInterfaceData& InVertexData)
-	{
-		InVertexData.BindReadVertex(METASOUND_GET_PARAM_NAME(Outputs::MidiStream), MidiStreamOutPin);
-	}
-
-	void FMidiStreamTransposerOperator_V0::Reset(const FResetParams& ResetParams)
-	{
-		MidiStreamOutPin->PrepareBlock();
-		
-		if (const FMidiClockReadRef* Clock = MidiStreamInPin->GetMidiClockSource())
-		{
-			MidiStreamOutPin->SetClockSource(*Clock);
-		}
-		CurrentTrackFilter = *TrackSelectInPin;
-		CurrentMidiChannelFilter = *ChannelSelectInPin;
-		
-		FString ParseTrackFilterErrorMessage;
-		if (!Filter.SetTrackFilterFromString(CurrentTrackFilter, ParseTrackFilterErrorMessage))
-		{
-			UE_LOG(LogMidiStreamTransposer, Warning, TEXT("Error parsing Track Filter string: %s"), *ParseTrackFilterErrorMessage);
-		}
-
-		FString ParseMidiChannelFilterErrorMessage;
-		if (!Filter.SetMidiChannelFilterFromString(CurrentMidiChannelFilter, ParseMidiChannelFilterErrorMessage))
-		{
-			UE_LOG(LogMidiStreamTransposer, Warning, TEXT("Error parsing Channel Filter string: %s"), *ParseMidiChannelFilterErrorMessage);
-		}
-	}
-
-	void FMidiStreamTransposerOperator_V0::Execute()
-	{
-		if (CurrentTrackFilter != *TrackSelectInPin)
-		{
-			CurrentTrackFilter = *TrackSelectInPin;
-			FString ParseTrackFilterErrorMessage;
-			if (!Filter.SetTrackFilterFromString(CurrentTrackFilter, ParseTrackFilterErrorMessage))
-			{
-				UE_LOG(LogMidiStreamTransposer, Warning, TEXT("Error parsing Track Filter string: %s"), *ParseTrackFilterErrorMessage);
-			}
-		}
-
-		if (CurrentMidiChannelFilter != *ChannelSelectInPin)
-		{
-			CurrentMidiChannelFilter = *ChannelSelectInPin;
-			FString ParseMidiChannelFilterErrorMessage;
-			if (!Filter.SetMidiChannelFilterFromString(CurrentMidiChannelFilter, ParseMidiChannelFilterErrorMessage))
-			{
-				UE_LOG(LogMidiStreamTransposer, Warning, TEXT("Error parsing Channel Filter string: %s"), *ParseMidiChannelFilterErrorMessage);
-			}
-		}
-
-		auto Transformer = [&](const FMidiStreamEvent& InEvent) -> FMidiStreamEvent
-			{
-				FMidiStreamEvent NewEvent = InEvent;
-				if (Filter(NewEvent) && (NewEvent.MidiMessage.IsNoteOn() || NewEvent.MidiMessage.IsNoteOff()))
-				{
-					const int32 NewNote = FMath::Clamp(*TranspositionInPin + NewEvent.MidiMessage.Data1, 0, 127);
-					NewEvent.MidiMessage.Data1 = NewNote;
-				}
-				return NewEvent;
-			};
-		MidiStreamOutPin->Copy(MidiStreamInPin, [](const FMidiStreamEvent&){return true;}, Transformer);
-	}
 }
 
 #undef LOCTEXT_NAMESPACE // "HarmonixMetaSound"
