[Harmonix] MIDI Text Trigger node: move old node to internal

diff --git a/Engine/Plugins/Runtime/Harmonix/Source/HarmonixMetasound/Private/HarmonixMetasound/Nodes/MidiTextTriggerNode.cpp b/Engine/Plugins/Runtime/Harmonix/Source/HarmonixMetasound/Private/HarmonixMetasound/Nodes/MidiTextTriggerNode.cpp
--- a/Engine/Plugins/Runtime/Harmonix/Source/HarmonixMetasound/Private/HarmonixMetasound/Nodes/MidiTextTriggerNode.cpp
+++ b/Engine/Plugins/Runtime/Harmonix/Source/HarmonixMetasound/Private/HarmonixMetasound/Nodes/MidiTextTriggerNode.cpp
@@ -25,7 +25,7 @@ namespace HarmonixMetasound
 		METASOUND_PARAM(TriggerOutput, "Trigger Out", "A trigger when the text is encountered.")
 	}
 
-	class FMidiTextTriggerOperator_V1 : public TExecutableOperator<FMidiTextTriggerOperator_V1>
+	class FMidiTextTriggerOperator_V1 final : public TExecutableOperator<FMidiTextTriggerOperator_V1>
 	{
 	public:
 		static const FNodeClassMetadata& GetNodeInfo()
@@ -116,7 +116,7 @@ namespace HarmonixMetasound
 			TriggerOutPin->Reset();
 		}
 		
-		virtual void Execute()
+		void Execute()
 		{
 			TriggerOutPin->AdvanceBlock();
 
@@ -150,118 +150,7 @@ namespace HarmonixMetasound
 		FTriggerWriteRef TriggerOutPin;
 	};
 
-	// Here for backward compat
-	class FMidiTextTriggerOperator_V0 : public FMidiTextTriggerOperator_V1
-	{
-	public:
-		static const FNodeClassMetadata& GetNodeInfo()
-		{
-			auto InitNodeInfo = []() -> FNodeClassMetadata
-			{
-				FNodeClassMetadata Info;
-				Info.ClassName        = { HarmonixNodeNamespace, TEXT("MidiTextTrigger"), TEXT("")};
-				Info.MajorVersion     = 0;
-				Info.MinorVersion     = 1;
-				Info.DisplayName      = METASOUND_LOCTEXT("MidiTextTriggerNode_DisplayName", "MIDI Text Trigger");
-				Info.Description      = METASOUND_LOCTEXT("MidiTextTriggerNode_Description", "Receives a MIDI stream, filters for the desired messages, and outputs triggers.");
-				Info.Author           = PluginAuthor;
-				Info.PromptIfMissing  = PluginNodeMissingPrompt;
-				Info.DefaultInterface = GetVertexInterface();
-				Info.CategoryHierarchy.Emplace(NodeCategories::Music);
-				Info.bDeprecated = true;
-				return Info;
-			};
-
-			static const FNodeClassMetadata Info = InitNodeInfo();
-
-			return Info;
-		}
-		
-		static const FVertexInterface& GetVertexInterface()
-		{
-			using namespace MidiTextTriggerPinNames;
-			using namespace CommonPinNames;
-
-			static const FVertexInterface Interface(
-				FInputVertexInterface(
-					TInputDataVertex<bool>(METASOUND_GET_PARAM_NAME_AND_METADATA(Inputs::Enable), true),
-					TInputDataVertex<FMidiStream>(METASOUND_GET_PARAM_NAME_AND_METADATA(Inputs::MidiStream)),
-					TInputDataVertex<int32>(METASOUND_GET_PARAM_NAME_AND_METADATA(Inputs::MidiTrackNumber),1),
-					TInputDataVertex<FString>(METASOUND_GET_PARAM_NAME_AND_METADATA(TextInput))
-				),
-				FOutputVertexInterface(
-					TOutputDataVertex<FTrigger>(METASOUND_GET_PARAM_NAME_AND_METADATA(TriggerOutput))
-					)
-				);
-
-			return Interface;
-		}
-		
-		static TUniquePtr<IOperator> CreateOperator(const FBuildOperatorParams& InParams, FBuildResults& OutResults)
-		{
-			using namespace MidiTextTriggerPinNames;
-			using namespace CommonPinNames;
-
-			const FInputVertexInterfaceData& InputData = InParams.InputData;
-			FBoolReadRef InEnabled          = InputData.GetOrCreateDefaultDataReadReference<bool>(METASOUND_GET_PARAM_NAME(Inputs::Enable), InParams.OperatorSettings);
-			FMidiStreamReadRef InMidiStream = InputData.GetOrConstructDataReadReference<FMidiStream>(METASOUND_GET_PARAM_NAME(Inputs::MidiStream), InParams.OperatorSettings);
-			FInt32ReadRef InTrackNumber     = InputData.GetOrCreateDefaultDataReadReference<int32>(METASOUND_GET_PARAM_NAME(Inputs::MidiTrackNumber), InParams.OperatorSettings);
-			FStringReadRef InText           = InputData.GetOrCreateDefaultDataReadReference<FString>(METASOUND_GET_PARAM_NAME(TextInput), InParams.OperatorSettings);
-
-			return MakeUnique<FMidiTextTriggerOperator_V0>(InParams, InEnabled, InMidiStream, InTrackNumber, InText);
-		}
-
-		FMidiTextTriggerOperator_V0(
-			const FBuildOperatorParams& InParams,
-			const FBoolReadRef& InEnabled,
-			const FMidiStreamReadRef& InMidiStream,
-			const FInt32ReadRef& InTrackNumber,
-			const FStringReadRef& InText)
-		: FMidiTextTriggerOperator_V1(InParams, InEnabled, InMidiStream, InText)
-		, TrackNumberInPin(InTrackNumber)
-		{
-			Reset(InParams);
-		}
-
-		virtual void BindInputs(FInputVertexInterfaceData& InVertexData) override
-		{
-			FMidiTextTriggerOperator_V1::BindInputs(InVertexData);
-			InVertexData.BindReadVertex(CommonPinNames::Inputs::MidiTrackNumberName, TrackNumberInPin);
-		}
-
-		virtual void Execute() override
-		{
-			TriggerOutPin->AdvanceBlock();
-
-			if (!*EnableInPin)
-			{
-				return;
-			}
-
-			if (const int32 TrackIndex = *TrackNumberInPin; TrackIndex >= 0)
-			{
-				const TArray<FMidiStreamEvent>& MidiEvents = MidiStreamInPin->GetEventsInBlock();
-				for (auto& Event : MidiEvents)
-				{
-					if (Event.TrackIndex != TrackIndex || !Event.MidiMessage.IsText())
-					{
-						continue;
-					}
-
-					const FString* EventString = MidiStreamInPin->GetMidiTrackText(TrackIndex, Event.MidiMessage.GetTextIndex());
-					if (EventString && EventString->Compare(*TextInPin) == 0)
-					{
-						TriggerOutPin->TriggerFrame(Event.BlockSampleFrameIndex);
-					}
-				}
-			}
-		}
-
-	private:
-		FInt32ReadRef TrackNumberInPin;
-	};
-
-	class FMidiTextTriggerNode_V1 : public FNodeFacade
+	class FMidiTextTriggerNode_V1 final : public FNodeFacade
 	{
 	public:
 		explicit FMidiTextTriggerNode_V1(const FNodeInitData& InInitData)
@@ -271,17 +160,6 @@ namespace HarmonixMetasound
 	};
 
 	METASOUND_REGISTER_NODE(FMidiTextTriggerNode_V1)
-
-	class FMidiTextTriggerNode_V0 : public FNodeFacade
-	{
-	public:
-		explicit FMidiTextTriggerNode_V0(const FNodeInitData& InInitData)
-			: FNodeFacade(InInitData.InstanceName, InInitData.InstanceID, TFacadeOperatorClass<FMidiTextTriggerOperator_V0>())
-		{}
-		virtual ~FMidiTextTriggerNode_V0() override = default;
-	};
-
-	METASOUND_REGISTER_NODE(FMidiTextTriggerNode_V0)
 }
 
 #undef LOCTEXT_NAMESPACE // "HarmonixMetaSound"
