MidiClockOffset Node

Simple MIDI Clock operator that takes an incoming clock and adds an offset to the current time in Bars, Beats, and Ms

#tests PIE

diff --git a/Engine/Plugins/Runtime/Harmonix/Source/HarmonixMetasound/Private/HarmonixMetasound/Nodes/MidiClockOffsetNode.cpp b/Engine/Plugins/Runtime/Harmonix/Source/HarmonixMetasound/Private/HarmonixMetasound/Nodes/MidiClockOffsetNode.cpp
new file mode 100644
--- /dev/null
+++ b/Engine/Plugins/Runtime/Harmonix/Source/HarmonixMetasound/Private/HarmonixMetasound/Nodes/MidiClockOffsetNode.cpp
@@ -0,0 +1,265 @@
+﻿// Copyright Epic Games, Inc. All Rights Reserved.
+
+#include "MetasoundExecutableOperator.h"
+#include "MetasoundFacade.h"
+#include "MetasoundNodeInterface.h"
+#include "MetasoundParamHelper.h"
+#include "MetasoundSampleCounter.h"
+#include "MetasoundStandardNodesCategories.h"
+#include "MetasoundVertex.h"
+
+#include "HarmonixMetasound/Common.h"
+#include "HarmonixMetasound/DataTypes/MidiStream.h"
+#include "HarmonixMetasound/DataTypes/MidiClock.h"
+#include "HarmonixMetasound/DataTypes/MidiClockEvent.h"
+#include "HarmonixMetasound/DataTypes/MidiAsset.h"
+#include "HarmonixMetasound/DataTypes/MusicTransport.h"
+
+#include "HarmonixMidi/MidiPlayCursor.h"
+#include "HarmonixMidi/MidiVoiceId.h"
+
+#define LOCTEXT_NAMESPACE "HarmonixMetaSound"
+
+namespace HarmonixMetasound::Nodes::MidiClockOffset
+{
+	using namespace Metasound;
+
+	FNodeClassName GetClassName()
+	{
+		return { HarmonixNodeNamespace, TEXT("MidiClockOffsetNode"), ""};
+	}
+
+	int32 GetCurrentMajorVersion()
+	{
+		return 0;
+	}
+
+	namespace Inputs
+	{
+		DEFINE_INPUT_METASOUND_PARAM(OffsetMs, "Offset (Ms)", "How much to offset the incoming clock by, in Milliseconds");
+		DEFINE_INPUT_METASOUND_PARAM(OffsetBars, "Offset (Bars)", "How much to offset the incoming clock by, in Bars");
+		DEFINE_INPUT_METASOUND_PARAM(OffsetBeats, "Offset (Beats)", "How much to offset the incoming clock by, in Beats");
+		DEFINE_METASOUND_PARAM_ALIAS(MidiClock, CommonPinNames::Inputs::MidiClock);
+	}
+
+	namespace Outputs
+	{
+		DEFINE_METASOUND_PARAM_ALIAS(MidiClock, CommonPinNames::Outputs::MidiClock);
+	}
+	
+	class FMidiClockOffsetOperator : public TExecutableOperator<FMidiClockOffsetOperator>, public FMusicTransportControllable
+	{
+	public:
+		static const FNodeClassMetadata& GetNodeInfo()
+		{
+			auto InitNodeInfo = []() -> FNodeClassMetadata
+			{
+				FNodeClassMetadata Info;
+				Info.ClassName			= GetClassName();
+				Info.MajorVersion		= GetCurrentMajorVersion();
+				Info.MinorVersion		= 1;
+				Info.DisplayName		= METASOUND_LOCTEXT("MIDIClockFollower_DisplayName", "MIDI Clock Offset");
+				Info.Description		= METASOUND_LOCTEXT("MIDIClockFollower_Description", "Offset the incoming clock by some combination of Bars, Beats, and Milliseconds.");
+				Info.Author				= PluginAuthor;
+				Info.PromptIfMissing	= PluginNodeMissingPrompt;
+				Info.DefaultInterface	= GetVertexInterface();
+				Info.CategoryHierarchy.Emplace(NodeCategories::Music);
+
+				return Info;
+			};
+
+			static const FNodeClassMetadata Info = InitNodeInfo();
+			return Info;
+		}
+		
+		static const FVertexInterface& GetVertexInterface()
+		{
+			static const FVertexInterface Interface(
+				FInputVertexInterface(
+					TInputDataVertex<FMidiClock>(METASOUND_GET_PARAM_NAME_AND_METADATA(Inputs::MidiClock)),
+					TInputDataVertex<int32>(METASOUND_GET_PARAM_NAME_AND_METADATA(Inputs::OffsetBars)),
+					TInputDataVertex<float>(METASOUND_GET_PARAM_NAME_AND_METADATA(Inputs::OffsetBeats)),
+					TInputDataVertex<float>(METASOUND_GET_PARAM_NAME_AND_METADATA(Inputs::OffsetMs))
+				),
+				FOutputVertexInterface(
+					TOutputDataVertex<FMidiClock>(METASOUND_GET_PARAM_NAME_AND_METADATA(Outputs::MidiClock))
+				)
+			);
+
+			return Interface;
+		}
+		
+		static TUniquePtr<IOperator> CreateOperator(const FBuildOperatorParams& InParams, FBuildResults& OutResults)
+		{
+			const FOperatorSettings& Settings = InParams.OperatorSettings;
+			const FInputVertexInterfaceData& InputData = InParams.InputData;
+			FMidiClockReadRef InMidiClock = InputData.GetOrConstructDataReadReference<FMidiClock>(METASOUND_GET_PARAM_NAME(Inputs::MidiClock), Settings);
+			FInt32ReadRef InOffsetBars = InputData.GetOrConstructDataReadReference<int32>(METASOUND_GET_PARAM_NAME(Inputs::OffsetBars), 0);
+			FFloatReadRef InOffsetBeats = InputData.GetOrConstructDataReadReference<float>(METASOUND_GET_PARAM_NAME(Inputs::OffsetBeats), 0.0f);
+			FFloatReadRef InOffsetMs = InputData.GetOrConstructDataReadReference<float>(METASOUND_GET_PARAM_NAME(Inputs::OffsetMs), 0.0f);
+			return MakeUnique<FMidiClockOffsetOperator>(InParams.OperatorSettings, InMidiClock, InOffsetBars, InOffsetBeats, InOffsetMs);	
+		}
+
+		FMidiClockOffsetOperator(const FOperatorSettings& InSettings, const FMidiClockReadRef& InMidiClock, const FInt32ReadRef& InOffsetBars, const FFloatReadRef& InOffsetBeats, const FFloatReadRef& InOffsetMs)
+			: FMusicTransportControllable(EMusicPlayerTransportState::Prepared)
+			, MidiClockIn(InMidiClock)
+			, OffsetBarsInPin(InOffsetBars)
+			, OffsetBeatsInPin(InOffsetBeats)
+			, OffsetMsInPin(InOffsetMs)
+			, MidiClockOut(FMidiClockWriteRef::CreateNew(InSettings))
+			, MidiClockEventCursor(MidiClockOut)
+			, BlockSize(InSettings.GetNumFramesPerBlock())
+		{
+			TSharedPtr<FMidiFileData> ConductorMidiData = FMidiClock::MakeClockConductorMidiData(120.0f, 4, 4);
+			MidiClockOut->AttachToMidiResource(ConductorMidiData, true, PrerollBars);
+		}
+
+		virtual void BindInputs(FInputVertexInterfaceData& InVertexData) override
+		{
+			InVertexData.BindReadVertex(METASOUND_GET_PARAM_NAME(Inputs::MidiClock), MidiClockIn);
+			InVertexData.BindReadVertex(METASOUND_GET_PARAM_NAME(Inputs::OffsetBars), OffsetBarsInPin);
+			InVertexData.BindReadVertex(METASOUND_GET_PARAM_NAME(Inputs::OffsetBeats), OffsetBeatsInPin);
+			InVertexData.BindReadVertex(METASOUND_GET_PARAM_NAME(Inputs::OffsetMs), OffsetMsInPin);
+			MidiClockOut->AttachToTimeAuthority(*MidiClockIn);
+		}
+		
+		virtual void BindOutputs(FOutputVertexInterfaceData& InVertexData) override
+		{
+			InVertexData.BindReadVertex(METASOUND_GET_PARAM_NAME(Outputs::MidiClock), MidiClockOut);
+		}
+
+		virtual void Reset(const FResetParams& Params)
+		{
+			BlockSize = Params.OperatorSettings.GetNumFramesPerBlock();
+			CurrentBlockSpanStart = 0;
+			MidiClockOut->ResetAndStart(0, true);
+
+			PrevOffsetBars = 0;
+			PrevOffsetBeats = 0.0f;
+			PrevOffsetMs = 0.0f;
+		}
+
+		int32 GetTickWithOffset(int32 InTick, int32 OffsetBars, float OffsetBeats, float OffsetMs)
+		{
+			int32 OffsetTick = InTick;
+			if (!FMath::IsNearlyZero(OffsetBeats) || OffsetBars != 0)
+			{
+				int32 BeatsPerBar = 0;
+				FMusicTimestamp OffsetTimestamp = MidiClockIn->GetBarMap().TickToMusicTimestamp(InTick, &BeatsPerBar);
+				if (BeatsPerBar != 0)
+				{
+					float Beats = ((OffsetTimestamp.Bar - 1) * BeatsPerBar) + (OffsetTimestamp.Beat - 1);
+					Beats += OffsetBars * BeatsPerBar + OffsetBeats;
+					OffsetTimestamp.Bar = 1 + FMath::Floor(Beats / BeatsPerBar);
+					OffsetTimestamp.Beat = 1 + FMath::Fmod(BeatsPerBar + FMath::Fmod(Beats, BeatsPerBar), BeatsPerBar);
+					OffsetTick = MidiClockIn->GetBarMap().MusicTimestampToTick(OffsetTimestamp);
+				}
+			}
+			const float Ms = MidiClockIn->GetSongMaps().TickToMs(OffsetTick);
+			OffsetTick = MidiClockIn->GetSongMaps().MsToTick(OffsetMs + Ms);
+
+			return OffsetTick;
+		}
+		
+		virtual void Execute()
+		{
+			float OffsetMs = *OffsetMsInPin;
+			int32 OffsetBars = *OffsetBarsInPin;
+			float OffsetBeats = *OffsetBeatsInPin;
+			MidiClockOut->PrepareBlock();
+			MidiClockOut->CopySpeedAndTempoChanges(MidiClockIn.Get());
+			const TArray<FMidiClockEvent>& MidiClockEvents = MidiClockIn->GetMidiClockEventsInBlock();
+			for (int32 EventIndex = 0; EventIndex < MidiClockEvents.Num(); ++EventIndex)
+			{
+				const FMidiClockEvent& Event = MidiClockEvents[EventIndex];
+				switch (Event.Type)
+				{
+				case FMidiClockEvent::EType::SeekTo:
+				case FMidiClockEvent::EType::Reset:
+					{
+						const int32 Tick = GetTickWithOffset(Event.StartTick, OffsetBars, OffsetBeats, OffsetMs);
+						FMusicSeekTarget SeekTarget;
+						SeekTarget.Type = ESeekPointType::Millisecond;
+						SeekTarget.Ms = MidiClockOut->GetSongMaps().TickToMs(Tick);
+								
+						MidiClockOut->SeekTo(Event.BlockFrameIndex, SeekTarget, PrerollBars);
+						break;
+					}
+				case FMidiClockEvent::EType::SeekThru:
+					{
+						const int32 Tick = GetTickWithOffset(Event.StartTick, OffsetBars, OffsetBeats, OffsetMs);
+						FMusicSeekTarget SeekTarget;
+						SeekTarget.Type = ESeekPointType::Millisecond;
+						SeekTarget.Ms = MidiClockOut->GetSongMaps().TickToMs(Tick + 1);
+							
+						MidiClockOut->SeekTo(Event.BlockFrameIndex, SeekTarget, PrerollBars);
+						break;
+					}
+				case FMidiClockEvent::EType::AdvanceThru:
+					{
+						const int32 Tick = GetTickWithOffset(Event.StartTick, OffsetBars, OffsetBeats, OffsetMs);
+						
+						// if our offset changed while we were advancing, seek to the new offset first
+						if (!FMath::IsNearlyEqual(PrevOffsetMs, OffsetMs) || !FMath::IsNearlyEqual(PrevOffsetBeats, OffsetBeats) || PrevOffsetBars != OffsetBars)
+						{
+							PrevOffsetMs = OffsetMs;
+							PrevOffsetBars = OffsetBars;
+							PrevOffsetBeats = OffsetBeats;
+
+							FMusicSeekTarget SeekTarget;
+							SeekTarget.Type = ESeekPointType::Millisecond;
+							SeekTarget.Ms = MidiClockOut->GetSongMaps().TickToMs(Tick);
+								
+							MidiClockOut->SeekTo(Event.BlockFrameIndex, SeekTarget, PrerollBars);
+						}
+
+						const float AdvanceToMs = MidiClockOut->GetSongMaps().TickToMs(Tick);
+						const float ClockInSpeed = MidiClockIn->GetSpeedAtBlockSampleFrame(EventIndex);
+						const float AdvanceRatio = MidiClockIn->GetSongMaps().GetTempoAtTick(Event.StartTick) / MidiClockOut->GetSongMaps().GetTempoAtTick(Tick);
+						MidiClockOut->InformOfCurrentAdvanceRate(ClockInSpeed * AdvanceRatio);
+						MidiClockOut->AdvanceHiResToMs(Event.BlockFrameIndex, AdvanceToMs, true);
+						break;
+					}
+				case FMidiClockEvent::EType::Loop:
+					{
+						break;
+					}
+							
+				}
+			}
+		}
+
+	protected:
+		//** INPUTS
+		FMidiClockReadRef MidiClockIn;
+		FInt32ReadRef OffsetBarsInPin;
+		FFloatReadRef OffsetBeatsInPin;
+		FFloatReadRef OffsetMsInPin;
+		int32 PrerollBars = 0;
+
+		//** OUTPUTS
+		FMidiClockWriteRef MidiClockOut;
+
+		//** DATA
+		FMidiClockEventCursor MidiClockEventCursor;
+		FSampleCount BlockSize      = 0;
+		int32 CurrentBlockSpanStart = 0;
+
+		float PrevOffsetMs = 0.0f;
+		int32 PrevOffsetBars = 0;
+		float PrevOffsetBeats = 0.0f;
+	};
+
+	class FMidiClockOffsetNode : public FNodeFacade
+	{
+	public:
+		FMidiClockOffsetNode(const FNodeInitData& InInitData)
+			: FNodeFacade(InInitData.InstanceName, InInitData.InstanceID, TFacadeOperatorClass<FMidiClockOffsetOperator>())
+		{}
+	virtual ~FMidiClockOffsetNode() = default;
+	};
+
+	METASOUND_REGISTER_NODE(FMidiClockOffsetNode)
+}
+
+#undef LOCTEXT_NAMESPACE // "HarmonixMetaSound"
\ No newline at end of file
diff --git a/Engine/Plugins/Runtime/Harmonix/Source/HarmonixMetasound/Public/HarmonixMetasound/Nodes/MidiClockOffsetNode.h b/Engine/Plugins/Runtime/Harmonix/Source/HarmonixMetasound/Public/HarmonixMetasound/Nodes/MidiClockOffsetNode.h
new file mode 100644
--- /dev/null
+++ b/Engine/Plugins/Runtime/Harmonix/Source/HarmonixMetasound/Public/HarmonixMetasound/Nodes/MidiClockOffsetNode.h
@@ -0,0 +1,26 @@
+// Copyright Epic Games, Inc. All Rights Reserved.
+
+#pragma once
+
+#include "MetasoundNodeInterface.h"
+
+#include "HarmonixMetasound/Common.h"
+
+namespace HarmonixMetasound::Nodes::MidiClockOffset
+{
+	HARMONIXMETASOUND_API Metasound::FNodeClassName GetClassName();
+	HARMONIXMETASOUND_API int32 GetCurrentMajorVersion();
+
+	namespace Inputs
+	{
+		DECLARE_METASOUND_PARAM_EXTERN(OffsetMs);
+		DECLARE_METASOUND_PARAM_EXTERN(OffsetBars);
+		DECLARE_METASOUND_PARAM_EXTERN(OffsetBeats);
+		DECLARE_METASOUND_PARAM_EXTERN(MidiClock);
+	}
+
+	namespace Outputs
+	{
+		DECLARE_METASOUND_PARAM_EXTERN(MidiClock);
+	}
+}
\ No newline at end of file
