[Harmonix] Clean up MidiConstants.h (code standard, removed unused, namespace instead of class with static fields/methods)

#tests automated tests

diff --git a/Engine/Plugins/Runtime/Harmonix/Source/HarmonixDsp/Private/HarmonixDsp/FusionSampler/FusionSampler.cpp b/Engine/Plugins/Runtime/Harmonix/Source/HarmonixDsp/Private/HarmonixDsp/FusionSampler/FusionSampler.cpp
--- a/Engine/Plugins/Runtime/Harmonix/Source/HarmonixDsp/Private/HarmonixDsp/FusionSampler/FusionSampler.cpp
+++ b/Engine/Plugins/Runtime/Harmonix/Source/HarmonixDsp/Private/HarmonixDsp/FusionSampler/FusionSampler.cpp
@@ -126,7 +126,7 @@ void FFusionSampler::VoicePoolWillDestruct(const FFusionVoicePool* InPool)
 void FFusionSampler::NoteOn(FMidiVoiceId InVoiceId, int8 InMidiNoteNumber, int8 InVelocity, int8 InMidiChannel, int32 InEventTick, int32 InCurrentTick, float InOffsetMs)
 {
 	// we ignore InMidiChannel because this is a single channel instrument!
-	if (InMidiNoteNumber > MidiConstants::kMaxNote)
+	if (InMidiNoteNumber > Harmonix::Midi::Constants::GMaxNote)
 	{
 		return;
 	}
@@ -154,7 +154,7 @@ void FFusionSampler::NoteOn(FMidiVoiceId InVoiceId, int8 InMidiNoteNumber, int8
 void FFusionSampler::NoteOnWithFrameOffset(FMidiVoiceId InVoiceId, int8 InMidiNoteNumber, int8 InVelocity, int8 InMidiChannel, int32 InNumFrames)
 {
 	// we ignore InMidiChannel because this is a single channel instrument!
-	if (InMidiNoteNumber > MidiConstants::kMaxNote)
+	if (InMidiNoteNumber > Harmonix::Midi::Constants::GMaxNote)
 	{
 		return;
 	}
@@ -281,7 +281,7 @@ void FFusionSampler::ResetNoteActions(bool ClearNotes)
 void FFusionSampler::ResetNoteStatus()
 {
 	FScopeLock Lock(&GetBusLock());
-	for (int32 NoteIdx = 0; NoteIdx < MidiConstants::kMaxNumNotes; ++NoteIdx)
+	for (int32 NoteIdx = 0; NoteIdx < Harmonix::Midi::Constants::GMaxNumNotes; ++NoteIdx)
 	{
 		NoteStatus[NoteIdx].KeyedOn = false;
 		NoteStatus[NoteIdx].NumActiveVoices = 0;
@@ -399,7 +399,7 @@ int32 FFusionSampler::GatherMatchingKeyzones(uint8 TransposedNote, uint8 InVeloc
 
 int32 FFusionSampler::StartNote(FMidiVoiceId InVoiceId, uint8 InMidiNoteNumber, uint8 InVelocity, bool IsNoteOn, int32 InEventTick, int32 InTriggerTick, float InOffsetMs)
 {
-	check(InMidiNoteNumber <= MidiConstants::kMaxNote);
+	check(InMidiNoteNumber <= Harmonix::Midi::Constants::GMaxNote);
 
 	if (VoicePool == nullptr)
 	{
@@ -855,9 +855,9 @@ void FFusionSampler::ResetPatchRelatedState()
 	SetPortamentoMode(EPortamentoMode::Legato);
 	SetPortamentoTime(Defaults::kPortamentoTimeMs / 1000.0f);
 
-	FMemory::Memset(LastStartLayerSelect, -1, sizeof(int8) * MidiConstants::kMaxNumNotes);
-	FMemory::Memset(LastStopLayerSelect, -1, sizeof(int8) * MidiConstants::kMaxNumNotes);
-	FMemory::Memset(LastVelocity, 0, sizeof(int8) * MidiConstants::kMaxNumNotes);
+	FMemory::Memset(LastStartLayerSelect, -1, sizeof(int8) * Harmonix::Midi::Constants::GMaxNumNotes);
+	FMemory::Memset(LastStopLayerSelect, -1, sizeof(int8) * Harmonix::Midi::Constants::GMaxNumNotes);
+	FMemory::Memset(LastVelocity, 0, sizeof(int8) * Harmonix::Midi::Constants::GMaxNumNotes);
 
 	TimeStretchEnvelopeOverride = -1;
 
@@ -1255,7 +1255,7 @@ void FFusionSampler::Process(uint32 InSliceIdx, uint32 InSubSliceIdx, TAudioBuff
 	}
 }
 
-void FFusionSampler::GetController(MidiConstants::EControllerID InController, int8& Msb, int8& Lsb, int8 InMidiChannel) const
+void FFusionSampler::GetController(Harmonix::Midi::Constants::EControllerID InController, int8& Msb, int8& Lsb, int8 InMidiChannel) const
 {
 	// we ignore InMidiChannel because this is a single channel instrument!
 
@@ -1269,33 +1269,36 @@ void FFusionSampler::GetController(MidiConstants::EControllerID InController, in
 	Lsb = 0; // default case.
 
 	float ValueP;
+
+	using namespace Harmonix::Midi::Constants;
+	
 	switch (InController)
 	{
-	case MidiConstants::EControllerID::BankSelection:
+	case EControllerID::BankSelection:
 	{
 		break;
 	}
 
-	case MidiConstants::EControllerID::Volume:
+	case EControllerID::Volume:
 	{
 		float Exp = HarmonixDsp::DBToMidiLinear(MidiChannelVolume);
 		Msb = (int8)(Exp * 127.0f);
 		break;
 	}
 
-	case MidiConstants::EControllerID::Expression:
+	case EControllerID::Expression:
 	{
 		float Exp = HarmonixDsp::LinearToDB(ExpressionGain);
 		Exp = HarmonixDsp::DBToMidiLinear(Exp);
 		Msb = (int8)(Exp * 127.0f);
 		break;
 	}
-	case MidiConstants::EControllerID::PanRight:
+	case EControllerID::PanRight:
 	{
 		Msb = (int8)((PanSettings.Detail.Pan + 1.0f) * 64.0f);
 		break;
 	}
-	case MidiConstants::EControllerID::Release:
+	case EControllerID::Release:
 	{
 		Min = 0.005f;
 		Max = 2.000f;
@@ -1306,7 +1309,7 @@ void FFusionSampler::GetController(MidiConstants::EControllerID InController, in
 		break;
 	}
 
-	case MidiConstants::EControllerID::Attack:
+	case EControllerID::Attack:
 	{
 		Min = 0.005f;
 		Max = 2.000f;
@@ -1317,136 +1320,136 @@ void FFusionSampler::GetController(MidiConstants::EControllerID InController, in
 		break;
 	}
 
-	case MidiConstants::EControllerID::PortamentoSwitch:
+	case EControllerID::PortamentoSwitch:
 	{
 		Msb = GetIsPortamentoEnabled() ? 127 : 0;
 		break;
 	}
 
-	case MidiConstants::EControllerID::PortamentoTime:
+	case EControllerID::PortamentoTime:
 	{
 		Msb = (int8)(gPortamentoTimeInterp.InverseClamped(PortamentoTimeMs) * 127.0f);
 		break;
 	}
 
-	case MidiConstants::EControllerID::FilterFrequency:
+	case EControllerID::FilterFrequency:
 	{
 		Msb = (int8)(gFreqInterp.InverseClamped(FilterSettings.Freq) * 127.0f);
 		break;
 	}
 
-	case MidiConstants::EControllerID::FilterQ:
+	case EControllerID::FilterQ:
 	{
 		Msb = (int8)(gFilterQInterp.InverseClamped(FilterSettings.Q) * 127.0f);
 		break;
 	}
 
-	case MidiConstants::EControllerID::CoarsePitchBend:
+	case EControllerID::CoarsePitchBend:
 	{
 		Msb = (int8)(GetPitchBend() * 127.0f);
 		break;
 	}
 
-	case MidiConstants::EControllerID::SampleStartTime:
+	case EControllerID::SampleStartTime:
 	{
 		Msb = (int8)(gStartPointTimeInterp.InverseClamped(GetStartPointMs()) * 127.0f);
 		break;
 	}
 
-	case MidiConstants::EControllerID::LFO0Frequency:
-	case MidiConstants::EControllerID::LFO1Frequency:
+	case EControllerID::LFO0Frequency:
+	case EControllerID::LFO1Frequency:
 	{
-		int32 LfoIdx = (InController == MidiConstants::EControllerID::LFO1Frequency);
+		int32 LfoIdx = (InController == EControllerID::LFO1Frequency);
 
 		float Freq = LfoSettings[LfoIdx].Freq;
 		Msb = (int8)(gLfoFreqInterp.InverseClamped(Freq) * 127.0f);
 		break;
 	}
 
-	case MidiConstants::EControllerID::LFO0Depth:
-	case MidiConstants::EControllerID::LFO1Depth:
+	case EControllerID::LFO0Depth:
+	case EControllerID::LFO1Depth:
 	{
-		int32 LfoIdx = (InController == MidiConstants::EControllerID::LFO1Depth);
+		int32 LfoIdx = (InController == EControllerID::LFO1Depth);
 
 		Msb = (int8)(LfoSettings[LfoIdx].Depth * 127.0f);
 		break;
 	}
 
-	case MidiConstants::EControllerID::BitCrushWetMix:
+	case EControllerID::BitCrushWetMix:
 	{
 		break;
 	}
 
-	case MidiConstants::EControllerID::BitCrushLevel:
+	case EControllerID::BitCrushLevel:
 	{
 		break;
 	}
 
-	case MidiConstants::EControllerID::BitCrushSampleHold:
+	case EControllerID::BitCrushSampleHold:
 	{
 		break;
 	}
 
-	case MidiConstants::EControllerID::DelayTime:
+	case EControllerID::DelayTime:
 	{
 		break;
 	}
 
-	case MidiConstants::EControllerID::DelayDryGain:
+	case EControllerID::DelayDryGain:
 	{
 		break;
 	}
 
-	case MidiConstants::EControllerID::DelayWetGain:
+	case EControllerID::DelayWetGain:
 	{
 		break;
 	}
 
-	case MidiConstants::EControllerID::DelayFeedback:
+	case EControllerID::DelayFeedback:
 	{
 		break;
 	}
-	case MidiConstants::EControllerID::DelayEQEnabled:
+	case EControllerID::DelayEQEnabled:
 	{
 		break;
 	}
-	case MidiConstants::EControllerID::DelayEQType:
+	case EControllerID::DelayEQType:
 	{
 		break;
 	}
-	case MidiConstants::EControllerID::DelayEQFreq:
+	case EControllerID::DelayEQFreq:
 	{
 		break;
 	}
-	case MidiConstants::EControllerID::DelayEQQ:
+	case EControllerID::DelayEQQ:
 	{
 		break;
 	}
-	case MidiConstants::EControllerID::DelayLFOEnabled:
+	case EControllerID::DelayLFOEnabled:
 	{
 		break;
 	}
-	case MidiConstants::EControllerID::DelayLFOBeatSync:
+	case EControllerID::DelayLFOBeatSync:
 	{
 		break;
 	}
-	case MidiConstants::EControllerID::DelayLFORate:
+	case EControllerID::DelayLFORate:
 	{
 		break;
 	}
-	case MidiConstants::EControllerID::DelayLFODepth:
+	case EControllerID::DelayLFODepth:
 	{
 		break;
 	}
-	case MidiConstants::EControllerID::DelayStereoType:
+	case EControllerID::DelayStereoType:
 	{
 		break;
 	}
-	case MidiConstants::EControllerID::DelayPanLeft:
+	case EControllerID::DelayPanLeft:
 	{
 		break;
 	}
-	case MidiConstants::EControllerID::DelayPanRight:
+	case EControllerID::DelayPanRight:
 	{
 	}
 	default:
@@ -1454,81 +1457,83 @@ void FFusionSampler::GetController(MidiConstants::EControllerID InController, in
 	}
 }
 
-void FFusionSampler::SetController(MidiConstants::EControllerID InController, float InValue)
+void FFusionSampler::SetController(Harmonix::Midi::Constants::EControllerID InController, float InValue)
 {
 	FScopeLock Lock(&GetBusLock());
 	// we ignore InMidiChannel because this is a single channel instrument!
 
+	using namespace Harmonix::Midi::Constants;
+
 	switch (InController)
 	{
-	case MidiConstants::EControllerID::BankSelection: 
+	case EControllerID::BankSelection: 
 	{
 		// swap out the fusion patch with the new desired patch
 
 		// ApplyPatchSettings(); 
 		break;
 	}
-	case MidiConstants::EControllerID::Volume: 
+	case EControllerID::Volume: 
 	{ 
 		SetMidiChannelVolume(InValue); 
 		break;
 	}
-	case MidiConstants::EControllerID::Expression: 
+	case EControllerID::Expression: 
 	{ 
 		SetMidiExpressionGain(InValue); 
 		break;
 	}
-	case MidiConstants::EControllerID::PanRight: 
+	case EControllerID::PanRight: 
 	{ 
 		PanSettings.Detail.Pan = InValue; 
 		break; 
 	}
-	case MidiConstants::EControllerID::Release: 
+	case EControllerID::Release: 
 	{ 
 		AdsrVolumeSettings.ReleaseTime = InValue; 
 		AdsrVolumeSettings.BuildReleaseTable(); 
 		break; 
 	}
-	case MidiConstants::EControllerID::Attack: 
+	case EControllerID::Attack: 
 	{ 
 		AdsrVolumeSettings.AttackTime = InValue; 
 		AdsrVolumeSettings.BuildAttackTable(); 
 		break; 
 	}
-	case MidiConstants::EControllerID::PortamentoSwitch: 
+	case EControllerID::PortamentoSwitch: 
 	{ 
 		SetIsPortamentoEnabled(InValue > 0.0f); 
 		break; 
 	}
-	case MidiConstants::EControllerID::PortamentoTime: 
+	case EControllerID::PortamentoTime: 
 	{ 
 		SetPortamentoTime(gPortamentoTimeInterp.EvalClamped(InValue) / 1000);
 		break; 
 	}
-	case MidiConstants::EControllerID::FilterFrequency: 
+	case EControllerID::FilterFrequency: 
 	{ 
 		FilterSettings.Freq = gFreqInterp.EvalClamped(InValue); 
 		break; 
 	}
-	case MidiConstants::EControllerID::FilterQ: 
+	case EControllerID::FilterQ: 
 	{ 
 		FilterSettings.Q = gFilterQInterp.EvalClamped(InValue); 
 		break; 
 	}
-	case MidiConstants::EControllerID::CoarsePitchBend: 
+	case EControllerID::CoarsePitchBend: 
 	{ 
 		SetPitchBend(InValue); 
 		break; 
 	}
-	case MidiConstants::EControllerID::SampleStartTime: 
+	case EControllerID::SampleStartTime: 
 	{ 
 		SetStartPointMs(gStartPointTimeInterp.EvalClamped(InValue));
 		break; 
 	}
-	case MidiConstants::EControllerID::LFO0Frequency:
-	case MidiConstants::EControllerID::LFO1Frequency: 
+	case EControllerID::LFO0Frequency:
+	case EControllerID::LFO1Frequency: 
 	{ 
-		int32 LfoIdx = (InController == MidiConstants::EControllerID::LFO1Frequency);
+		int32 LfoIdx = (InController == EControllerID::LFO1Frequency);
 
 		if (LfoSettings[LfoIdx].Freq != InValue)
 		{
@@ -1538,10 +1543,10 @@ void FFusionSampler::SetController(MidiConstants::EControllerID InController, fl
 		}
 		break; 
 	}
-	case MidiConstants::EControllerID::LFO0Depth:
-	case MidiConstants::EControllerID::LFO1Depth: 
+	case EControllerID::LFO0Depth:
+	case EControllerID::LFO1Depth: 
 	{ 
-		int32 LfoIdx = (InController == MidiConstants::EControllerID::LFO1Depth);
+		int32 LfoIdx = (InController == EControllerID::LFO1Depth);
 		if (LfoSettings[LfoIdx].Depth != InValue)
 		{
 			LfoSettings[LfoIdx].Depth = InValue;
@@ -1550,120 +1555,120 @@ void FFusionSampler::SetController(MidiConstants::EControllerID InController, fl
 		}
 		break; 
 	}
-	case MidiConstants::EControllerID::BitCrushWetMix: 
+	case EControllerID::BitCrushWetMix: 
 	{
 		break; 
 	}
-	case MidiConstants::EControllerID::BitCrushLevel: 
+	case EControllerID::BitCrushLevel: 
 	{
 		break; 
 	}
-	case MidiConstants::EControllerID::BitCrushSampleHold: 
+	case EControllerID::BitCrushSampleHold: 
 	{
 		break; 
 	}
-	case MidiConstants::EControllerID::DelayTime: 
+	case EControllerID::DelayTime: 
 	{
 		break; 
 	}
-	case MidiConstants::EControllerID::DelayDryGain: 
+	case EControllerID::DelayDryGain: 
 	{
 		break; 
 	}
-	case MidiConstants::EControllerID::DelayWetGain: 
+	case EControllerID::DelayWetGain: 
 	{
 		break; 
 	}
-	case MidiConstants::EControllerID::DelayFeedback: 
+	case EControllerID::DelayFeedback: 
 	{
 		break; 
 	}
-	case MidiConstants::EControllerID::DelayEQEnabled: 
+	case EControllerID::DelayEQEnabled: 
 	{
 		break; 
 	}
-	case MidiConstants::EControllerID::DelayEQType: 
+	case EControllerID::DelayEQType: 
 	{
 		break; 
 	}
-	case MidiConstants::EControllerID::DelayEQFreq: 
+	case EControllerID::DelayEQFreq: 
 	{
 		break; 
 	}
-	case MidiConstants::EControllerID::DelayEQQ: 
+	case EControllerID::DelayEQQ: 
 	{
 		break; 
 	}
-	case MidiConstants::EControllerID::DelayLFOEnabled: 
+	case EControllerID::DelayLFOEnabled: 
 	{
 		break; 
 	}
-	case MidiConstants::EControllerID::DelayLFOBeatSync: 
+	case EControllerID::DelayLFOBeatSync: 
 	{
 		break; 
 	}
-	case MidiConstants::EControllerID::DelayLFORate: 
+	case EControllerID::DelayLFORate: 
 	{
 		break; 
 	}
-	case MidiConstants::EControllerID::DelayLFODepth: 
+	case EControllerID::DelayLFODepth: 
 	{
 		break; 
 	}
-	case MidiConstants::EControllerID::DelayStereoType: 
+	case EControllerID::DelayStereoType: 
 	{
 		break; 
 	}
-	case MidiConstants::EControllerID::DelayPanLeft: 
+	case EControllerID::DelayPanLeft: 
 	{
 		break; 
 	}
-	case MidiConstants::EControllerID::DelayPanRight: 
+	case EControllerID::DelayPanRight: 
 	{
 		break; 
 	}
-	case MidiConstants::EControllerID::TimeStretchEnvelopeOrder:
+	case EControllerID::TimeStretchEnvelopeOrder:
 	{
 		TimeStretchEnvelopeOverride = (int16)InValue;
 		UpdateVoicesForEnvelopeOrderChange();
 		break;
 	}
-	case MidiConstants::EControllerID::SubStreamVol1:
+	case EControllerID::SubStreamVol1:
 	{ 
 		SetSubstreamMidiGain(0, (int32)InValue); 
 		break; 
 	}
-	case MidiConstants::EControllerID::SubStreamVol2: 
+	case EControllerID::SubStreamVol2: 
 	{ 
 		SetSubstreamMidiGain(1, (int32)InValue); 
 		break; 
 	}
-	case MidiConstants::EControllerID::SubStreamVol3: 
+	case EControllerID::SubStreamVol3: 
 	{
 		SetSubstreamMidiGain(2, (int32)InValue); 
 		break;
 	}
-	case MidiConstants::EControllerID::SubStreamVol4: 
+	case EControllerID::SubStreamVol4: 
 	{
 		SetSubstreamMidiGain(3, (int32)InValue);
 		break; 
 	}
-	case MidiConstants::EControllerID::SubStreamVol5: 
+	case EControllerID::SubStreamVol5: 
 	{ 
 		SetSubstreamMidiGain(4, (int32)InValue); 
 		break;
 	}
-	case MidiConstants::EControllerID::SubStreamVol6:
+	case EControllerID::SubStreamVol6:
 	{ 
 		SetSubstreamMidiGain(5, (int32)InValue); 
 		break; 
 	}
-	case MidiConstants::EControllerID::SubStreamVol7: 
+	case EControllerID::SubStreamVol7: 
 	{ 
 		SetSubstreamMidiGain(6, (int32)InValue); 
 		break;
 	}
-	case MidiConstants::EControllerID::SubStreamVol8:
+	case EControllerID::SubStreamVol8:
 	{
 		SetSubstreamMidiGain(7, (int32)InValue); 
 		break; 
@@ -1673,7 +1678,7 @@ void FFusionSampler::SetController(MidiConstants::EControllerID InController, fl
 	}
 }
 
-void FFusionSampler::Set7BitControllerImpl(MidiConstants::EControllerID InController, int8 InValue, int8 InMidiChannel)
+void FFusionSampler::Set7BitControllerImpl(Harmonix::Midi::Constants::EControllerID InController, int8 InValue, int8 InMidiChannel)
 {
 	FScopeLock Lock(&GetBusLock());
 	// we ignore InMidiChannel because this is a single channel instrument!
@@ -1684,115 +1689,117 @@ void FFusionSampler::Set7BitControllerImpl(MidiConstants::EControllerID InContro
 
 	float ValueP = 0.0f;
 
+	using namespace Harmonix::Midi::Constants;
+
 	switch (InController)
 	{
-	case MidiConstants::EControllerID::BankSelection: 
+	case EControllerID::BankSelection: 
 	{ 
 		ValueP = (float)InValue; 
 		break; 
 	}
-	case MidiConstants::EControllerID::Volume: 
+	case EControllerID::Volume: 
 	{ 
 		ValueP = HarmonixDsp::MidiLinearToDB(InValue); 
 		break; 
 	}
-	case MidiConstants::EControllerID::Expression: 
+	case EControllerID::Expression: 
 	{ 
 		ValueP = HarmonixDsp::DBToLinear(HarmonixDsp::MidiLinearToDB(InValue)); 
 		break; 
 	}
-	case MidiConstants::EControllerID::PanRight: 
+	case EControllerID::PanRight: 
 	{ 
 		ValueP = ((float)InValue / 64.0f) - 1.0f; 
 		break; 
 	}
-	case MidiConstants::EControllerID::Release: 
+	case EControllerID::Release: 
 	{ 
 		Min = 0.005f; Max = 2.000f; Range = Max - Min;
 		ValueP = (float)InValue / 127.0f;
 		ValueP = ValueP * Range + Min; 
 		break; 
 	}
-	case MidiConstants::EControllerID::Attack: 
+	case EControllerID::Attack: 
 	{ 
 		Min = 0.005f; Max = 2.000f; Range = Max - Min;
 		ValueP = (float)InValue / 127.0f;
 		ValueP = ValueP * Range + Min; 
 		break; 
 	}
-	case MidiConstants::EControllerID::PortamentoSwitch: 
+	case EControllerID::PortamentoSwitch: 
 	{ 
 		ValueP = (InValue > 0) ? 1.0f : 0.0f; 
 		break; 
 	}
-	case MidiConstants::EControllerID::PortamentoTime:
-	case MidiConstants::EControllerID::FilterFrequency:
-	case MidiConstants::EControllerID::FilterQ:
-	case MidiConstants::EControllerID::CoarsePitchBend:
-	case MidiConstants::EControllerID::SampleStartTime: 
+	case EControllerID::PortamentoTime:
+	case EControllerID::FilterFrequency:
+	case EControllerID::FilterQ:
+	case EControllerID::CoarsePitchBend:
+	case EControllerID::SampleStartTime: 
 	{ 
 		ValueP = (float)InValue / 127.0f; 
 		break; 
 	}
-	case MidiConstants::EControllerID::LFO0Frequency:
-	case MidiConstants::EControllerID::LFO1Frequency: 
+	case EControllerID::LFO0Frequency:
+	case EControllerID::LFO1Frequency: 
 	{ 
 		ValueP = gLfoFreqInterp.EvalClamped((float)InValue / 127.0f); 
 		break; 
 	}
-	case MidiConstants::EControllerID::LFO0Depth:
-	case MidiConstants::EControllerID::LFO1Depth: 
+	case EControllerID::LFO0Depth:
+	case EControllerID::LFO1Depth: 
 	{ 
 		ValueP = (float)InValue / 127.0f; 
 		break; 
 	}
-	case MidiConstants::EControllerID::BitCrushWetMix:
+	case EControllerID::BitCrushWetMix:
 	{ 
 		ValueP = (float)InValue * 100.0f / 127.0f; 
 		break; 
 	}
-	case MidiConstants::EControllerID::BitCrushLevel: 
+	case EControllerID::BitCrushLevel: 
 	{ 
 		ValueP = (float)InValue * 15.0f / 127.0f; 
 		break; 
 	}
-	case MidiConstants::EControllerID::BitCrushSampleHold: 
+	case EControllerID::BitCrushSampleHold: 
 	{ 
 		ValueP = (1.5f + (((float)InValue / 127.0f) * 15));
 		break; 
 	}
-	case MidiConstants::EControllerID::DelayTime:
-	case MidiConstants::EControllerID::DelayDryGain:
-	case MidiConstants::EControllerID::DelayWetGain:
-	case MidiConstants::EControllerID::DelayEQEnabled:
-	case MidiConstants::EControllerID::DelayEQType:
-	case MidiConstants::EControllerID::DelayEQFreq:
-	case MidiConstants::EControllerID::DelayEQQ:
-	case MidiConstants::EControllerID::DelayLFOEnabled:
-	case MidiConstants::EControllerID::DelayLFOBeatSync:
-	case MidiConstants::EControllerID::DelayLFORate:
-	case MidiConstants::EControllerID::DelayLFODepth:
-	case MidiConstants::EControllerID::DelayStereoType:
-	case MidiConstants::EControllerID::DelayPanLeft:
-	case MidiConstants::EControllerID::DelayPanRight:
-	case MidiConstants::EControllerID::DelayFeedback: 
+	case EControllerID::DelayTime:
+	case EControllerID::DelayDryGain:
+	case EControllerID::DelayWetGain:
+	case EControllerID::DelayEQEnabled:
+	case EControllerID::DelayEQType:
+	case EControllerID::DelayEQFreq:
+	case EControllerID::DelayEQQ:
+	case EControllerID::DelayLFOEnabled:
+	case EControllerID::DelayLFOBeatSync:
+	case EControllerID::DelayLFORate:
+	case EControllerID::DelayLFODepth:
+	case EControllerID::DelayStereoType:
+	case EControllerID::DelayPanLeft:
+	case EControllerID::DelayPanRight:
+	case EControllerID::DelayFeedback: 
 	{ 
 		ValueP = (float)InValue / 127.0f; 
 		break; 
 	}
-	case MidiConstants::EControllerID::TimeStretchEnvelopeOrder: 
+	case EControllerID::TimeStretchEnvelopeOrder: 
 	{ 
 		ValueP = (float)InValue; 
 		break; 
 	}
-	case MidiConstants::EControllerID::SubStreamVol1:
-	case MidiConstants::EControllerID::SubStreamVol2:
-	case MidiConstants::EControllerID::SubStreamVol3:
-	case MidiConstants::EControllerID::SubStreamVol4:
-	case MidiConstants::EControllerID::SubStreamVol5:
-	case MidiConstants::EControllerID::SubStreamVol6:
-	case MidiConstants::EControllerID::SubStreamVol7:
-	case MidiConstants::EControllerID::SubStreamVol8: 
+	case EControllerID::SubStreamVol1:
+	case EControllerID::SubStreamVol2:
+	case EControllerID::SubStreamVol3:
+	case EControllerID::SubStreamVol4:
+	case EControllerID::SubStreamVol5:
+	case EControllerID::SubStreamVol6:
+	case EControllerID::SubStreamVol7:
+	case EControllerID::SubStreamVol8: 
 	{ 
 		ValueP = (float)InValue; 
 		break; 
@@ -1802,7 +1809,7 @@ void FFusionSampler::Set7BitControllerImpl(MidiConstants::EControllerID InContro
 	SetController(InController, ValueP);
 }
 
-void FFusionSampler::Set14BitControllerImpl(MidiConstants::EControllerID InController, int16 InValue, int8 InMidiChannel)
+void FFusionSampler::Set14BitControllerImpl(Harmonix::Midi::Constants::EControllerID InController, int16 InValue, int8 InMidiChannel)
 {
 	FScopeLock Lock(&GetBusLock());
 	// we ignore InMidiChannel because this is a single channel instrument!
@@ -1813,29 +1820,31 @@ void FFusionSampler::Set14BitControllerImpl(MidiConstants::EControllerID InContr
 
 	float ValueP;
 
+	using namespace Harmonix::Midi::Constants;
+
 	switch (InController)
 	{
-	case MidiConstants::EControllerID::BankSelection: 
+	case EControllerID::BankSelection: 
 	{
 		ValueP = (float)InValue; 
 		break; 
 	}
-	case MidiConstants::EControllerID::Volume: 
+	case EControllerID::Volume: 
 	{ 
 		ValueP = HarmonixDsp::Midi14BitLinearToDB(InValue); 
 		break; 
 	}
-	case MidiConstants::EControllerID::Expression: 
+	case EControllerID::Expression: 
 	{ 
 		ValueP = HarmonixDsp::DBToLinear(HarmonixDsp::Midi14BitLinearToDB(InValue));
 		break; 
 	}
-	case MidiConstants::EControllerID::PanRight: 
+	case EControllerID::PanRight: 
 	{
 		ValueP = ((float)InValue / 8191.0f) - 1.0f;
 		break;
 	}
-	case MidiConstants::EControllerID::Release: 
+	case EControllerID::Release: 
 	{ 
 		Min = 0.005f; Max = 2.000f; 
 		Range = Max - Min;
@@ -1843,7 +1852,7 @@ void FFusionSampler::Set14BitControllerImpl(MidiConstants::EControllerID InContr
 		ValueP = ValueP * Range + Min; 
 		break; 
 	}
-	case MidiConstants::EControllerID::Attack: 
+	case EControllerID::Attack: 
 	{ 
 		Min = 0.005f; Max = 2.000f; 
 		Range = Max - Min;
@@ -1851,79 +1860,79 @@ void FFusionSampler::Set14BitControllerImpl(MidiConstants::EControllerID InContr
 		ValueP = ValueP * Range + Min; 
 		break; 
 	}
-	case MidiConstants::EControllerID::PortamentoSwitch: 
+	case EControllerID::PortamentoSwitch: 
 	{ 
 		ValueP = (InValue > 0) ? 1.0f : 0.0f; 
 		break; 
 	}
-	case MidiConstants::EControllerID::PortamentoTime:
-	case MidiConstants::EControllerID::FilterFrequency:
-	case MidiConstants::EControllerID::FilterQ:
-	case MidiConstants::EControllerID::CoarsePitchBend:
-	case MidiConstants::EControllerID::SampleStartTime: 
+	case EControllerID::PortamentoTime:
+	case EControllerID::FilterFrequency:
+	case EControllerID::FilterQ:
+	case EControllerID::CoarsePitchBend:
+	case EControllerID::SampleStartTime: 
 	{ 
 		ValueP = (float)InValue / 16383.0f; 
 		break; 
 	}
-	case MidiConstants::EControllerID::LFO0Frequency:
-	case MidiConstants::EControllerID::LFO1Frequency: 
+	case EControllerID::LFO0Frequency:
+	case EControllerID::LFO1Frequency: 
 	{ 
 		ValueP = gLfoFreqInterp.EvalClamped((float)InValue / 16383.0f); 
 		break; 
 	}
-	case MidiConstants::EControllerID::LFO0Depth:
-	case MidiConstants::EControllerID::LFO1Depth: 
+	case EControllerID::LFO0Depth:
+	case EControllerID::LFO1Depth: 
 	{ 
 		ValueP = (float)InValue / 16383.0f;
 		break;
 	}
-	case MidiConstants::EControllerID::BitCrushWetMix: 
+	case EControllerID::BitCrushWetMix: 
 	{
 		ValueP = (float)InValue * 100.0f / 16383.0f;
 		break; 
 	}
-	case MidiConstants::EControllerID::BitCrushLevel:
+	case EControllerID::BitCrushLevel:
 	{ 
 		ValueP = (float)InValue * 15.0f / 16383.0f; 
 		break;
 	}
-	case MidiConstants::EControllerID::BitCrushSampleHold: 
+	case EControllerID::BitCrushSampleHold: 
 	{ 
 		ValueP = (1.5f + (((float)InValue / 16383.0f) * 15));
 		break; 
 	}
-	case MidiConstants::EControllerID::DelayTime:
-	case MidiConstants::EControllerID::DelayDryGain:
-	case MidiConstants::EControllerID::DelayWetGain:
-	case MidiConstants::EControllerID::DelayEQEnabled:
-	case MidiConstants::EControllerID::DelayEQType:
-	case MidiConstants::EControllerID::DelayEQFreq:
-	case MidiConstants::EControllerID::DelayEQQ:
-	case MidiConstants::EControllerID::DelayLFOEnabled:
-	case MidiConstants::EControllerID::DelayLFOBeatSync:
-	case MidiConstants::EControllerID::DelayLFORate:
-	case MidiConstants::EControllerID::DelayLFODepth:
-	case MidiConstants::EControllerID::DelayStereoType:
-	case MidiConstants::EControllerID::DelayPanLeft:
-	case MidiConstants::EControllerID::DelayPanRight:
-	case MidiConstants::EControllerID::DelayFeedback:
+	case EControllerID::DelayTime:
+	case EControllerID::DelayDryGain:
+	case EControllerID::DelayWetGain:
+	case EControllerID::DelayEQEnabled:
+	case EControllerID::DelayEQType:
+	case EControllerID::DelayEQFreq:
+	case EControllerID::DelayEQQ:
+	case EControllerID::DelayLFOEnabled:
+	case EControllerID::DelayLFOBeatSync:
+	case EControllerID::DelayLFORate:
+	case EControllerID::DelayLFODepth:
+	case EControllerID::DelayStereoType:
+	case EControllerID::DelayPanLeft:
+	case EControllerID::DelayPanRight:
+	case EControllerID::DelayFeedback:
 	{ 
 		ValueP = (float)InValue / 16383.0f;
 		break; 
 	}
-	case MidiConstants::EControllerID::TimeStretchEnvelopeOrder: 
+	case EControllerID::TimeStretchEnvelopeOrder: 
 	{ 
 		ValueP = (float)(InValue >> 7); 
 		break; 
 	}
-	case MidiConstants::EControllerID::SubStreamVol1:
-	case MidiConstants::EControllerID::SubStreamVol2:
-	case MidiConstants::EControllerID::SubStreamVol3:
-	case MidiConstants::EControllerID::SubStreamVol4:
-	case MidiConstants::EControllerID::SubStreamVol5:
-	case MidiConstants::EControllerID::SubStreamVol6:
-	case MidiConstants::EControllerID::SubStreamVol7:
-	case MidiConstants::EControllerID::SubStreamVol8: 
+	case EControllerID::SubStreamVol1:
+	case EControllerID::SubStreamVol2:
+	case EControllerID::SubStreamVol3:
+	case EControllerID::SubStreamVol4:
+	case EControllerID::SubStreamVol5:
+	case EControllerID::SubStreamVol6:
+	case EControllerID::SubStreamVol7:
+	case EControllerID::SubStreamVol8: 
 	{ 
 		ValueP = (float)(InValue >> 7); 
 		break; 
diff --git a/Engine/Plugins/Runtime/Harmonix/Source/HarmonixDsp/Private/HarmonixDsp/FusionSampler/FusionVoice.cpp b/Engine/Plugins/Runtime/Harmonix/Source/HarmonixDsp/Private/HarmonixDsp/FusionSampler/FusionVoice.cpp
--- a/Engine/Plugins/Runtime/Harmonix/Source/HarmonixDsp/Private/HarmonixDsp/FusionSampler/FusionVoice.cpp
+++ b/Engine/Plugins/Runtime/Harmonix/Source/HarmonixDsp/Private/HarmonixDsp/FusionSampler/FusionVoice.cpp
@@ -364,7 +364,7 @@ void FFusionVoice::AttackWithTargetNote(uint8 InMidiNoteNumber, float InGain, in
 			if (ErrorMs > 10.0f)
 			{
 				// gotta let the render code skip forward a bit in process otherwise this voice will be out of sync by a lot!
-				StartBeat = (float)InEventTick / MidiConstants::kTicksPerQuarterNote;
+				StartBeat = (float)InEventTick / Harmonix::Midi::Constants::GTicksPerQuarterNote;
 			}
 			else
 			{
@@ -378,7 +378,7 @@ void FFusionVoice::AttackWithTargetNote(uint8 InMidiNoteNumber, float InGain, in
 	}
 	else
 	{
-		StartBeat = (float)InEventTick / MidiConstants::kTicksPerQuarterNote;
+		StartBeat = (float)InEventTick / Harmonix::Midi::Constants::GTicksPerQuarterNote;
 	}
 
 	// multiply by 100.0?
diff --git a/Engine/Plugins/Runtime/Harmonix/Source/HarmonixDsp/Private/HarmonixDsp/Instruments/VirtualInstrument.cpp b/Engine/Plugins/Runtime/Harmonix/Source/HarmonixDsp/Private/HarmonixDsp/Instruments/VirtualInstrument.cpp
--- a/Engine/Plugins/Runtime/Harmonix/Source/HarmonixDsp/Private/HarmonixDsp/Instruments/VirtualInstrument.cpp
+++ b/Engine/Plugins/Runtime/Harmonix/Source/HarmonixDsp/Private/HarmonixDsp/Instruments/VirtualInstrument.cpp
@@ -18,13 +18,13 @@ void FVirtualInstrument::ResetMidiState()
 	ResetMidiStateImpl();
 }
 
-void FVirtualInstrument::Set7BitController(MidiConstants::EControllerID InController, int8 InByteValue, int8 InMidiChannel)
+void FVirtualInstrument::Set7BitController(Harmonix::Midi::Constants::EControllerID InController, int8 InByteValue, int8 InMidiChannel)
 {
 	LastCcVals[InMidiChannel][(uint8)InController] = InByteValue;
 	Set7BitControllerImpl(InController, InByteValue, InMidiChannel);
 }
 
-void FVirtualInstrument::Set14BitController(MidiConstants::EControllerID InController, int16 InByteValue, int8 InMidiChannel)
+void FVirtualInstrument::Set14BitController(Harmonix::Midi::Constants::EControllerID InController, int16 InByteValue, int8 InMidiChannel)
 {
 	int32 LsbIdx;
 	int32 MsbIdx;
@@ -32,37 +32,38 @@ void FVirtualInstrument::Set14BitController(MidiConstants::EControllerID InContr
 	{
 		LastCcVals[InMidiChannel][MsbIdx] = (InByteValue >> 7) & 0x7F;
 		LastCcVals[InMidiChannel][LsbIdx] = InByteValue & 0x7F;
-		Set14BitControllerImpl((MidiConstants::EControllerID)MsbIdx, InByteValue, InMidiChannel);
+		Set14BitControllerImpl((Harmonix::Midi::Constants::EControllerID)MsbIdx, InByteValue, InMidiChannel);
 	}
 }
 
 void FVirtualInstrument::HandleMidiMessage(FMidiVoiceId InVoiceId, int8 InStatus, int8 InData1, int8 InData2, int32 InEventTick, int32 InCurrentTick, float InMsOffset)
 {
+	using namespace Harmonix::Midi::Constants;
 	int8 InChannel = InStatus & 0xF;
 	switch (InStatus & 0xF0)
 	{
-	case MidiConstants::kNoteOff:
+	case GNoteOff:
 		NoteOff(InVoiceId, InData1, InChannel);
 		break;
-	case MidiConstants::kNoteOn:
+	case GNoteOn:
 		NoteOn(InVoiceId, InData1, InData2, InChannel, InEventTick, InCurrentTick, InMsOffset);
 		break;
-	case MidiConstants::kPolyPres:
+	case GPolyPres:
 		PolyPressure(InVoiceId, InData1, InData2, InChannel);
 		break;
-	case MidiConstants::kChanPres:
+	case GChanPres:
 		ChannelPressure(InData1, InData2, InChannel);
 		break;
-	case MidiConstants::kControl:
-		SetHighOrLowControllerByte((MidiConstants::EControllerID)InData1, InData2, InChannel);
+	case GControl:
+		SetHighOrLowControllerByte((EControllerID)InData1, InData2, InChannel);
 		break;
-	case MidiConstants::kPitch:
+	case GPitch:
 		SetPitchBend(FMidiMsg::GetPitchBendFromData(InData1, InData2), InChannel);
 		break;
 	}
 }
 
-bool FVirtualInstrument::IsHighResController(MidiConstants::EControllerID InControllerId, bool& bIsHighResLowByte)
+bool FVirtualInstrument::IsHighResController(Harmonix::Midi::Constants::EControllerID InControllerId, bool& bIsHighResLowByte)
 {
 	uint8 Id = (uint8)InControllerId;
 	if (Id < 32)
@@ -88,7 +89,7 @@ bool FVirtualInstrument::IsHighResController(MidiConstants::EControllerID InCont
 	return false;
 }
 
-bool FVirtualInstrument::GetMsbLsbIndexes(MidiConstants::EControllerID InControllerId, int& InMsb, int& InLsb)
+bool FVirtualInstrument::GetMsbLsbIndexes(Harmonix::Midi::Constants::EControllerID InControllerId, int& InMsb, int& InLsb)
 {
 	uint8 Id = (uint8)InControllerId;
 	if (Id < 32)
@@ -115,16 +116,17 @@ bool FVirtualInstrument::GetMsbLsbIndexes(MidiConstants::EControllerID InControl
 }
 
 // convert a value in [0,127] to the range expected by SetController
-float FVirtualInstrument::ConvertCCValue(MidiConstants::EControllerID InController, uint8 InValue)
+float FVirtualInstrument::ConvertCCValue(Harmonix::Midi::Constants::EControllerID InController, uint8 InValue)
 {
+	using namespace Harmonix::Midi::Constants;
 	check(InValue <= 127);
 	switch (InController)
 	{
-	case MidiConstants::EControllerID::PanRight:
-	case MidiConstants::EControllerID::CoarsePitchBend:
+	case EControllerID::PanRight:
+	case EControllerID::CoarsePitchBend:
 		// map to [-1, 1], making sure 64 maps exactly to 0
 		return FMath::Max(-1.0f, ((float)InValue - 64) / 63);
-	case MidiConstants::EControllerID::PortamentoSwitch:
+	case EControllerID::PortamentoSwitch:
 		// 63 and below are "off", 64 and above are "on"
 		return InValue >= 64 ? 1.0f : 0.0f;
 	default:
@@ -133,7 +135,7 @@ float FVirtualInstrument::ConvertCCValue(MidiConstants::EControllerID InControll
 	}
 }
 
-void FVirtualInstrument::SetController(MidiConstants::EControllerID InController, float InValue, int8 InMidiChannel /*= 0*/)
+void FVirtualInstrument::SetController(Harmonix::Midi::Constants::EControllerID InController, float InValue, int8 InMidiChannel /*= 0*/)
 {
 	int32 MsbIndex;
 	int32 LsbIndex;
@@ -147,7 +149,7 @@ void FVirtualInstrument::SetController(MidiConstants::EControllerID InController
 	}
 }
 
-void FVirtualInstrument::SetHighOrLowControllerByte(MidiConstants::EControllerID InController, int8 InValue, int8 InMidiChannel)
+void FVirtualInstrument::SetHighOrLowControllerByte(Harmonix::Midi::Constants::EControllerID InController, int8 InValue, int8 InMidiChannel)
 {
 	LastCcVals[InMidiChannel][(uint8)InController] = InValue;
 	int32 HighByteIdx;
@@ -160,7 +162,7 @@ void FVirtualInstrument::SetHighOrLowControllerByte(MidiConstants::EControllerID
 		}
 		else if (LastCcVals[InMidiChannel][HighByteIdx] != -1)
 		{
-			Set14BitController((MidiConstants::EControllerID)HighByteIdx, (LastCcVals[InMidiChannel][HighByteIdx] << 7 | LastCcVals[InMidiChannel][LowByteIdx]), InMidiChannel);
+			Set14BitController((Harmonix::Midi::Constants::EControllerID)HighByteIdx, (LastCcVals[InMidiChannel][HighByteIdx] << 7 | LastCcVals[InMidiChannel][LowByteIdx]), InMidiChannel);
 		}
 	}
 	else
diff --git a/Engine/Plugins/Runtime/Harmonix/Source/HarmonixDsp/Public/HarmonixDsp/FusionSampler/FusionSampler.h b/Engine/Plugins/Runtime/Harmonix/Source/HarmonixDsp/Public/HarmonixDsp/FusionSampler/FusionSampler.h
--- a/Engine/Plugins/Runtime/Harmonix/Source/HarmonixDsp/Public/HarmonixDsp/FusionSampler/FusionSampler.h
+++ b/Engine/Plugins/Runtime/Harmonix/Source/HarmonixDsp/Public/HarmonixDsp/FusionSampler/FusionSampler.h
@@ -58,7 +58,7 @@ public:
 	virtual void  SetExtraPitchBend(float Semitones, int8 InMidiChannel = 0) override;
 	virtual float GetPitchBend(int8 InMidiChannel = 0) const override;
 	virtual void  SetPitchBend(float Value, int8 InMidiChannel = 0) override;
-	virtual void  GetController(MidiConstants::EControllerID InController, int8& OutMsb, int8& OutLsb, int8 InMidiChannel = 0) const override;
+	virtual void  GetController(Harmonix::Midi::Constants::EControllerID InController, int8& OutMsb, int8& OutLsb, int8 InMidiChannel = 0) const override;
 
 	// Mix Volume (dB) and associated linear gain (0,1]
 	virtual void  SetMidiChannelVolume(float InVolume, float InSeconds = 0.0f, int8 InMidiChannel = 0) override;
@@ -139,13 +139,13 @@ protected:
 	// These are overrides from VirtualInstrument
 	virtual void  ResetInstrumentStateImpl() override;
 	virtual void  ResetMidiStateImpl() override;
-	virtual void  Set7BitControllerImpl(MidiConstants::EControllerID InController, int8 Value, int8 InMidiChannel = 0) override;
-	virtual void  Set14BitControllerImpl(MidiConstants::EControllerID InController, int16 Value, int8 InMidiChannel = 0) override;
+	virtual void  Set7BitControllerImpl(Harmonix::Midi::Constants::EControllerID InController, int8 Value, int8 InMidiChannel = 0) override;
+	virtual void  Set14BitControllerImpl(Harmonix::Midi::Constants::EControllerID InController, int16 Value, int8 InMidiChannel = 0) override;
 	
 	// NOTE: This next function's InValue is NOT some standard midi value. By the time it is called it is assumed 
 	// the units and range of InValue are correct for the given controller id. So, for example, if InController 
 	// is MidiConstants::LFO0Frequency than InValue is treated as Hz!
-	void SetController(MidiConstants::EControllerID InController, float InValue);
+	void SetController(Harmonix::Midi::Constants::EControllerID InController, float InValue);
 
 	void SetVoicePool(FSharedFusionVoicePoolPtr InPool, bool NoCallbacks);
 	FSharedFusionVoicePoolPtr GetVoicePool() const { return VoicePool; }
@@ -355,13 +355,13 @@ private:
 	};
 
 	TArray<FPendingNoteAction> PendingNoteActions;
-	FMIDINoteStatus NoteStatus[MidiConstants::kMaxNumNotes];
+	FMIDINoteStatus NoteStatus[Harmonix::Midi::Constants::GMaxNumNotes];
 
 	struct FFusionPatchData* FusionPatchData = nullptr;
 
-	int8 LastStartLayerSelect[MidiConstants::kMaxNumNotes];
-	int8 LastStopLayerSelect[MidiConstants::kMaxNumNotes];
-	int8 LastVelocity[MidiConstants::kMaxNumNotes];
+	int8 LastStartLayerSelect[Harmonix::Midi::Constants::GMaxNumNotes];
+	int8 LastStopLayerSelect[Harmonix::Midi::Constants::GMaxNumNotes];
+	int8 LastVelocity[Harmonix::Midi::Constants::GMaxNumNotes];
 
 	// "scratch pad" used during process to get each voice's output audio
 	static const int32 kScratchBufferFrames = 2048;
@@ -396,7 +396,7 @@ private:
 	float SubstreamGain[kMaxSubstreams];
 	void SetSubstreamMidiGain(int32 InIndex, uint8 InMidiGain);
 
-	int32 TicksPerQuarterNote = MidiConstants::kTicksPerQuarterNoteInt;
+	int32 TicksPerQuarterNote = Harmonix::Midi::Constants::GTicksPerQuarterNoteInt;
 
 	const FGainTable* GainTable = nullptr;
 };
diff --git a/Engine/Plugins/Runtime/Harmonix/Source/HarmonixDsp/Public/HarmonixDsp/Instruments/VirtualInstrument.h b/Engine/Plugins/Runtime/Harmonix/Source/HarmonixDsp/Public/HarmonixDsp/Instruments/VirtualInstrument.h
--- a/Engine/Plugins/Runtime/Harmonix/Source/HarmonixDsp/Public/HarmonixDsp/Instruments/VirtualInstrument.h
+++ b/Engine/Plugins/Runtime/Harmonix/Source/HarmonixDsp/Public/HarmonixDsp/Instruments/VirtualInstrument.h
@@ -64,14 +64,14 @@ public:
 	virtual void SetRawPitchMultiplier(float RawPitch) = 0;
 	virtual float GetRawPitchMultiplier() const = 0;
 
-	void Set7BitController(MidiConstants::EControllerID InController, int8 InByteValue, int8 InMidiChannel = 0);
-	void Set14BitController(MidiConstants::EControllerID InController, int16 InByteValue, int8 InMidiChannel = 0);
+	void Set7BitController(Harmonix::Midi::Constants::EControllerID InController, int8 InByteValue, int8 InMidiChannel = 0);
+	void Set14BitController(Harmonix::Midi::Constants::EControllerID InController, int16 InByteValue, int8 InMidiChannel = 0);
 
-	virtual void GetController(MidiConstants::EControllerID InController, int8& InMsb, int8& InLsb, int8 InMidiChannel = 0) const = 0;
-	void SetController(MidiConstants::EControllerID InController, float InValue, int8 InMidiChannel = 0);
-	void SetHighOrLowControllerByte(MidiConstants::EControllerID InController, int8 InValue, int8 InMidiChannel = 0);
+	virtual void GetController(Harmonix::Midi::Constants::EControllerID InController, int8& InMsb, int8& InLsb, int8 InMidiChannel = 0) const = 0;
+	void SetController(Harmonix::Midi::Constants::EControllerID InController, float InValue, int8 InMidiChannel = 0);
+	void SetHighOrLowControllerByte(Harmonix::Midi::Constants::EControllerID InController, int8 InValue, int8 InMidiChannel = 0);
 
-	virtual float GetController(MidiConstants::EControllerID InController, int8 InMidiChannel = 0)
+	virtual float GetController(Harmonix::Midi::Constants::EControllerID InController, int8 InMidiChannel = 0)
 	{
 		int8 msb;
 		int8 lsb;
@@ -129,17 +129,17 @@ public:
 	virtual void  SetMidiChannelMute(bool InMute, int8 InMidiChannel = 0) = 0;
 	virtual bool  GetMidiChannelMute(int8 InMidiChannel = 0) const = 0;
 
-	static bool IsHighResController(MidiConstants::EControllerID ControllerId, bool& bIsHighResLowByte);
-	static bool GetMsbLsbIndexes(MidiConstants::EControllerID ControllerId, int& InMsb, int& InLsb);
+	static bool IsHighResController(Harmonix::Midi::Constants::EControllerID ControllerId, bool& bIsHighResLowByte);
+	static bool GetMsbLsbIndexes(Harmonix::Midi::Constants::EControllerID ControllerId, int& InMsb, int& InLsb);
 
 protected:
-	virtual void Set7BitControllerImpl(MidiConstants::EControllerID InController, int8 InValue, int8 InMidiChannel = 0) = 0;
-	virtual void Set14BitControllerImpl(MidiConstants::EControllerID InController, int16 InValue, int8 InMidiChannel = 0) = 0;
+	virtual void Set7BitControllerImpl(Harmonix::Midi::Constants::EControllerID InController, int8 InValue, int8 InMidiChannel = 0) = 0;
+	virtual void Set14BitControllerImpl(Harmonix::Midi::Constants::EControllerID InController, int16 InValue, int8 InMidiChannel = 0) = 0;
 	virtual void ResetInstrumentStateImpl() = 0; // heavy weight! should completely resets the instrument.
 	virtual void ResetMidiStateImpl() = 0;       // less heavy weight. should do all notes off, reset pitch bend, etc.
 
 	// convert a value in [0,127] to the range expected by SetController
-	static float ConvertCCValue(MidiConstants::EControllerID InController, uint8 InValue);
+	static float ConvertCCValue(Harmonix::Midi::Constants::EControllerID InController, uint8 InValue);
 
 	FName Name;
 	FHarmonixMeterData TimingData;
diff --git a/Engine/Plugins/Runtime/Harmonix/Source/HarmonixDspEditor/Private/FusionPatchImportOptions.cpp b/Engine/Plugins/Runtime/Harmonix/Source/HarmonixDspEditor/Private/FusionPatchImportOptions.cpp
--- a/Engine/Plugins/Runtime/Harmonix/Source/HarmonixDspEditor/Private/FusionPatchImportOptions.cpp
+++ b/Engine/Plugins/Runtime/Harmonix/Source/HarmonixDspEditor/Private/FusionPatchImportOptions.cpp
@@ -489,7 +489,7 @@ bool FKeyzoneNoteParser::ParseNoteName(const FString& Input, FParseResult& Outpu
 	FRegexMatcher Matcher(Pattern, Input);
 	if (Matcher.FindNext())
 	{
-		Output.RootNote = MidiConstants::GetNoteNumberFromNoteName(TCHAR_TO_ANSI(*Matcher.GetCaptureGroup(1)));
+		Output.RootNote = Harmonix::Midi::Constants::GetNoteNumberFromNoteName(TCHAR_TO_ANSI(*Matcher.GetCaptureGroup(1)));
 		return true;
 	}
 	return false;
diff --git a/Engine/Plugins/Runtime/Harmonix/Source/HarmonixMetasound/Private/HarmonixMetasound/Components/MusicClockComponent.cpp b/Engine/Plugins/Runtime/Harmonix/Source/HarmonixMetasound/Private/HarmonixMetasound/Components/MusicClockComponent.cpp
--- a/Engine/Plugins/Runtime/Harmonix/Source/HarmonixMetasound/Private/HarmonixMetasound/Components/MusicClockComponent.cpp
+++ b/Engine/Plugins/Runtime/Harmonix/Source/HarmonixMetasound/Private/HarmonixMetasound/Components/MusicClockComponent.cpp
@@ -616,8 +616,8 @@ void UMusicClockComponent::BroadcastSongPosChanges()
 void UMusicClockComponent::MakeDefaultSongMap()
 {
 	DefaultMaps.EmptyAllMaps();
-	DefaultMaps.Init(MidiConstants::kTicksPerQuarterNoteInt);
-	DefaultMaps.GetTempoMap().AddTempoInfoPoint(MidiConstants::BPMToMidiTempo(Tempo), 0);
+	DefaultMaps.Init(Harmonix::Midi::Constants::GTicksPerQuarterNoteInt);
+	DefaultMaps.GetTempoMap().AddTempoInfoPoint(Harmonix::Midi::Constants::BPMToMidiTempo(Tempo), 0);
 	DefaultMaps.GetBarMap().AddTimeSignatureAtBarIncludingCountIn(0, TimeSignatureNum, TimeSignatureDenom);
 }
 
diff --git a/Engine/Plugins/Runtime/Harmonix/Source/HarmonixMetasound/Private/HarmonixMetasound/DataTypes/MidiClock.cpp b/Engine/Plugins/Runtime/Harmonix/Source/HarmonixMetasound/Private/HarmonixMetasound/DataTypes/MidiClock.cpp
--- a/Engine/Plugins/Runtime/Harmonix/Source/HarmonixMetasound/Private/HarmonixMetasound/DataTypes/MidiClock.cpp
+++ b/Engine/Plugins/Runtime/Harmonix/Source/HarmonixMetasound/Private/HarmonixMetasound/DataTypes/MidiClock.cpp
@@ -97,26 +97,29 @@ int32 SubdivisionToMidiTicks(const EMidiClockSubdivisionQuantization Division, c
 	{
 		return 0;
 	}
+
+	using namespace Harmonix::Midi::Constants;
+	
 	switch (Division)
 	{
 	case EMidiClockSubdivisionQuantization::Bar: 					return SongMap.GetBarMap().GetTicksInBarAfterPoint(BarMapPointIndex);
-	case EMidiClockSubdivisionQuantization::Beat:					return (MidiConstants::kTicksPerQuarterNoteInt * 4) / TimeSignature->Denominator;
-	case EMidiClockSubdivisionQuantization::ThirtySecondNote:		return MidiConstants::kTicksPerQuarterNoteInt / 8;
-	case EMidiClockSubdivisionQuantization::SixteenthNote:			return MidiConstants::kTicksPerQuarterNoteInt / 4;
-	case EMidiClockSubdivisionQuantization::EighthNote:				return MidiConstants::kTicksPerQuarterNoteInt / 2;
-	case EMidiClockSubdivisionQuantization::QuarterNote:			return MidiConstants::kTicksPerQuarterNoteInt;
-	case EMidiClockSubdivisionQuantization::HalfNote:				return MidiConstants::kTicksPerQuarterNoteInt * 2;
-	case EMidiClockSubdivisionQuantization::WholeNote:				return MidiConstants::kTicksPerQuarterNoteInt * 4;
-	case EMidiClockSubdivisionQuantization::DottedSixteenthNote:	return (MidiConstants::kTicksPerQuarterNoteInt / 4) + (MidiConstants::kTicksPerQuarterNoteInt / 8);
-	case EMidiClockSubdivisionQuantization::DottedEighthNote:		return (MidiConstants::kTicksPerQuarterNoteInt / 2) + (MidiConstants::kTicksPerQuarterNoteInt / 4);
-	case EMidiClockSubdivisionQuantization::DottedQuarterNote:		return (MidiConstants::kTicksPerQuarterNoteInt)     + (MidiConstants::kTicksPerQuarterNoteInt / 2);
-	case EMidiClockSubdivisionQuantization::DottedHalfNote:			return (MidiConstants::kTicksPerQuarterNoteInt * 2) + (MidiConstants::kTicksPerQuarterNoteInt);
-	case EMidiClockSubdivisionQuantization::DottedWholeNote:		return (MidiConstants::kTicksPerQuarterNoteInt * 4) + (MidiConstants::kTicksPerQuarterNoteInt * 2);
-	case EMidiClockSubdivisionQuantization::SixteenthNoteTriplet:   return (MidiConstants::kTicksPerQuarterNoteInt / 2) / 3;
-	case EMidiClockSubdivisionQuantization::EighthNoteTriplet:		return MidiConstants::kTicksPerQuarterNoteInt / 3;
-	case EMidiClockSubdivisionQuantization::QuarterNoteTriplet:		return (MidiConstants::kTicksPerQuarterNoteInt * 2) / 3;
-	case EMidiClockSubdivisionQuantization::HalfNoteTriplet:        return (MidiConstants::kTicksPerQuarterNoteInt * 4) / 3;
-	default:	/* Beat */											return (MidiConstants::kTicksPerQuarterNoteInt * 4) / TimeSignature->Denominator;
+	case EMidiClockSubdivisionQuantization::Beat:					return (GTicksPerQuarterNoteInt * 4) / TimeSignature->Denominator;
+	case EMidiClockSubdivisionQuantization::ThirtySecondNote:		return GTicksPerQuarterNoteInt / 8;
+	case EMidiClockSubdivisionQuantization::SixteenthNote:			return GTicksPerQuarterNoteInt / 4;
+	case EMidiClockSubdivisionQuantization::EighthNote:				return GTicksPerQuarterNoteInt / 2;
+	case EMidiClockSubdivisionQuantization::QuarterNote:			return GTicksPerQuarterNoteInt;
+	case EMidiClockSubdivisionQuantization::HalfNote:				return GTicksPerQuarterNoteInt * 2;
+	case EMidiClockSubdivisionQuantization::WholeNote:				return GTicksPerQuarterNoteInt * 4;
+	case EMidiClockSubdivisionQuantization::DottedSixteenthNote:	return (GTicksPerQuarterNoteInt / 4) + (GTicksPerQuarterNoteInt / 8);
+	case EMidiClockSubdivisionQuantization::DottedEighthNote:		return (GTicksPerQuarterNoteInt / 2) + (GTicksPerQuarterNoteInt / 4);
+	case EMidiClockSubdivisionQuantization::DottedQuarterNote:		return (GTicksPerQuarterNoteInt)     + (GTicksPerQuarterNoteInt / 2);
+	case EMidiClockSubdivisionQuantization::DottedHalfNote:			return (GTicksPerQuarterNoteInt * 2) + (GTicksPerQuarterNoteInt);
+	case EMidiClockSubdivisionQuantization::DottedWholeNote:		return (GTicksPerQuarterNoteInt * 4) + (GTicksPerQuarterNoteInt * 2);
+	case EMidiClockSubdivisionQuantization::SixteenthNoteTriplet:   return (GTicksPerQuarterNoteInt / 2) / 3;
+	case EMidiClockSubdivisionQuantization::EighthNoteTriplet:		return GTicksPerQuarterNoteInt / 3;
+	case EMidiClockSubdivisionQuantization::QuarterNoteTriplet:		return (GTicksPerQuarterNoteInt * 2) / 3;
+	case EMidiClockSubdivisionQuantization::HalfNoteTriplet:        return (GTicksPerQuarterNoteInt * 4) / 3;
+	default:	/* Beat */											return (GTicksPerQuarterNoteInt * 4) / TimeSignature->Denominator;
 	}
 }
 
@@ -240,7 +243,7 @@ namespace HarmonixMetasound
 
 		check(BlockFrameIndex >= MyMidiClock->TempoChangesInBlock.Last().BlockSampleFrameIndex);
 		MyMidiClock->HasTempoChangeInBlock = true;
-		float Bpm = MidiConstants::MidiTempoToBPM(Tempo);
+		float Bpm = Harmonix::Midi::Constants::MidiTempoToBPM(Tempo);
 		if (MyMidiClock->TempoChangesInBlock.Last().BlockSampleFrameIndex == BlockFrameIndex)
 		{
 			MyMidiClock->TempoChangesInBlock.Last().Tempo = Bpm;
@@ -570,7 +573,7 @@ namespace HarmonixMetasound
 
 		// add tempo info
 		float TempoBPM = FMath::Max(1.0f, InTempoBPM);
-		int32 MidiTempo = MidiConstants::BPMToMidiTempo(TempoBPM);
+		int32 MidiTempo = Harmonix::Midi::Constants::BPMToMidiTempo(TempoBPM);
 		Track.AddEvent(FMidiEvent(0, FMidiMsg(MidiTempo)));
 		TempoMap.AddTempoInfoPoint(MidiTempo, 0);
 
diff --git a/Engine/Plugins/Runtime/Harmonix/Source/HarmonixMetasound/Private/HarmonixMetasound/Nodes/FusionSamplerNode.cpp b/Engine/Plugins/Runtime/Harmonix/Source/HarmonixMetasound/Private/HarmonixMetasound/Nodes/FusionSamplerNode.cpp
--- a/Engine/Plugins/Runtime/Harmonix/Source/HarmonixMetasound/Private/HarmonixMetasound/Nodes/FusionSamplerNode.cpp
+++ b/Engine/Plugins/Runtime/Harmonix/Source/HarmonixMetasound/Private/HarmonixMetasound/Nodes/FusionSamplerNode.cpp
@@ -229,37 +229,40 @@ namespace HarmonixMetasound
 		void UpdatePatchOverrides()
 		{
 			static_assert(FFusionPatchSettings::kNumLfos == 2);
+
+			using namespace Harmonix::Midi::Constants;
+			
 			if (*Lfo0RateInPin > 0.0f)
 			{
-				SetController(MidiConstants::LFO0Frequency, *Lfo0RateInPin);
+				SetController(EControllerID::LFO0Frequency, *Lfo0RateInPin);
 			}
 			else if (FusionPatchDataPtr)
 			{
-				SetController(MidiConstants::LFO0Frequency, FusionPatchDataPtr->GetSettings().Lfo[0].Freq);
+				SetController(EControllerID::LFO0Frequency, FusionPatchDataPtr->GetSettings().Lfo[0].Freq);
 			}
 			if (*Lfo0DepthInPin >= 0.0f)
 			{
-				SetController(MidiConstants::LFO0Depth, *Lfo0DepthInPin);
+				SetController(EControllerID::LFO0Depth, *Lfo0DepthInPin);
 			}
 			else if (FusionPatchDataPtr)
 			{
-				SetController(MidiConstants::LFO0Depth, FusionPatchDataPtr->GetSettings().Lfo[0].Depth);
+				SetController(EControllerID::LFO0Depth, FusionPatchDataPtr->GetSettings().Lfo[0].Depth);
 			}
 			if (*Lfo1RateInPin > 0.0f)
 			{
-				SetController(MidiConstants::LFO1Frequency, *Lfo1RateInPin);
+				SetController(EControllerID::LFO1Frequency, *Lfo1RateInPin);
 			}
 			else if (FusionPatchDataPtr)
 			{
-				SetController(MidiConstants::LFO1Frequency, FusionPatchDataPtr->GetSettings().Lfo[1].Freq);
+				SetController(EControllerID::LFO1Frequency, FusionPatchDataPtr->GetSettings().Lfo[1].Freq);
 			}
 			if (*Lfo1DepthInPin >= 0.0f)
 			{
-				SetController(MidiConstants::LFO1Depth, *Lfo1DepthInPin);
+				SetController(EControllerID::LFO1Depth, *Lfo1DepthInPin);
 			}
 			else if (FusionPatchDataPtr)
 			{
-				SetController(MidiConstants::LFO1Depth, FusionPatchDataPtr->GetSettings().Lfo[1].Depth);
+				SetController(EControllerID::LFO1Depth, FusionPatchDataPtr->GetSettings().Lfo[1].Depth);
 			}
 
 			float FineTuneCentsTotal = 0;
diff --git a/Engine/Plugins/Runtime/Harmonix/Source/HarmonixMetasound/Private/HarmonixMetasound/Nodes/MidiCCTriggerNode.cpp b/Engine/Plugins/Runtime/Harmonix/Source/HarmonixMetasound/Private/HarmonixMetasound/Nodes/MidiCCTriggerNode.cpp
--- a/Engine/Plugins/Runtime/Harmonix/Source/HarmonixMetasound/Private/HarmonixMetasound/Nodes/MidiCCTriggerNode.cpp
+++ b/Engine/Plugins/Runtime/Harmonix/Source/HarmonixMetasound/Private/HarmonixMetasound/Nodes/MidiCCTriggerNode.cpp
@@ -355,8 +355,9 @@ namespace HarmonixMetasound::Nodes::MidiCCTriggerNode
 					continue;
 				}
 
-				if (Event.TrackIndex != *Inputs.TrackNumber || Event.MidiMessage.GetStdChannel() + 1 != *Inputs.ChannelNumber
-					|| Event.MidiMessage.GetStdStatus() != MidiConstants::kControl)
+				if (Event.TrackIndex != *Inputs.TrackNumber
+					|| !Event.MidiMessage.IsControlChange()
+					|| Event.MidiMessage.GetStdChannel() + 1 != *Inputs.ChannelNumber)
 				{
 					continue;
 				}
diff --git a/Engine/Plugins/Runtime/Harmonix/Source/HarmonixMetasound/Private/HarmonixMetasound/Nodes/StepSequencePlayerNode.cpp b/Engine/Plugins/Runtime/Harmonix/Source/HarmonixMetasound/Private/HarmonixMetasound/Nodes/StepSequencePlayerNode.cpp
--- a/Engine/Plugins/Runtime/Harmonix/Source/HarmonixMetasound/Private/HarmonixMetasound/Nodes/StepSequencePlayerNode.cpp
+++ b/Engine/Plugins/Runtime/Harmonix/Source/HarmonixMetasound/Private/HarmonixMetasound/Nodes/StepSequencePlayerNode.cpp
@@ -525,7 +525,7 @@ namespace HarmonixMetasound
 		}
 
 		// Min float possible is 64th note triplets.
-		OutTicksPerCell = FMath::Max(StepSizeQuarterNotes, 0.0416666f) * MidiConstants::kTicksPerQuarterNote;
+		OutTicksPerCell = FMath::Max(StepSizeQuarterNotes, 0.0416666f) * Midi::Constants::GTicksPerQuarterNote;
 		OutTableTickLength = OutTicksPerCell * LengthColumns;
 	}
 
diff --git a/Engine/Plugins/Runtime/Harmonix/Source/HarmonixMetasound/Public/HarmonixMetasound/DataTypes/MidiStream.h b/Engine/Plugins/Runtime/Harmonix/Source/HarmonixMetasound/Public/HarmonixMetasound/DataTypes/MidiStream.h
--- a/Engine/Plugins/Runtime/Harmonix/Source/HarmonixMetasound/Public/HarmonixMetasound/DataTypes/MidiStream.h
+++ b/Engine/Plugins/Runtime/Harmonix/Source/HarmonixMetasound/Public/HarmonixMetasound/DataTypes/MidiStream.h
@@ -390,7 +390,7 @@ namespace HarmonixMetasound
 		FMidiFileProxyPtr MidiFileSourceOfEvents;
 		int32 NumFramesPerBlock = 0;
 		Metasound::FSampleRate SampleRate  = 0;
-		int32 TicksPerQuarterNote = MidiConstants::kTicksPerQuarterNoteInt;
+		int32 TicksPerQuarterNote = Harmonix::Midi::Constants::GTicksPerQuarterNoteInt;
 
 		FMidiTimestampTransportState CurrentTransportState;
 		TArray<FMidiTimestampTransportState> TransportChangesInBlock;
diff --git a/Engine/Plugins/Runtime/Harmonix/Source/HarmonixMetasoundTests/Private/Analysis/MidiClockVertexAnalyzerTests.cpp b/Engine/Plugins/Runtime/Harmonix/Source/HarmonixMetasoundTests/Private/Analysis/MidiClockVertexAnalyzerTests.cpp
--- a/Engine/Plugins/Runtime/Harmonix/Source/HarmonixMetasoundTests/Private/Analysis/MidiClockVertexAnalyzerTests.cpp
+++ b/Engine/Plugins/Runtime/Harmonix/Source/HarmonixMetasoundTests/Private/Analysis/MidiClockVertexAnalyzerTests.cpp
@@ -54,7 +54,7 @@ namespace HarmonixMetasoundTests::MidiClockVertexAnalyzer
 		MidiData->LastEventTick = std::numeric_limits<int32>::max();
 		MidiData->Tracks[0].AddEvent(FMidiEvent(0, FMidiMsg(static_cast<uint8>(TimeSigNumerator), static_cast<uint8>(TimeSigDenominator))));
 		BarMap.AddTimeSignatureAtBarIncludingCountIn(0, TimeSigNumerator, TimeSigNumerator);
-		const int32 MidiTempo = MidiConstants::BPMToMidiTempo(Tempo);
+		const int32 MidiTempo = Harmonix::Midi::Constants::BPMToMidiTempo(Tempo);
 		MidiData->Tracks[0].AddEvent(FMidiEvent(0, FMidiMsg(MidiTempo)));
 		TempoMap.AddTempoInfoPoint(MidiTempo, 0);
 		MidiData->Tracks[0].Sort();
diff --git a/Engine/Plugins/Runtime/Harmonix/Source/HarmonixMetasoundTests/Private/Analysis/MidiStreamVertexAnalyzerTests.cpp b/Engine/Plugins/Runtime/Harmonix/Source/HarmonixMetasoundTests/Private/Analysis/MidiStreamVertexAnalyzerTests.cpp
--- a/Engine/Plugins/Runtime/Harmonix/Source/HarmonixMetasoundTests/Private/Analysis/MidiStreamVertexAnalyzerTests.cpp
+++ b/Engine/Plugins/Runtime/Harmonix/Source/HarmonixMetasoundTests/Private/Analysis/MidiStreamVertexAnalyzerTests.cpp
@@ -56,7 +56,7 @@ namespace HarmonixMetasoundTests::MidiStreamVertexAnalyzer
 		MidiData->LastEventTick = std::numeric_limits<int32>::max();
 		MidiData->Tracks[0].AddEvent(FMidiEvent(0, FMidiMsg(static_cast<uint8>(TimeSigNum), static_cast<uint8>(TimeSigDenom))));
 		MidiData->SongMaps.GetBarMap().AddTimeSignatureAtBarIncludingCountIn(0, TimeSigNum, TimeSigDenom);
-		const int32 MidiTempo = MidiConstants::BPMToMidiTempo(Tempo);
+		const int32 MidiTempo = Harmonix::Midi::Constants::BPMToMidiTempo(Tempo);
 		MidiData->Tracks[0].AddEvent(FMidiEvent(0, FMidiMsg(MidiTempo)));
 		MidiData->SongMaps.GetTempoMap().AddTempoInfoPoint(MidiTempo, 0);
 		MidiData->Tracks[0].Sort();
@@ -109,7 +109,7 @@ namespace HarmonixMetasoundTests::MidiStreamVertexAnalyzer
 		
 		bool CallbackSuccess = false;
 
-		constexpr int32 EventIntervalTicks = MidiConstants::kTicksPerQuarterNoteInt;
+		constexpr int32 EventIntervalTicks = Harmonix::Midi::Constants::GTicksPerQuarterNoteInt;
 		int32 NextExpectedNoteNumber = 60;
 		int32 NextExpectedVelocity = 20;
 		int32 NextExpectedEventTicks = 0;
@@ -151,7 +151,7 @@ namespace HarmonixMetasoundTests::MidiStreamVertexAnalyzer
 		const int32 NumBlocksToTest = FMath::CeilToInt((static_cast<float>(NumEventsToSend) * SampleRate / (Tempo / 60)) / NumSamplesPerBlock);
 		int32 SampleRemainder = 0;
 		int32 SampleCount = 0;
-		float TicksPerBlock = MidiConstants::kTicksPerQuarterNote * NumSamplesPerBlock * (Tempo / 60) / SampleRate;
+		float TicksPerBlock = Harmonix::Midi::Constants::GTicksPerQuarterNote * NumSamplesPerBlock * (Tempo / 60) / SampleRate;
 		float CurrentBlockTick = 0;
 		int32 NextEventTick = 0;
 		int32 NextNoteNumber = NextExpectedNoteNumber;
diff --git a/Engine/Plugins/Runtime/Harmonix/Source/HarmonixMetasoundTests/Private/MidiOps/MidiTrackFilterTests.cpp b/Engine/Plugins/Runtime/Harmonix/Source/HarmonixMetasoundTests/Private/MidiOps/MidiTrackFilterTests.cpp
--- a/Engine/Plugins/Runtime/Harmonix/Source/HarmonixMetasoundTests/Private/MidiOps/MidiTrackFilterTests.cpp
+++ b/Engine/Plugins/Runtime/Harmonix/Source/HarmonixMetasoundTests/Private/MidiOps/MidiTrackFilterTests.cpp
@@ -25,7 +25,7 @@ namespace Harmonix::Midi::Ops::Tests
 		constexpr uint8 Data1 = 78;
 		constexpr uint8 Data2 = 87;
 		constexpr float Tempo = 135.7f;
-		const int32 MidiTempo = MidiConstants::BPMToMidiTempo(Tempo);
+		const int32 MidiTempo = Harmonix::Midi::Constants::BPMToMidiTempo(Tempo);
 		constexpr uint8 TimeSigNumerator = 7;
 		constexpr uint8 TimeSigDenominator = 8;
 
diff --git a/Engine/Plugins/Runtime/Harmonix/Source/HarmonixMetasoundTests/Private/MidiOps/StuckNoteGuardTests.cpp b/Engine/Plugins/Runtime/Harmonix/Source/HarmonixMetasoundTests/Private/MidiOps/StuckNoteGuardTests.cpp
--- a/Engine/Plugins/Runtime/Harmonix/Source/HarmonixMetasoundTests/Private/MidiOps/StuckNoteGuardTests.cpp
+++ b/Engine/Plugins/Runtime/Harmonix/Source/HarmonixMetasoundTests/Private/MidiOps/StuckNoteGuardTests.cpp
@@ -29,7 +29,7 @@ namespace Harmonix::Midi::Ops::Tests
 		constexpr uint8 Velocity = 54;
 
 		// Add some note ons on all the channels
-		for (uint8 ChannelIdx = 0; ChannelIdx < MidiConstants::kNumChannels; ++ChannelIdx)
+		for (uint8 ChannelIdx = 0; ChannelIdx < Constants::GNumChannels; ++ChannelIdx)
 		{
 			InputStream.AddMidiEvent({ static_cast<uint32>(0), FMidiMsg::CreateNoteOn(ChannelIdx, NoteNumber, Velocity)});
 		}
@@ -55,7 +55,7 @@ namespace Harmonix::Midi::Ops::Tests
 		// Add some note ons on all the channels
 		InputStream.PrepareBlock();
 		OutputStream.PrepareBlock();
-		for (uint8 ChannelIdx = 0; ChannelIdx < MidiConstants::kNumChannels; ++ChannelIdx)
+		for (uint8 ChannelIdx = 0; ChannelIdx < Constants::GNumChannels; ++ChannelIdx)
 		{
 			InputStream.AddMidiEvent({ static_cast<uint32>(0), FMidiMsg::CreateNoteOn(ChannelIdx, NoteNumber, Velocity)});
 		}
@@ -67,7 +67,7 @@ namespace Harmonix::Midi::Ops::Tests
 		// Now send note offs on all the channels, process, and expect no additional note offs
 		InputStream.PrepareBlock();
 		OutputStream.PrepareBlock();
-		for (uint8 ChannelIdx = 0; ChannelIdx < MidiConstants::kNumChannels; ++ChannelIdx)
+		for (uint8 ChannelIdx = 0; ChannelIdx < Constants::GNumChannels; ++ChannelIdx)
 		{
 			InputStream.AddMidiEvent({ static_cast<uint32>(0), FMidiMsg::CreateNoteOff(ChannelIdx, NoteNumber)});
 		}
diff --git a/Engine/Plugins/Runtime/Harmonix/Source/HarmonixMetasoundTests/Private/Nodes/DelayNodeTests.cpp b/Engine/Plugins/Runtime/Harmonix/Source/HarmonixMetasoundTests/Private/Nodes/DelayNodeTests.cpp
--- a/Engine/Plugins/Runtime/Harmonix/Source/HarmonixMetasoundTests/Private/Nodes/DelayNodeTests.cpp
+++ b/Engine/Plugins/Runtime/Harmonix/Source/HarmonixMetasoundTests/Private/Nodes/DelayNodeTests.cpp
@@ -251,7 +251,7 @@ namespace HarmonixMetasoundTests::DelayNode
 			MidiData->LastEventTick = std::numeric_limits<int32>::max();
 			MidiData->Tracks[0].AddEvent(FMidiEvent(0, FMidiMsg(static_cast<uint8>(TimeSigNum), static_cast<uint8>(TimeSigDenom))));
 			BarMap.AddTimeSignatureAtBarIncludingCountIn(0, TimeSigNum, TimeSigDenom);
-			const int32 MidiTempo = MidiConstants::BPMToMidiTempo(Tempo);
+			const int32 MidiTempo = Harmonix::Midi::Constants::BPMToMidiTempo(Tempo);
 			MidiData->Tracks[0].AddEvent(FMidiEvent(0, FMidiMsg(MidiTempo)));
 			TempoMap.AddTempoInfoPoint(MidiTempo, 0);
 			MidiData->Tracks[0].Sort();
diff --git a/Engine/Plugins/Runtime/Harmonix/Source/HarmonixMetasoundTests/Private/Nodes/MetronomeNodeTests.cpp b/Engine/Plugins/Runtime/Harmonix/Source/HarmonixMetasoundTests/Private/Nodes/MetronomeNodeTests.cpp
--- a/Engine/Plugins/Runtime/Harmonix/Source/HarmonixMetasoundTests/Private/Nodes/MetronomeNodeTests.cpp
+++ b/Engine/Plugins/Runtime/Harmonix/Source/HarmonixMetasoundTests/Private/Nodes/MetronomeNodeTests.cpp
@@ -112,11 +112,11 @@ namespace HarmonixMetasoundTests::MetronomeNode
 				return false;
 			}
 
-			float DefaultTicksPerSec = 120.0f * MidiConstants::kTicksPerQuarterNote / 60.0f;
+			float DefaultTicksPerSec = 120.0f * Harmonix::Midi::Constants::GTicksPerQuarterNote / 60.0f;
 			float DefaultTicksPerMs = DefaultTicksPerSec / 1000.0f;
 
 			// do some math to figure out how fast the clock should be advancing...
-			float TicksPerSec = Params.Tempo * MidiConstants::kTicksPerQuarterNote / 60.0f;
+			float TicksPerSec = Params.Tempo * Harmonix::Midi::Constants::GTicksPerQuarterNote / 60.0f;
 			float TicksPerMs = TicksPerSec / 1000.0f;
 			float SecsPerBlock = Params.NumSamplesPerBlock / Params.SampleRate;
 			float TicksPerBlock = TicksPerSec * SecsPerBlock;
diff --git a/Engine/Plugins/Runtime/Harmonix/Source/HarmonixMetasoundTests/Private/Nodes/MidiPlayerNodeTests.cpp b/Engine/Plugins/Runtime/Harmonix/Source/HarmonixMetasoundTests/Private/Nodes/MidiPlayerNodeTests.cpp
--- a/Engine/Plugins/Runtime/Harmonix/Source/HarmonixMetasoundTests/Private/Nodes/MidiPlayerNodeTests.cpp
+++ b/Engine/Plugins/Runtime/Harmonix/Source/HarmonixMetasoundTests/Private/Nodes/MidiPlayerNodeTests.cpp
@@ -174,12 +174,12 @@ namespace HarmonixMetasoundTests::MidiPlayerNode
 			TOptional<FMidiClockReadRef> MidiPlayerClockOut = Generator->GetOutputReadReference<FMidiClock>(NodeNames::MidiPlayer::ClockOut);
 
 
-			float DefaultTicksPerSec = 120.0f * MidiConstants::kTicksPerQuarterNote / 60.0f;
+			float DefaultTicksPerSec = 120.0f * Harmonix::Midi::Constants::GTicksPerQuarterNote / 60.0f;
 			float DefaultTicksPerMs = DefaultTicksPerSec / 1000.0f;
 
 			// do some math to figure out how fast the clock should be advancing...
 			float Tempo = Params.MidiPlayer.MidiFile ? Params.MidiPlayer.MidiFile->SongMaps.GetTempoAtTick(0) : 120.0f;
-			float TicksPerSec = Tempo * MidiConstants::kTicksPerQuarterNote / 60.0f;
+			float TicksPerSec = Tempo * Harmonix::Midi::Constants::GTicksPerQuarterNote / 60.0f;
 			float TicksPerMs = TicksPerSec / 1000.0f;
 			float SecsPerBlock = Params.NumSamplesPerBlock / Params.SampleRate;
 			float TicksPerBlock = TicksPerSec * SecsPerBlock;
@@ -263,7 +263,7 @@ namespace HarmonixMetasoundTests::MidiPlayerNode
 
 			// add tempo info
 			float TempoBpm = FMath::Max(1.0f, InTempoBpm);
-			int32 MidiTempo = MidiConstants::BPMToMidiTempo(TempoBpm);
+			int32 MidiTempo = Harmonix::Midi::Constants::BPMToMidiTempo(TempoBpm);
 			Track.AddEvent(FMidiEvent(0, FMidiMsg(MidiTempo)));
 			TempoMap.AddTempoInfoPoint(MidiTempo, 0);
 
diff --git a/Engine/Plugins/Runtime/Harmonix/Source/HarmonixMetasoundTests/Private/Nodes/MidiTextTriggerNodeTests.cpp b/Engine/Plugins/Runtime/Harmonix/Source/HarmonixMetasoundTests/Private/Nodes/MidiTextTriggerNodeTests.cpp
--- a/Engine/Plugins/Runtime/Harmonix/Source/HarmonixMetasoundTests/Private/Nodes/MidiTextTriggerNodeTests.cpp
+++ b/Engine/Plugins/Runtime/Harmonix/Source/HarmonixMetasoundTests/Private/Nodes/MidiTextTriggerNodeTests.cpp
@@ -21,16 +21,16 @@ namespace HarmonixMetasoundTests::MidiTextTriggerNode
 		// Make a midi file...
 		UMidiFile* TheMidi = NewObject<UMidiFile>();
 		// Set the initial tempo...
-		TheMidi->GetSongMaps()->GetTempoMap().AddTempoInfoPoint(MidiConstants::BPMToMidiTempo(120),0);
+		TheMidi->GetSongMaps()->GetTempoMap().AddTempoInfoPoint(Harmonix::Midi::Constants::BPMToMidiTempo(120),0);
 		// Make a track to hold some text events. This will be track 1 (track 0 is always the conductor track)...
 		FMidiTrack* TextTrack = TheMidi->AddTrack("TextEventTest");
 		// Make a text event...
 		uint16 TextIndex = TextTrack->AddText("ThisIsText");
-		FMidiEvent ATextEvent(0, FMidiMsg::CreateText(TextIndex, MidiConstants::kMeta_Text));
+		FMidiEvent ATextEvent(0, FMidiMsg::CreateText(TextIndex, Harmonix::Midi::Constants::GMeta_Text));
 		TextTrack->AddEvent(ATextEvent);
 		// Make another text event...
 		TextIndex = TextTrack->AddText("ThisIsAnotherText");
-		FMidiEvent ASecondTextEvent(TheMidi->GetSongMaps()->MsToTick(15.0f), FMidiMsg::CreateText(TextIndex, MidiConstants::kMeta_Text));
+		FMidiEvent ASecondTextEvent(TheMidi->GetSongMaps()->MsToTick(15.0f), FMidiMsg::CreateText(TextIndex, Harmonix::Midi::Constants::GMeta_Text));
 		TextTrack->AddEvent(ASecondTextEvent);
 		// Tell the midi file its tracks have been changed so it can recalculate song length data...
 		TheMidi->TracksChanged();
diff --git a/Engine/Plugins/Runtime/Harmonix/Source/HarmonixMetasoundTests/Private/Nodes/MorphingLfoNode.spec.cpp b/Engine/Plugins/Runtime/Harmonix/Source/HarmonixMetasoundTests/Private/Nodes/MorphingLfoNode.spec.cpp
--- a/Engine/Plugins/Runtime/Harmonix/Source/HarmonixMetasoundTests/Private/Nodes/MorphingLfoNode.spec.cpp
+++ b/Engine/Plugins/Runtime/Harmonix/Source/HarmonixMetasoundTests/Private/Nodes/MorphingLfoNode.spec.cpp
@@ -75,7 +75,7 @@ namespace HarmonixMetasoundTests::MorphingLFONode
 		MidiData->LastEventTick = std::numeric_limits<int32>::max();
 		MidiData->Tracks[0].AddEvent(FMidiEvent(0, FMidiMsg(static_cast<uint8>(TimeSigNumerator), static_cast<uint8>(TimeSigDenominator))));
 		BarMap.AddTimeSignatureAtBarIncludingCountIn(0, TimeSigNumerator, TimeSigNumerator);
-		const int32 MidiTempo = MidiConstants::BPMToMidiTempo(Tempo);
+		const int32 MidiTempo = Harmonix::Midi::Constants::BPMToMidiTempo(Tempo);
 		MidiData->Tracks[0].AddEvent(FMidiEvent(0, FMidiMsg(MidiTempo)));
 		TempoMap.AddTempoInfoPoint(MidiTempo, 0);
 		MidiData->Tracks[0].Sort();
diff --git a/Engine/Plugins/Runtime/Harmonix/Source/HarmonixMetasoundTests/Private/Subsystems/MidiClockUpdateSubsystemTests.cpp b/Engine/Plugins/Runtime/Harmonix/Source/HarmonixMetasoundTests/Private/Subsystems/MidiClockUpdateSubsystemTests.cpp
--- a/Engine/Plugins/Runtime/Harmonix/Source/HarmonixMetasoundTests/Private/Subsystems/MidiClockUpdateSubsystemTests.cpp
+++ b/Engine/Plugins/Runtime/Harmonix/Source/HarmonixMetasoundTests/Private/Subsystems/MidiClockUpdateSubsystemTests.cpp
@@ -25,7 +25,7 @@ namespace HarmonixMetasoundTests::MidiClockUpdateSubsystem
 			MidiData->Tracks[0].AddEvent(FMidiEvent(0, FMidiMsg(static_cast<uint8>(TimeSigNum), static_cast<uint8>(TimeSigDenom))));
 			FBarMap& BarMap = MidiData->SongMaps.GetBarMap();
 			BarMap.AddTimeSignatureAtBarIncludingCountIn(0, TimeSigNum, TimeSigDenom);
-			const int32 MidiTempo = MidiConstants::BPMToMidiTempo(Tempo);
+			const int32 MidiTempo = Harmonix::Midi::Constants::BPMToMidiTempo(Tempo);
 			MidiData->Tracks[0].AddEvent(FMidiEvent(0, FMidiMsg(MidiTempo)));
 			FTempoMap& TempoMap = MidiData->SongMaps.GetTempoMap();
 			TempoMap.AddTempoInfoPoint(MidiTempo, 0);
diff --git a/Engine/Plugins/Runtime/Harmonix/Source/HarmonixMidi/Private/HarmonixMidi/Blueprint/MidiFunctionLibrary.cpp b/Engine/Plugins/Runtime/Harmonix/Source/HarmonixMidi/Private/HarmonixMidi/Blueprint/MidiFunctionLibrary.cpp
--- a/Engine/Plugins/Runtime/Harmonix/Source/HarmonixMidi/Private/HarmonixMidi/Blueprint/MidiFunctionLibrary.cpp
+++ b/Engine/Plugins/Runtime/Harmonix/Source/HarmonixMidi/Private/HarmonixMidi/Blueprint/MidiFunctionLibrary.cpp
@@ -8,61 +8,61 @@
 
 FMidiNote UMidiNoteFunctionLibrary::GetMinMidiNote()
 {
-	return FMidiNote(MidiConstants::kMinNote);
+	return FMidiNote(Harmonix::Midi::Constants::GMinNote);
 }
 
 FMidiNote UMidiNoteFunctionLibrary::GetMaxMidiNote()
 {
-	return FMidiNote(MidiConstants::kMaxNote);
+	return FMidiNote(Harmonix::Midi::Constants::GMaxNote);
 }
 
 uint8 UMidiNoteFunctionLibrary::GetMinNoteNumber()
 {
-	return MidiConstants::kMinNote;
+	return Harmonix::Midi::Constants::GMinNote;
 }
 
 uint8 UMidiNoteFunctionLibrary::GetMaxNoteNumber()
 {
-	return MidiConstants::kMaxNote;
+	return Harmonix::Midi::Constants::GMaxNote;
 }
 
 int UMidiNoteFunctionLibrary::GetMaxNumNotes()
 {
-	return MidiConstants::kMaxNumNotes;
+	return Harmonix::Midi::Constants::GMaxNumNotes;
 }
 
 uint8 UMidiNoteFunctionLibrary::GetMinNoteVelocity()
 {
-	return MidiConstants::kMinVelocity;
+	return Harmonix::Midi::Constants::GMinVelocity;
 }
 
 uint8 UMidiNoteFunctionLibrary::GetMaxNoteVelocity()
 {
-	return MidiConstants::kMaxVelocity;
+	return Harmonix::Midi::Constants::GMaxVelocity;
 }
 
 float UMusicalTickFunctionLibrary::GetTicksPerQuarterNote()
 {
-	return MidiConstants::kTicksPerQuarterNote;
+	return Harmonix::Midi::Constants::GTicksPerQuarterNote;
 }
 
 int32 UMusicalTickFunctionLibrary::GetTicksPerQuarterNoteInt()
 {
-	return MidiConstants::kTicksPerQuarterNoteInt;
+	return Harmonix::Midi::Constants::GTicksPerQuarterNoteInt;
 }
 
 float UMusicalTickFunctionLibrary::GetQuarterNotesPerTick()
 {
-	return MidiConstants::kQuarterNotesPerTick;
+	return Harmonix::Midi::Constants::GQuarterNotesPerTick;
 }
 
 float UMusicalTickFunctionLibrary::TickToQuarterNote(float InTick)
 {
-	return InTick * MidiConstants::kQuarterNotesPerTick;
+	return InTick * Harmonix::Midi::Constants::GQuarterNotesPerTick;
 }
 
 float UMusicalTickFunctionLibrary::QuarterNoteToTick(float InQuarterNote)
 {
-	return InQuarterNote * MidiConstants::kTicksPerQuarterNote;
+	return InQuarterNote * Harmonix::Midi::Constants::GTicksPerQuarterNote;
 }
 
diff --git a/Engine/Plugins/Runtime/Harmonix/Source/HarmonixMidi/Private/HarmonixMidi/Blueprint/MidiNote.cpp b/Engine/Plugins/Runtime/Harmonix/Source/HarmonixMidi/Private/HarmonixMidi/Blueprint/MidiNote.cpp
--- a/Engine/Plugins/Runtime/Harmonix/Source/HarmonixMidi/Private/HarmonixMidi/Blueprint/MidiNote.cpp
+++ b/Engine/Plugins/Runtime/Harmonix/Source/HarmonixMidi/Private/HarmonixMidi/Blueprint/MidiNote.cpp
@@ -17,11 +17,11 @@ FName FMidiNote::GetDisplayName() const
 
 FString FMidiNote::ToString() const
 {
-	int32 Octave = MidiConstants::GetNoteOctaveFromNoteNumber(NoteNumber);
+	int32 Octave = Harmonix::Midi::Constants::GetNoteOctaveFromNoteNumber(NoteNumber);
 
-	MidiConstants::ENoteNameEnharmonicStyle Style = MidiConstants::ENoteNameEnharmonicStyle::Sharp;
+	Harmonix::Midi::Constants::ENoteNameEnharmonicStyle Style = Harmonix::Midi::Constants::ENoteNameEnharmonicStyle::Sharp;
 
-	return FString::Printf(TEXT("%s%d"), *FString(MidiConstants::GetNoteNameFromNoteNumber(NoteNumber, Style)), Octave);
+	return FString::Printf(TEXT("%s%d"), *FString(Harmonix::Midi::Constants::GetNoteNameFromNoteNumber(NoteNumber, Style)), Octave);
 }
 
 FMidiNote FMidiNote::FromString(const FString& NoteName)
@@ -32,7 +32,7 @@ FMidiNote FMidiNote::FromString(const FString& NoteName)
 
 uint8 FMidiNote::NoteNumberFromString(const FString& NoteName)
 {
-	return MidiConstants::GetNoteNumberFromNoteName(TCHAR_TO_ANSI(*NoteName));
+	return Harmonix::Midi::Constants::GetNoteNumberFromNoteName(TCHAR_TO_ANSI(*NoteName));
 }
 
 #ifdef WITH_EDITOR
@@ -58,9 +58,9 @@ uint8 FMidiNote::NoteNumberFromEditorString(const FString& EditorName)
 	if (EditorName.FindLastChar(' ', Index))
 	{
 		FString NoteName = EditorName.RightChop(Index + 1);
-		return MidiConstants::GetNoteNumberFromNoteName(TCHAR_TO_ANSI(*NoteName));
+		return Harmonix::Midi::Constants::GetNoteNumberFromNoteName(TCHAR_TO_ANSI(*NoteName));
 	}
-	return MidiConstants::GetNoteNumberFromNoteName(TCHAR_TO_ANSI(*EditorName));
+	return Harmonix::Midi::Constants::GetNoteNumberFromNoteName(TCHAR_TO_ANSI(*EditorName));
 }
 
 #endif
diff --git a/Engine/Plugins/Runtime/Harmonix/Source/HarmonixMidi/Private/HarmonixMidi/MidiConstants.cpp b/Engine/Plugins/Runtime/Harmonix/Source/HarmonixMidi/Private/HarmonixMidi/MidiConstants.cpp
--- a/Engine/Plugins/Runtime/Harmonix/Source/HarmonixMidi/Private/HarmonixMidi/MidiConstants.cpp
+++ b/Engine/Plugins/Runtime/Harmonix/Source/HarmonixMidi/Private/HarmonixMidi/MidiConstants.cpp
@@ -1,225 +1,231 @@
 // Copyright Epic Games, Inc. All Rights Reserved.
+
 #include "HarmonixMidi/MidiConstants.h"
 
 DEFINE_LOG_CATEGORY(LogMidi);
 
-FString MidiConstants::GetTextTypeName(uint8 TextType)
+namespace Harmonix::Midi::Constants
 {
-	switch (TextType)
+	FString GetTextTypeName(uint8 TextType)
 	{
-	case kMeta_Text: return TEXT("Text");
-	case kMeta_Copyright: return TEXT("Copyright");
-	case kMeta_TrackName: return TEXT("Track Name");
-	case kMeta_InstrumentName: return TEXT("Instrument Name");
-	case kMeta_Lyric: return TEXT("Lyric");
-	case kMeta_Marker: return TEXT("Marker");
-	case kMeta_CuePoint: return TEXT("Cue Point");
-	default:
-		checkNoEntry();
-		break;
+		switch (TextType)
+		{
+		case GMeta_Text: return TEXT("Text");
+		case GMeta_Copyright: return TEXT("Copyright");
+		case GMeta_TrackName: return TEXT("Track Name");
+		case GMeta_InstrumentName: return TEXT("Instrument Name");
+		case GMeta_Lyric: return TEXT("Lyric");
+		case GMeta_Marker: return TEXT("Marker");
+		case GMeta_CuePoint: return TEXT("Cue Point");
+		default:
+			checkNoEntry();
+			break;
+		}
+		return TEXT("");
 	}
-	return TEXT("");
-}
-
-float MidiConstants::RoundToStandardBeatPrecision(float InBeat, int TimeSignatureDenominator)
-{
-	// this is the precision to which the "rounded" elapsed beats get presented.
-	const double kQuarterNotePrecision = 0.01 * ((double)TimeSignatureDenominator / 4.0); // account for beat being a quarter note, eighth note, etc.
-	double x = (double)InBeat;
-	x += 0.5 * kQuarterNotePrecision;
-	x *= 1.0 / kQuarterNotePrecision;
-	x = FMath::Floor(x);
-	x *= kQuarterNotePrecision;
-	return x;
-}
 
-FString MidiConstants::GetControllerName(EControllerID ControllerId)
-{
-	switch (ControllerId)
+	float RoundToStandardBeatPrecision(float InBeat, int TimeSignatureDenominator)
 	{
-	case EControllerID::BankSelection:       return FString::Format(TEXT("({0}) BankSelection"), { ControllerId });
-	case EControllerID::ModWheel:			 return FString::Format(TEXT("({0}) ModWheel"), { ControllerId });
-	case EControllerID::Breath:				 return FString::Format(TEXT("({0}) Breath"), { ControllerId });
-	case EControllerID::PortamentoTime:		 return FString::Format(TEXT("({0}) PortamentoTime"), { ControllerId });
-	case EControllerID::DataCoarse:			 return FString::Format(TEXT("({0}) DataCoarse"), { ControllerId });
-	case EControllerID::Volume:				 return FString::Format(TEXT("({0}) Volume"), { ControllerId });
-	case EControllerID::Balance:			 return FString::Format(TEXT("({0}) Balance"), { ControllerId });
-	case EControllerID::PanRight:			 return FString::Format(TEXT("({0}) Pan (0 = left, 1 = right)"), { ControllerId });
-	case EControllerID::Expression:			 return FString::Format(TEXT("({0}) Expression"), { ControllerId });
-	case EControllerID::BitCrushWetMix:		 return FString::Format(TEXT("({0}) BitCrushWetMix"), { ControllerId });
-	case EControllerID::BitCrushLevel:		 return FString::Format(TEXT("({0}) BitCrushLevel"), { ControllerId });
-	case EControllerID::BitCrushSampleHold:	 return FString::Format(TEXT("({0}) BitCrushSampleHold"), { ControllerId });
-	case EControllerID::LFO0Frequency:		 return FString::Format(TEXT("({0}) LFO0Frequency"), { ControllerId });
-	case EControllerID::LFO1Frequency:		 return FString::Format(TEXT("({0}) LFO1Frequency"), { ControllerId });
-	case EControllerID::LFO0Depth:			 return FString::Format(TEXT("({0}) LFO0Depth"), { ControllerId });
-	case EControllerID::LFO1Depth:			 return FString::Format(TEXT("({0}) LFO1Depth"), { ControllerId });
-	case EControllerID::DelayTime:			 return FString::Format(TEXT("({0}) DelayTime"), { ControllerId });
-	case EControllerID::DelayDryGain:		 return FString::Format(TEXT("({0}) DelayDryGain"), { ControllerId });
-	case EControllerID::DelayWetGain:		 return FString::Format(TEXT("({0}) DelayWetGain"), { ControllerId });
-	case EControllerID::DelayFeedback:		 return FString::Format(TEXT("({0}) DelayFeedback"), { ControllerId });
-	case EControllerID::CoarsePitchBend:	 return FString::Format(TEXT("({0}) CoarsePitchBend"), { ControllerId });
-	case EControllerID::SampleStartTime:	 return FString::Format(TEXT("({0}) SampleStartTime"), { ControllerId });
-	case EControllerID::DataFine:			 return FString::Format(TEXT("({0}) DataFine"), { ControllerId });
-	case EControllerID::SubStreamVol1:		 return FString::Format(TEXT("({0}) SubStreamVol1"), { ControllerId });
-	case EControllerID::SubStreamVol2:		 return FString::Format(TEXT("({0}) SubStreamVol2"), { ControllerId });
-	case EControllerID::SubStreamVol3:		 return FString::Format(TEXT("({0}) SubStreamVol3"), { ControllerId });
-	case EControllerID::SubStreamVol4:		 return FString::Format(TEXT("({0}) SubStreamVol4"), { ControllerId });
-	case EControllerID::SubStreamVol5:		 return FString::Format(TEXT("({0}) SubStreamVol5"), { ControllerId });
-	case EControllerID::SubStreamVol6:		 return FString::Format(TEXT("({0}) SubStreamVol6"), { ControllerId });
-	case EControllerID::SubStreamVol7:		 return FString::Format(TEXT("({0}) SubStreamVol7"), { ControllerId });
-	case EControllerID::SubStreamVol8:		 return FString::Format(TEXT("({0}) SubStreamVol8"), { ControllerId });
-	case EControllerID::DelayEQEnabled:		 return FString::Format(TEXT("({0}) DelayEQEnabled"), { ControllerId });
-	case EControllerID::DelayEQType:		 return FString::Format(TEXT("({0}) DelayEQType"), { ControllerId });
-	case EControllerID::DelayEQFreq:		 return FString::Format(TEXT("({0}) DelayEQFreq"), { ControllerId });
-	case EControllerID::DelayEQQ:			 return FString::Format(TEXT("({0}) DelayEQQ"), { ControllerId });
-	case EControllerID::Hold:				 return FString::Format(TEXT("({0}) Hold"), { ControllerId });
-	case EControllerID::PortamentoSwitch:	 return FString::Format(TEXT("({0}) PortamentoSwitch"), { ControllerId });
-	case EControllerID::Sustenuto:			 return FString::Format(TEXT("({0}) Sustenuto"), { ControllerId });
-	case EControllerID::SoftPedal:			 return FString::Format(TEXT("({0}) SoftPedal"), { ControllerId });
-	case EControllerID::Legato:				 return FString::Format(TEXT("({0}) Legato"), { ControllerId });
-	case EControllerID::Hold2:				 return FString::Format(TEXT("({0}) Hold2"), { ControllerId });
-	case EControllerID::FilterQ:			 return FString::Format(TEXT("({0}) FilterQ"), { ControllerId });
-	case EControllerID::Release:			 return FString::Format(TEXT("({0}) Release"), { ControllerId });
-	case EControllerID::Attack:				 return FString::Format(TEXT("({0}) Attack"), { ControllerId });
-	case EControllerID::FilterFrequency:	 return FString::Format(TEXT("({0}) FilterFrequency"), { ControllerId });
-	case EControllerID::TimeStretchEnvelopeOrder:		 
-		return FString::Format(TEXT("({0}) TimeStretchEnvelopeOrder"), { ControllerId });
-	case EControllerID::DelayLFOBeatSync:	 return FString::Format(TEXT("({0}) DelayLFOBeatSync"), { ControllerId });
-	case EControllerID::DelayLFOEnabled:	 return FString::Format(TEXT("({0}) DelayLFOEnabled"), { ControllerId });
-	case EControllerID::DelayLFORate:		 return FString::Format(TEXT("({0}) DelayLFORate"), { ControllerId });
-	case EControllerID::DelayLFODepth:		 return FString::Format(TEXT("({0}) DelayLFODepth"), { ControllerId });
-	case EControllerID::DelayStereoType:	 return FString::Format(TEXT("({0}) DelayStereoType"), { ControllerId });
-	case EControllerID::DelayPanLeft:		 return FString::Format(TEXT("({0}) DelayPanLeft"), { ControllerId });
-	case EControllerID::DelayPanRight:		 return FString::Format(TEXT("({0}) DelayPanRight"), { ControllerId });
-	case EControllerID::RPNFine:			 return FString::Format(TEXT("({0}) RPNFine"), { ControllerId });
-	case EControllerID::RPNCourse:			 return FString::Format(TEXT("({0}) RPNCourse"), { ControllerId });
-	case EControllerID::AllSoundOff:		 return FString::Format(TEXT("({0}) AllSoundOff"), { ControllerId });
-	case EControllerID::Reset:				 return FString::Format(TEXT("({0}) Reset"), { ControllerId });
-	case EControllerID::AllNotesOff:		 return FString::Format(TEXT("({0}) AllNotesOff"), { ControllerId });
+		// this is the precision to which the "rounded" elapsed beats get presented.
+		const double kQuarterNotePrecision = 0.01 * ((double)TimeSignatureDenominator / 4.0); // account for beat being a quarter note, eighth note, etc.
+		double x = (double)InBeat;
+		x += 0.5 * kQuarterNotePrecision;
+		x *= 1.0 / kQuarterNotePrecision;
+		x = FMath::Floor(x);
+		x *= kQuarterNotePrecision;
+		return x;
 	}
-	return FString::Format(TEXT("({0}) - <unassigned>"), { ControllerId });
-}
 
-void MidiConstants::GetControllerNames(TArray<FString>& Names)
-{
-	for (int32 i = 0; i < 128; ++i)
+	FString GetControllerName(EControllerID ControllerId)
 	{
-		Names.Add(GetControllerName((EControllerID)i));
+		switch (ControllerId)
+		{
+		case EControllerID::BankSelection:       return FString::Format(TEXT("({0}) BankSelection"), { static_cast<uint8>(ControllerId) });
+		case EControllerID::ModWheel:			 return FString::Format(TEXT("({0}) ModWheel"), { static_cast<uint8>(ControllerId) });
+		case EControllerID::Breath:				 return FString::Format(TEXT("({0}) Breath"), { static_cast<uint8>(ControllerId) });
+		case EControllerID::PortamentoTime:		 return FString::Format(TEXT("({0}) PortamentoTime"), { static_cast<uint8>(ControllerId) });
+		case EControllerID::DataCoarse:			 return FString::Format(TEXT("({0}) DataCoarse"), { static_cast<uint8>(ControllerId) });
+		case EControllerID::Volume:				 return FString::Format(TEXT("({0}) Volume"), { static_cast<uint8>(ControllerId) });
+		case EControllerID::Balance:			 return FString::Format(TEXT("({0}) Balance"), { static_cast<uint8>(ControllerId) });
+		case EControllerID::PanRight:			 return FString::Format(TEXT("({0}) Pan (0 = left, 1 = right)"), { static_cast<uint8>(ControllerId) });
+		case EControllerID::Expression:			 return FString::Format(TEXT("({0}) Expression"), { static_cast<uint8>(ControllerId) });
+		case EControllerID::BitCrushWetMix:		 return FString::Format(TEXT("({0}) BitCrushWetMix"), { static_cast<uint8>(ControllerId) });
+		case EControllerID::BitCrushLevel:		 return FString::Format(TEXT("({0}) BitCrushLevel"), { static_cast<uint8>(ControllerId) });
+		case EControllerID::BitCrushSampleHold:	 return FString::Format(TEXT("({0}) BitCrushSampleHold"), { static_cast<uint8>(ControllerId) });
+		case EControllerID::LFO0Frequency:		 return FString::Format(TEXT("({0}) LFO0Frequency"), { static_cast<uint8>(ControllerId) });
+		case EControllerID::LFO1Frequency:		 return FString::Format(TEXT("({0}) LFO1Frequency"), { static_cast<uint8>(ControllerId) });
+		case EControllerID::LFO0Depth:			 return FString::Format(TEXT("({0}) LFO0Depth"), { static_cast<uint8>(ControllerId) });
+		case EControllerID::LFO1Depth:			 return FString::Format(TEXT("({0}) LFO1Depth"), { static_cast<uint8>(ControllerId) });
+		case EControllerID::DelayTime:			 return FString::Format(TEXT("({0}) DelayTime"), { static_cast<uint8>(ControllerId) });
+		case EControllerID::DelayDryGain:		 return FString::Format(TEXT("({0}) DelayDryGain"), { static_cast<uint8>(ControllerId) });
+		case EControllerID::DelayWetGain:		 return FString::Format(TEXT("({0}) DelayWetGain"), { static_cast<uint8>(ControllerId) });
+		case EControllerID::DelayFeedback:		 return FString::Format(TEXT("({0}) DelayFeedback"), { static_cast<uint8>(ControllerId) });
+		case EControllerID::CoarsePitchBend:	 return FString::Format(TEXT("({0}) CoarsePitchBend"), { static_cast<uint8>(ControllerId) });
+		case EControllerID::SampleStartTime:	 return FString::Format(TEXT("({0}) SampleStartTime"), { static_cast<uint8>(ControllerId) });
+		case EControllerID::DataFine:			 return FString::Format(TEXT("({0}) DataFine"), { static_cast<uint8>(ControllerId) });
+		case EControllerID::SubStreamVol1:		 return FString::Format(TEXT("({0}) SubStreamVol1"), { static_cast<uint8>(ControllerId) });
+		case EControllerID::SubStreamVol2:		 return FString::Format(TEXT("({0}) SubStreamVol2"), { static_cast<uint8>(ControllerId) });
+		case EControllerID::SubStreamVol3:		 return FString::Format(TEXT("({0}) SubStreamVol3"), { static_cast<uint8>(ControllerId) });
+		case EControllerID::SubStreamVol4:		 return FString::Format(TEXT("({0}) SubStreamVol4"), { static_cast<uint8>(ControllerId) });
+		case EControllerID::SubStreamVol5:		 return FString::Format(TEXT("({0}) SubStreamVol5"), { static_cast<uint8>(ControllerId) });
+		case EControllerID::SubStreamVol6:		 return FString::Format(TEXT("({0}) SubStreamVol6"), { static_cast<uint8>(ControllerId) });
+		case EControllerID::SubStreamVol7:		 return FString::Format(TEXT("({0}) SubStreamVol7"), { static_cast<uint8>(ControllerId) });
+		case EControllerID::SubStreamVol8:		 return FString::Format(TEXT("({0}) SubStreamVol8"), { static_cast<uint8>(ControllerId) });
+		case EControllerID::DelayEQEnabled:		 return FString::Format(TEXT("({0}) DelayEQEnabled"), { static_cast<uint8>(ControllerId) });
+		case EControllerID::DelayEQType:		 return FString::Format(TEXT("({0}) DelayEQType"), { static_cast<uint8>(ControllerId) });
+		case EControllerID::DelayEQFreq:		 return FString::Format(TEXT("({0}) DelayEQFreq"), { static_cast<uint8>(ControllerId) });
+		case EControllerID::DelayEQQ:			 return FString::Format(TEXT("({0}) DelayEQQ"), { static_cast<uint8>(ControllerId) });
+		case EControllerID::Hold:				 return FString::Format(TEXT("({0}) Hold"), { static_cast<uint8>(ControllerId) });
+		case EControllerID::PortamentoSwitch:	 return FString::Format(TEXT("({0}) PortamentoSwitch"), { static_cast<uint8>(ControllerId) });
+		case EControllerID::Sustenuto:			 return FString::Format(TEXT("({0}) Sustenuto"), { static_cast<uint8>(ControllerId) });
+		case EControllerID::SoftPedal:			 return FString::Format(TEXT("({0}) SoftPedal"), { static_cast<uint8>(ControllerId) });
+		case EControllerID::Legato:				 return FString::Format(TEXT("({0}) Legato"), { static_cast<uint8>(ControllerId) });
+		case EControllerID::Hold2:				 return FString::Format(TEXT("({0}) Hold2"), { static_cast<uint8>(ControllerId) });
+		case EControllerID::FilterQ:			 return FString::Format(TEXT("({0}) FilterQ"), { static_cast<uint8>(ControllerId) });
+		case EControllerID::Release:			 return FString::Format(TEXT("({0}) Release"), { static_cast<uint8>(ControllerId) });
+		case EControllerID::Attack:				 return FString::Format(TEXT("({0}) Attack"), { static_cast<uint8>(ControllerId) });
+		case EControllerID::FilterFrequency:	 return FString::Format(TEXT("({0}) FilterFrequency"), { static_cast<uint8>(ControllerId) });
+		case EControllerID::TimeStretchEnvelopeOrder:		 
+			return FString::Format(TEXT("({0}) TimeStretchEnvelopeOrder"), { static_cast<uint8>(ControllerId) });
+		case EControllerID::DelayLFOBeatSync:	 return FString::Format(TEXT("({0}) DelayLFOBeatSync"), { static_cast<uint8>(ControllerId) });
+		case EControllerID::DelayLFOEnabled:	 return FString::Format(TEXT("({0}) DelayLFOEnabled"), { static_cast<uint8>(ControllerId) });
+		case EControllerID::DelayLFORate:		 return FString::Format(TEXT("({0}) DelayLFORate"), { static_cast<uint8>(ControllerId) });
+		case EControllerID::DelayLFODepth:		 return FString::Format(TEXT("({0}) DelayLFODepth"), { static_cast<uint8>(ControllerId) });
+		case EControllerID::DelayStereoType:	 return FString::Format(TEXT("({0}) DelayStereoType"), { static_cast<uint8>(ControllerId) });
+		case EControllerID::DelayPanLeft:		 return FString::Format(TEXT("({0}) DelayPanLeft"), { static_cast<uint8>(ControllerId) });
+		case EControllerID::DelayPanRight:		 return FString::Format(TEXT("({0}) DelayPanRight"), { static_cast<uint8>(ControllerId) });
+		case EControllerID::RPNFine:			 return FString::Format(TEXT("({0}) RPNFine"), { static_cast<uint8>(ControllerId) });
+		case EControllerID::RPNCourse:			 return FString::Format(TEXT("({0}) RPNCourse"), { static_cast<uint8>(ControllerId) });
+		case EControllerID::AllSoundOff:		 return FString::Format(TEXT("({0}) AllSoundOff"), { static_cast<uint8>(ControllerId) });
+		case EControllerID::Reset:				 return FString::Format(TEXT("({0}) Reset"), { static_cast<uint8>(ControllerId) });
+		case EControllerID::AllNotesOff:		 return FString::Format(TEXT("({0}) AllNotesOff"), { static_cast<uint8>(ControllerId) });
+		}
+		return FString::Format(TEXT("({0}) - <unassigned>"), { static_cast<uint8>(ControllerId) });
 	}
-}
 
-const int8 MidiConstants::GetNoteNumberFromNoteName(const char* InName)
-{
-	if (isdigit(*InName))
-		return (int8)atoi(InName);
-
-	char Name[6];
-	size_t Index = 0;
-	for (; Index < 5 && Index < strlen(InName); Index++)
-		Name[Index] = (char)toupper(InName[Index]);
-	Name[Index] = 0;
-	char* Walk = Name;
-
-	int32 Note = -1;
-	switch (*Walk)
+	void GetControllerNames(TArray<FString>& Names)
 	{
-	case 'C': Note = 0;  break;
-	case 'D': Note = 2;  break;
-	case 'E': Note = 4;  break;
-	case 'F': Note = 5;  break;
-	case 'G': Note = 7;  break;
-	case 'A': Note = 9;  break;
-	case 'B': Note = 11; break;
+		for (int32 i = 0; i < 128; ++i)
+		{
+			Names.Add(GetControllerName((EControllerID)i));
+		}
 	}
-	if (Note == -1)
-		return 0;
-	Walk++;
-	if (*Walk == '#') 
+
+	int8 GetNoteNumberFromNoteName(const char* InName)
 	{
-		Note++; 
-		Walk++; 
+		if (isdigit(*InName))
+			return (int8)atoi(InName);
+
+		char Name[6];
+		size_t Index = 0;
+		for (; Index < 5 && Index < strlen(InName); Index++)
+			Name[Index] = (char)toupper(InName[Index]);
+		Name[Index] = 0;
+		char* Walk = Name;
+
+		int32 Note = -1;
+		switch (*Walk)
+		{
+		case 'C': Note = 0;  break;
+		case 'D': Note = 2;  break;
+		case 'E': Note = 4;  break;
+		case 'F': Note = 5;  break;
+		case 'G': Note = 7;  break;
+		case 'A': Note = 9;  break;
+		case 'B': Note = 11; break;
+		default: break;
+		}
+		if (Note == -1)
+			return 0;
+		Walk++;
+		if (*Walk == '#') 
+		{
+			Note++; 
+			Walk++; 
+		}
+		else if (*Walk == 'B')
+		{
+			Note--;
+			Walk++; 
+		}
+
+		int32 Octave = 0;
+		if (*Walk != '-')
+			Octave = atoi(Walk) + 1;
+
+		return int8(Note + (Octave * GNotesPerOctave));
 	}
-	else if (*Walk == 'B')
+
+	const char* GetNoteNameFromNoteNumber(uint8 MidiNoteNumber, ENoteNameEnharmonicStyle Style)
 	{
-		Note--;
-		Walk++; 
+		typedef char const* CStr;
+		static CStr sSharpAndFlat[GNotesPerOctave];
+		static CStr sSharp[GNotesPerOctave];
+		static CStr sFlat[GNotesPerOctave];
+
+		sSharpAndFlat[0] = "C";
+		sSharpAndFlat[1] = "C#/Db";
+		sSharpAndFlat[2] = "D";
+		sSharpAndFlat[3] = "D#/Eb";
+		sSharpAndFlat[4] = "E";
+		sSharpAndFlat[5] = "F";
+		sSharpAndFlat[6] = "F#/Gb";
+		sSharpAndFlat[7] = "G";
+		sSharpAndFlat[8] = "G#/Ab";
+		sSharpAndFlat[9] = "A";
+		sSharpAndFlat[10] = "A#/Bb";
+		sSharpAndFlat[11] = "B";
+
+		sSharp[0] = "C";
+		sSharp[1] = "C#";
+		sSharp[2] = "D";
+		sSharp[3] = "D#";
+		sSharp[4] = "E";
+		sSharp[5] = "F";
+		sSharp[6] = "F#";
+		sSharp[7] = "G";
+		sSharp[8] = "G#";
+		sSharp[9] = "A";
+		sSharp[10] = "A#";
+		sSharp[11] = "B";
+
+		sFlat[0] = "C";
+		sFlat[1] = "Db";
+		sFlat[2] = "D";
+		sFlat[3] = "Eb";
+		sFlat[4] = "E";
+		sFlat[5] = "F";
+		sFlat[6] = "Gb";
+		sFlat[7] = "G";
+		sFlat[8] = "Ab";
+		sFlat[9] = "A";
+		sFlat[10] = "Bb";
+		sFlat[11] = "B";
+
+		CStr const* NameTable;
+		
+		switch (Style)
+		{
+		case ENoteNameEnharmonicStyle::Flat:
+			NameTable = sFlat;
+			break;
+
+		case ENoteNameEnharmonicStyle::Sharp:
+			NameTable = sSharp;
+			break;
+		case ENoteNameEnharmonicStyle::SharpAndFlat:
+		default:
+			NameTable = sSharpAndFlat;
+			break;
+		}
+
+		check(NameTable != nullptr);
+		uint8 Index = MidiNoteNumber % GNotesPerOctave;
+		return NameTable[Index];
 	}
 
-	int32 Octave = 0;
-	if (*Walk != '-')
-		Octave = atoi(Walk) + 1;
-
-	return int8(Note + (Octave * kNotesPerOctave));
-}
-
-const char* MidiConstants::GetNoteNameFromNoteNumber(uint8 MidiNoteNumber, ENoteNameEnharmonicStyle Style)
-{
-	typedef char const* CStr;
-	static CStr sSharpAndFlat[kNotesPerOctave];
-	static CStr sSharp[kNotesPerOctave];
-	static CStr sFlat[kNotesPerOctave];
-
-	sSharpAndFlat[0] = "C";
-	sSharpAndFlat[1] = "C#/Db";
-	sSharpAndFlat[2] = "D";
-	sSharpAndFlat[3] = "D#/Eb";
-	sSharpAndFlat[4] = "E";
-	sSharpAndFlat[5] = "F";
-	sSharpAndFlat[6] = "F#/Gb";
-	sSharpAndFlat[7] = "G";
-	sSharpAndFlat[8] = "G#/Ab";
-	sSharpAndFlat[9] = "A";
-	sSharpAndFlat[10] = "A#/Bb";
-	sSharpAndFlat[11] = "B";
-
-	sSharp[0] = "C";
-	sSharp[1] = "C#";
-	sSharp[2] = "D";
-	sSharp[3] = "D#";
-	sSharp[4] = "E";
-	sSharp[5] = "F";
-	sSharp[6] = "F#";
-	sSharp[7] = "G";
-	sSharp[8] = "G#";
-	sSharp[9] = "A";
-	sSharp[10] = "A#";
-	sSharp[11] = "B";
-
-	sFlat[0] = "C";
-	sFlat[1] = "Db";
-	sFlat[2] = "D";
-	sFlat[3] = "Eb";
-	sFlat[4] = "E";
-	sFlat[5] = "F";
-	sFlat[6] = "Gb";
-	sFlat[7] = "G";
-	sFlat[8] = "Ab";
-	sFlat[9] = "A";
-	sFlat[10] = "Bb";
-	sFlat[11] = "B";
-
-	CStr const* NameTable = nullptr;
-	switch (Style)
+	int8 GetNoteOctaveFromNoteNumber(uint8 MidiNoteNumber)
 	{
-	case ENoteNameEnharmonicStyle::Flat:
-		NameTable = sFlat;
-		break;
-
-	case ENoteNameEnharmonicStyle::Sharp:
-		NameTable = sSharp;
-		break;
-	default:
-	case ENoteNameEnharmonicStyle::SharpAndFlat:
-		NameTable = sSharpAndFlat;
-		break;
+		return (int8)MidiNoteNumber / GNotesPerOctave - 1;
 	}
-
-	check(NameTable != nullptr);
-	uint8 Index = MidiNoteNumber % kNotesPerOctave;
-	return NameTable[Index];
-}
-
-int8 MidiConstants::GetNoteOctaveFromNoteNumber(uint8 MidiNoteNumber)
-{
-	return (int8)MidiNoteNumber / kNotesPerOctave - 1;
 }
diff --git a/Engine/Plugins/Runtime/Harmonix/Source/HarmonixMidi/Private/HarmonixMidi/MidiFile.cpp b/Engine/Plugins/Runtime/Harmonix/Source/HarmonixMidi/Private/HarmonixMidi/MidiFile.cpp
--- a/Engine/Plugins/Runtime/Harmonix/Source/HarmonixMidi/Private/HarmonixMidi/MidiFile.cpp
+++ b/Engine/Plugins/Runtime/Harmonix/Source/HarmonixMidi/Private/HarmonixMidi/MidiFile.cpp
@@ -119,7 +119,10 @@ void UMidiFile::PostInitProperties()
 #endif
 }
 
-void UMidiFile::LoadStdMidiFile(const FString& FilePath, int32 DesiredTicksPerQuarterNote, MidiConstants::EMidiTextEventEncoding InTextEncoding)
+void UMidiFile::LoadStdMidiFile(
+	const FString& FilePath,
+	int32 DesiredTicksPerQuarterNote,
+	Harmonix::Midi::Constants::EMidiTextEventEncoding InTextEncoding)
 {
 	FString Filename = FPaths::GetCleanFilename(FilePath);
 	IPlatformFile& PlatformFileApi = FPlatformFileManager::Get().GetPlatformFile();
@@ -136,13 +139,22 @@ void UMidiFile::LoadStdMidiFile(const FString& FilePath, int32 DesiredTicksPerQu
 #endif
 }
 
-void UMidiFile::LoadStdMidiFile(void* Buffer, int32 BufferSize, const FString& Filename, int32 DesiredTicksPerQuarterNote, MidiConstants::EMidiTextEventEncoding InTextEncoding)
+void UMidiFile::LoadStdMidiFile(
+	void* Buffer,
+	int32 BufferSize,
+	const FString& Filename,
+	int32 DesiredTicksPerQuarterNote,
+	Harmonix::Midi::Constants::EMidiTextEventEncoding InTextEncoding)
 {
 	TSharedPtr<FBufferReader> BufferArchive = MakeShared<FBufferReader>(Buffer, BufferSize, false);
 	LoadStdMidiFile(BufferArchive, Filename, DesiredTicksPerQuarterNote, InTextEncoding);
 }
 
-void UMidiFile::LoadStdMidiFile(TSharedPtr<FArchive> Archive, const FString& Filename, int32 DesiredTicksPerQuarterNote, MidiConstants::EMidiTextEventEncoding InTextEncoding)
+void UMidiFile::LoadStdMidiFile(
+	TSharedPtr<FArchive> Archive,
+	const FString& Filename,
+	int32 DesiredTicksPerQuarterNote,
+	Harmonix::Midi::Constants::EMidiTextEventEncoding InTextEncoding)
 {
 	TheMidiData.TicksPerQuarterNote = DesiredTicksPerQuarterNote;
 	TheMidiData.SongMaps.Init(TheMidiData.TicksPerQuarterNote);
@@ -606,8 +618,10 @@ void UMidiFile::ConformMidiFileLength(EMidiFileLengthConformOption Option)
 			}
 			uint8 MsgStatus = Msg.Status;
 
+			using namespace Harmonix::Midi::Constants;
+			
 			//filter chan press/pitch bend/program change events
-			if (MsgStatus == MidiConstants::kChanPres|| MsgStatus == MidiConstants::kPitch || MsgStatus == MidiConstants::kProgram)
+			if (MsgStatus == GChanPres|| MsgStatus == GPitch || MsgStatus == GProgram)
 			{
 				int32 OtherEventIndexToRemove = -1;
 				//check if there exist multiple events with same status on the last tick
@@ -623,7 +637,7 @@ void UMidiFile::ConformMidiFileLength(EMidiFileLengthConformOption Option)
 				{
 					Events.RemoveAt(OtherEventIndexToRemove);
 					NumItemsRemoved++;
-					if (MsgStatus == MidiConstants::kChanPres || MsgStatus == MidiConstants::kPolyPres)
+					if (MsgStatus == GChanPres || MsgStatus == GPolyPres)
 					{
 						NumAftertouchEventsRemoved++;
 					}
@@ -636,7 +650,7 @@ void UMidiFile::ConformMidiFileLength(EMidiFileLengthConformOption Option)
 			}
 
 			//filter Control Change events and Poly Pres events
-			if (MsgStatus == MidiConstants::kControl || MsgStatus == MidiConstants::kPolyPres)
+			if (MsgStatus == GControl || MsgStatus == GPolyPres)
 			{
 				uint8 CurrentEventControllerId = Msg.Data1;
 				int32 ControlEventIndexToRemove = -1; 
@@ -753,7 +767,7 @@ void FMidiEventReceiver::OnEndOfTrack(int InLastTick)
 		FMidiTrack& Track = File.GetTracks().Last();
 		uint16 StringIndex = Track.AddText("track_1");
 
-		Track.AddEvent(FMidiEvent(0, FMidiMsg::CreateText(StringIndex, MidiConstants::kMeta_TrackName)));
+		Track.AddEvent(FMidiEvent(0, FMidiMsg::CreateText(StringIndex, Harmonix::Midi::Constants::GMeta_TrackName)));
 		Track.Sort();
 		CurrentTrackHasName = true;
 	}
@@ -780,7 +794,7 @@ void FMidiEventReceiver::OnText(int32 Tick, const FString& Str, uint8 Type)
 	FMidiTrack& Track = File.GetTracks().Last();
 	uint16 StringIndex = Track.AddText(Str);
 
-	if (Type == MidiConstants::kMeta_TrackName)
+	if (Type == Harmonix::Midi::Constants::GMeta_TrackName)
 	{
 		// Track name should always appear at tick 0!
 		if (Tick != 0 && !CurrentTrackHasName)
@@ -855,7 +869,7 @@ void FMidiFileData::AddTempoChange(int32 TrackIdx, int32 Tick, float TempoBPM)
 	check(Tracks.IsValidIndex(TrackIdx));
 
 	FTempoMap& TempoMap = SongMaps.GetTempoMap();
-	int32 MidiTempo = MidiConstants::BPMToMidiTempo(TempoBPM);
+	int32 MidiTempo = Harmonix::Midi::Constants::BPMToMidiTempo(TempoBPM);
 	Tracks[TrackIdx].AddEvent(FMidiEvent(Tick, FMidiMsg(MidiTempo)));
 	Tracks[TrackIdx].Sort();
 	TempoMap.AddTempoInfoPoint(MidiTempo, Tick);
diff --git a/Engine/Plugins/Runtime/Harmonix/Source/HarmonixMidi/Private/HarmonixMidi/MidiMsg.cpp b/Engine/Plugins/Runtime/Harmonix/Source/HarmonixMidi/Private/HarmonixMidi/MidiMsg.cpp
--- a/Engine/Plugins/Runtime/Harmonix/Source/HarmonixMidi/Private/HarmonixMidi/MidiMsg.cpp
+++ b/Engine/Plugins/Runtime/Harmonix/Source/HarmonixMidi/Private/HarmonixMidi/MidiMsg.cpp
@@ -16,7 +16,7 @@ FMidiMsg FMidiMsg::CreateNoteOff(int32 Channel, int32 Note)
 
 FMidiMsg FMidiMsg::CreateControlChange(uint8 Channel, uint8 ControlNumber, uint8 Value)
 {
-	const uint8 StatusByte = MidiConstants::kControl | (Channel & 0xf);
+	const uint8 StatusByte = Harmonix::Midi::Constants::GControl | (Channel & 0xf);
 	return FMidiMsg{ StatusByte, ControlNumber, Value };
 }
 
@@ -115,55 +115,57 @@ void FMidiMsg::WriteStdMidi(int32 Tick, FMidiWriter& Writer, const FMidiTrack& T
 
 FString FMidiMsg::ToString(const FMidiMsg& Message, const FMidiTrack* Track /*= nullptr*/)
 {
+	using namespace Harmonix::Midi::Constants;
+	
 	FStringFormatOrderedArguments Args;
 	switch (Message.Type)
 	{
 	case EType::Std:
 		switch (Message.GetStdStatusType())
 		{
-		case MidiConstants::kNoteOff:
+		case GNoteOff:
 			Args.Add(Message.GetStdChannel() + 1);
 			Args.Add(Message.GetStdData1());
 			return FString::Format(TEXT("Note Off: Channel {0}, Note {1}"), Args);
-		case MidiConstants::kNoteOn:
+		case GNoteOn:
 			Args.Add(Message.GetStdChannel() + 1);
 			Args.Add(Message.GetStdData1());
 			Args.Add(Message.GetStdData2());
 			return FString::Format(TEXT("Note On: Channel {0}, Note {1}, Velocity {2}"), Args);
-		case MidiConstants::kPolyPres:
+		case GPolyPres:
 			Args.Add(Message.GetStdChannel() + 1);
 			Args.Add(Message.GetStdData1());
 			Args.Add(Message.GetStdData2());
 			return FString::Format(TEXT("Poly Pressure: Channel {0}, Note {1}, Pressure {2}"), Args);
-		case MidiConstants::kControl:
+		case GControl:
 			Args.Add(Message.GetStdChannel() + 1);
-			Args.Add(MidiConstants::GetControllerName((MidiConstants::EControllerID)Message.GetStdData1()));
+			Args.Add(GetControllerName((EControllerID)Message.GetStdData1()));
 			Args.Add(Message.GetStdData2());
 			return FString::Format(TEXT("Controller: Channel {0}, {1}, Value {2}"), Args);
-		case MidiConstants::kProgram:
+		case GProgram:
 			Args.Add(Message.GetStdChannel() + 1);
 			Args.Add(Message.GetStdData1());
 			return FString::Format(TEXT("Program Change: Channel {0}, Program {1}"), Args);
-		case MidiConstants::kChanPres:
+		case GChanPres:
 			Args.Add(Message.GetStdChannel() + 1);
 			Args.Add(Message.GetStdData1());
 			return FString::Format(TEXT("Aftertouch: Channel {0}, Value {1}"), Args);
-		case MidiConstants::kPitch:
+		case GPitch:
 			Args.Add(Message.GetStdChannel() + 1);
 			Args.Add(Message.GetPitchBendFromData());
 			return FString::Format(TEXT("Pitchbend: Channel {0}, Value {1}"), Args);
-		case MidiConstants::kSystem:
+		case GSystem:
 			return TEXT("!! SYstem Exclusive !!");
 		}
 	case EType::Tempo:
-		Args.Add(MidiConstants::MidiTempoToBPM(Message.GetMicrosecPerQuarterNote()));
+		Args.Add(MidiTempoToBPM(Message.GetMicrosecPerQuarterNote()));
 		return FString::Format(TEXT("Tempo Change: {0} bpm"), Args);
 	case EType::TimeSig:
 		Args.Add(Message.GetTimeSigNumerator());
 		Args.Add(Message.GetTimeSigDenominator());
 		return FString::Format(TEXT("Time Signature: {0}/{1}"), Args);
 	case EType::Text:
-		Args.Add(MidiConstants::GetTextTypeName(Message.GetTextType()));
+		Args.Add(GetTextTypeName(Message.GetTextType()));
 		if (Track)
 		{
 			Args.Add(Track->GetTextAtIndex(Message.GetTextIndex()));
@@ -176,9 +178,9 @@ FString FMidiMsg::ToString(const FMidiMsg& Message, const FMidiTrack* Track /*=
 	case EType::Runtime:
 		switch (Message.Status)
 		{
-			case MidiConstants::kRuntimeAllNotesOffStatus:
+			case GRuntimeAllNotesOffStatus:
 				return TEXT("Runtime: All Notes Off (allow ADSR releases.)");
-			case MidiConstants::kRuntimeAllNotesKillStatus:
+			case GRuntimeAllNotesKillStatus:
 				return TEXT("Runtime: All Notes Kill (No ADSR releases. Kill all rendering.)");
 			default:
 				checkNoEntry();
diff --git a/Engine/Plugins/Runtime/Harmonix/Source/HarmonixMidi/Private/HarmonixMidi/MidiPlayCursorMgr.cpp b/Engine/Plugins/Runtime/Harmonix/Source/HarmonixMidi/Private/HarmonixMidi/MidiPlayCursorMgr.cpp
--- a/Engine/Plugins/Runtime/Harmonix/Source/HarmonixMidi/Private/HarmonixMidi/MidiPlayCursorMgr.cpp
+++ b/Engine/Plugins/Runtime/Harmonix/Source/HarmonixMidi/Private/HarmonixMidi/MidiPlayCursorMgr.cpp
@@ -25,7 +25,7 @@ FMidiPlayCursorMgr::FMidiPlayCursorMgr()
 	, CurrentAdvanceRate(1.f)
 {
 	// setup the default tempo map to have one entry...
-	DefaultMaps.GetTempoMap().AddTempoInfoPoint(MidiConstants::BPMToMidiTempo(120.0f), 0);
+	DefaultMaps.GetTempoMap().AddTempoInfoPoint(Harmonix::Midi::Constants::BPMToMidiTempo(120.0f), 0);
 	DefaultMaps.GetBarMap().AddTimeSignatureAtBarIncludingCountIn(0, 4, 4);
 }
 
diff --git a/Engine/Plugins/Runtime/Harmonix/Source/HarmonixMidi/Private/HarmonixMidi/MidiReader.cpp b/Engine/Plugins/Runtime/Harmonix/Source/HarmonixMidi/Private/HarmonixMidi/MidiReader.cpp
--- a/Engine/Plugins/Runtime/Harmonix/Source/HarmonixMidi/Private/HarmonixMidi/MidiReader.cpp
+++ b/Engine/Plugins/Runtime/Harmonix/Source/HarmonixMidi/Private/HarmonixMidi/MidiReader.cpp
@@ -78,7 +78,13 @@ FArchive& operator<<(FArchive& Archive, FMidiChunkHeader& Header)
 
 ///////////////////////////////////////////////////////////////////////////////
 // FStdMidiFileReader
-FStdMidiFileReader::FStdMidiFileReader(void* Buffer, int32 BufferSize, const FString& InFileName, IMidiReceiver* InReceiver, int32 TicksPerQuarterNote, MidiConstants::EMidiTextEventEncoding InTextEncoding)
+FStdMidiFileReader::FStdMidiFileReader(
+	void* Buffer,
+	int32 BufferSize,
+	 const FString& InFileName,
+	 IMidiReceiver* InReceiver,
+	 int32 TicksPerQuarterNote,
+	 Harmonix::Midi::Constants::EMidiTextEventEncoding InTextEncoding)
 	: Filename(InFileName)
 	, Receiver(InReceiver)
 	, TextEncoding(InTextEncoding)
@@ -90,7 +96,11 @@ FStdMidiFileReader::FStdMidiFileReader(void* Buffer, int32 BufferSize, const FSt
 	Init();
 }
 
-FStdMidiFileReader::FStdMidiFileReader(const FString& FilePath, IMidiReceiver* InReceiver, int32 TicksPerQuarterNote, MidiConstants::EMidiTextEventEncoding InTextEncoding)
+FStdMidiFileReader::FStdMidiFileReader(
+	const FString& FilePath,
+	IMidiReceiver* InReceiver,
+	int32 TicksPerQuarterNote,
+	Harmonix::Midi::Constants::EMidiTextEventEncoding InTextEncoding)
 	: Receiver(InReceiver)
 	, TextEncoding(InTextEncoding)
 	, DestinationTicksPerQuarterNote(TicksPerQuarterNote)
@@ -104,7 +114,12 @@ FStdMidiFileReader::FStdMidiFileReader(const FString& FilePath, IMidiReceiver* I
 	Init();
 }
 
-FStdMidiFileReader::FStdMidiFileReader(TSharedPtr<FArchive> Archive, const FString& InFilename, IMidiReceiver* InReceiver, int32 TicksPerQuarterNote, MidiConstants::EMidiTextEventEncoding InTextEncoding)
+FStdMidiFileReader::FStdMidiFileReader(
+	TSharedPtr<FArchive> Archive,
+	const FString& InFilename,
+	IMidiReceiver* InReceiver,
+	int32 TicksPerQuarterNote,
+	Harmonix::Midi::Constants::EMidiTextEventEncoding InTextEncoding)
 	: Filename(InFilename)
 	, InputArchive(Archive)
 	, Receiver(InReceiver)
@@ -115,7 +130,13 @@ FStdMidiFileReader::FStdMidiFileReader(TSharedPtr<FArchive> Archive, const FStri
 	Init();
 }
 
-bool FStdMidiFileReader::ReadStdMidiFileForReceiver(void* Buffer, int32 BufferSize, const FString& FileName, IMidiReceiver* Receiver, int32 TicksPerQuarterNote, MidiConstants::EMidiTextEventEncoding InTextEncoding)
+bool FStdMidiFileReader::ReadStdMidiFileForReceiver(
+	void* Buffer,
+	int32 BufferSize,
+	const FString& FileName,
+	 IMidiReceiver* Receiver,
+	  int32 TicksPerQuarterNote,
+	  Harmonix::Midi::Constants::EMidiTextEventEncoding InTextEncoding)
 {
 	bool Success = true;
 	FStdMidiFileReader Reader(Buffer, BufferSize, FileName, Receiver, TicksPerQuarterNote, InTextEncoding);
@@ -376,7 +397,7 @@ void FStdMidiFileReader::ReadTrackHeader(FArchive& Archive)
 		{
 			CurrentTrackName = TEXT("Track-0");
 		}
-		Receiver->OnText(0, CurrentTrackName, MidiConstants::kMeta_TrackName);
+		Receiver->OnText(0, CurrentTrackName, Harmonix::Midi::Constants::GMeta_TrackName);
 	}
 }
 
@@ -410,10 +431,10 @@ void FStdMidiFileReader::ReadEvent(FArchive& Archive)
 	}
 
 	Archive << Status;
-	if (MidiConstants::IsStatus(Status))  // this byte is truly a status byte
+	if (Harmonix::Midi::Constants::IsStatus(Status))  // this byte is truly a status byte
 	{
 		RunningStatus = false;
-		if (!MidiConstants::IsSystem(Status))
+		if (!Harmonix::Midi::Constants::IsSystem(Status))
 		{
 			// save the current status byte, in case the next message uses running
 			// status.  Note that this is only for midi events; running status can
@@ -431,7 +452,7 @@ void FStdMidiFileReader::ReadEvent(FArchive& Archive)
 	}
 
 	// Check if this is a system (ie, non-channel) message
-	if (MidiConstants::IsSystem(Status))
+	if (Harmonix::Midi::Constants::IsSystem(Status))
 	{
 		ReadSystemEvent(CurrentTick, Status, Archive);
 	}                   // All other messages are channel (non-system) messages
@@ -448,39 +469,40 @@ void FStdMidiFileReader::ReadEvent(FArchive& Archive)
 
 void FStdMidiFileReader::ReadMidiEvent(int32 Tick, uint8 Status, uint8 Data1, FArchive& Archive)
 {
+	using namespace Harmonix::Midi::Constants;
 	uint8 Data2 = 0;
 	bool ValidMidiEvent = false;
-	switch (Status & MidiConstants::kMessageTypeMask)
+	switch (Status & GMessageTypeMask)
 	{
 		// special case for NoteOn (3 byte message)
-	case MidiConstants::kNoteOn:
+	case GNoteOn:
 		Archive << Data2;
 		ValidMidiEvent = true;
 		// convert note-on vel=0 -> note offs
 		if (Data2 == 0)
 		{
-			Status = MidiConstants::kNoteOff | (Status & MidiConstants::kChannelMask);
+			Status = GNoteOff | (Status & GChannelMask);
 		}
 		break;
 
 		// other 3 byte messages
-	case MidiConstants::kNoteOff:
-	case MidiConstants::kControl:
-	case MidiConstants::kPitch:
-	case MidiConstants::kPolyPres:
+	case GNoteOff:
+	case GControl:
+	case GPitch:
+	case GPolyPres:
 		Archive << Data2;
 		ValidMidiEvent = true;
 		break;
 
 		// 2 byte messages
-	case MidiConstants::kProgram:
-	case MidiConstants::kChanPres:
+	case GProgram:
+	case GChanPres:
 		Data2 = 0;
 		ValidMidiEvent = true;
 		break;
 
 	default:
-		UE_LOG(LogMidi, Error, TEXT("%s (%s): Cannot parse event %i"), *Filename, *CurrentTrackName, (Status & MidiConstants::kMessageTypeMask));
+		UE_LOG(LogMidi, Error, TEXT("%s (%s): Cannot parse event %i"), *Filename, *CurrentTrackName, (Status & GMessageTypeMask));
 		break;
 	}
 
@@ -494,10 +516,12 @@ void FStdMidiFileReader::ReadMidiEvent(int32 Tick, uint8 Status, uint8 Data1, FA
 // Read system event from the stream.  Status byte has already been read.
 void FStdMidiFileReader::ReadSystemEvent(int32 Tick, uint8 Status, FArchive& Archive)
 {
+	using namespace Harmonix::Midi::Constants;
+	
 	switch (Status)
 	{
-	case MidiConstants::kFile_SysEx:
-	case MidiConstants::kFile_Escape:
+	case GFile_SysEx:
+	case GFile_Escape:
 		{
 			// All system events start with length. We'll
 			// seek past that length and ignore these events
@@ -506,7 +530,7 @@ void FStdMidiFileReader::ReadSystemEvent(int32 Tick, uint8 Status, FArchive& Arc
 			Archive.Seek(Archive.Tell() + PacketLength);
 		}
 		break;
-	case MidiConstants::kFile_Meta:
+	case GFile_Meta:
 		{
 			uint8 type;
 			Archive << type;
@@ -523,34 +547,36 @@ void FStdMidiFileReader::ReadSystemEvent(int32 Tick, uint8 Status, FArchive& Arc
 // Status byte and event type have already been read.
 void FStdMidiFileReader::ReadMetaEvent(int32 Tick, uint8 Type, FArchive& Archive)
 {
+	using namespace Harmonix::Midi::Constants;
+	
 	int32 Length = Midi::VarLenNumber::Read(Archive);
 	int32 StartPos = int32(Archive.Tell());
 	FString WorkingString;
 
 	switch (Type)
 	{
-	case MidiConstants::kMeta_TrackName:
+	case GMeta_TrackName:
 		CurrentTrackName = ReadText(Archive, Length);
 		Receiver->OnText(Tick, CurrentTrackName, Type);
 		break;
-	case MidiConstants::kMeta_Copyright:
-	case MidiConstants::kMeta_Marker:
-	case MidiConstants::kMeta_CuePoint:
+	case GMeta_Copyright:
+	case GMeta_Marker:
+	case GMeta_CuePoint:
 		WorkingString = ReadText(Archive, Length);
 		if (TrackFilteringMode != ETrackFilteringMode::NonConductorEvents)
 		{
 			Receiver->OnText(Tick, WorkingString, Type);
 		}
 		break;
-	case MidiConstants::kMeta_Text:
-	case MidiConstants::kMeta_Lyric:
+	case GMeta_Text:
+	case GMeta_Lyric:
 		WorkingString = ReadText(Archive, Length);
 		if (TrackFilteringMode != ETrackFilteringMode::ConductorEvents)
 		{
 			Receiver->OnText(Tick, WorkingString, Type);
 		}
 		break;
-	case MidiConstants::kMeta_Tempo:
+	case GMeta_Tempo:
 		{
 			uint8 Byte1, Byte2, Byte3;
 			Archive << Byte1 << Byte2 << Byte3;
@@ -562,7 +588,7 @@ void FStdMidiFileReader::ReadMetaEvent(int32 Tick, uint8 Type, FArchive& Archive
 		}
 		break;
 
-	case MidiConstants::kMeta_EndOfTrack:
+	case GMeta_EndOfTrack:
 		ProcessMidiList();  // in case there is anything left to send out.
 		Receiver->OnEndOfTrack(Tick);
 		if (CurrentTrackIndex == NumTracks - 1)
@@ -576,7 +602,7 @@ void FStdMidiFileReader::ReadMetaEvent(int32 Tick, uint8 Type, FArchive& Archive
 		}
 		break;
 
-	case MidiConstants::kMeta_TimeSig:
+	case GMeta_TimeSig:
 		{
 			uint8 Numerator, DenominatorExp;
 			Archive << Numerator << DenominatorExp;
@@ -603,7 +629,6 @@ void FStdMidiFileReader::ReadMetaEvent(int32 Tick, uint8 Type, FArchive& Archive
 			{
 				// before telling any receivers about the time signature we update our own local tempo map
 				// so we can print better error and warning messages...
-				FMusicTimestamp Timestamp;
 				check(Tick == 0 || BarMap->GetNumTimeSignaturePoints() > 0);
 				int32 BarIndex = BarMap->TickToBarIncludingCountIn(Tick);
 				if (!BarMap->AddTimeSignatureAtBarIncludingCountIn(BarIndex, Numerator, Denominator, true, false))
@@ -621,11 +646,11 @@ void FStdMidiFileReader::ReadMetaEvent(int32 Tick, uint8 Type, FArchive& Archive
 		}
 		break;
 
-	case MidiConstants::kMeta_ChannelPrefix:
-	case MidiConstants::kMeta_Port:
-	case MidiConstants::kMeta_KeySig:
-	case MidiConstants::kMeta_SMPTE:
-	case MidiConstants::kMeta_InstrumentName:
+	case GMeta_ChannelPrefix:
+	case GMeta_Port:
+	case GMeta_KeySig:
+	case GMeta_SMPTE:
+	case GMeta_InstrumentName:
 		// these are valid events, but not currently supported by
 		// MidiReceiver; do nothing
 		break;
@@ -648,7 +673,7 @@ void FStdMidiFileReader::ReadMetaEvent(int32 Tick, uint8 Type, FArchive& Archive
 FString FStdMidiFileReader::ReadText(FArchive& Archive, int32 Length)
 {
 	unsigned char* AsUtf8 = nullptr;
-	if (TextEncoding == MidiConstants::EMidiTextEventEncoding::Latin1)
+	if (TextEncoding == Harmonix::Midi::Constants::EMidiTextEventEncoding::Latin1)
 	{
 		unsigned char* AsLatin = (unsigned char*)FMemory::Malloc(Length + 1);
 		Archive.Serialize(AsLatin, Length);
diff --git a/Engine/Plugins/Runtime/Harmonix/Source/HarmonixMidi/Private/HarmonixMidi/MidiSongPos.cpp b/Engine/Plugins/Runtime/Harmonix/Source/HarmonixMidi/Private/HarmonixMidi/MidiSongPos.cpp
--- a/Engine/Plugins/Runtime/Harmonix/Source/HarmonixMidi/Private/HarmonixMidi/MidiSongPos.cpp
+++ b/Engine/Plugins/Runtime/Harmonix/Source/HarmonixMidi/Private/HarmonixMidi/MidiSongPos.cpp
@@ -22,8 +22,8 @@ int32 SongPosCmp(const FMidiSongPos& lhs, const FMidiSongPos& rhs)
 	}
 
 	// beat
-	if (MidiConstants::RoundToStandardBeatPrecision(lhs.Timestamp.Beat, lhs.TimeSigDenominator) ==
-		MidiConstants::RoundToStandardBeatPrecision(rhs.Timestamp.Beat, rhs.TimeSigDenominator))
+	if (Harmonix::Midi::Constants::RoundToStandardBeatPrecision(lhs.Timestamp.Beat, lhs.TimeSigDenominator) ==
+		Harmonix::Midi::Constants::RoundToStandardBeatPrecision(rhs.Timestamp.Beat, rhs.TimeSigDenominator))
 	{
 		return 0;
 	}
diff --git a/Engine/Plugins/Runtime/Harmonix/Source/HarmonixMidi/Private/HarmonixMidi/MidiTrack.cpp b/Engine/Plugins/Runtime/Harmonix/Source/HarmonixMidi/Private/HarmonixMidi/MidiTrack.cpp
--- a/Engine/Plugins/Runtime/Harmonix/Source/HarmonixMidi/Private/HarmonixMidi/MidiTrack.cpp
+++ b/Engine/Plugins/Runtime/Harmonix/Source/HarmonixMidi/Private/HarmonixMidi/MidiTrack.cpp
@@ -21,7 +21,7 @@ FMidiTrack::FMidiTrack(const FString& Name)
 	, PrimaryMidiChannel(-1)
 {
 	uint16 stringIndex = AddText(Name);
-	AddEvent(FMidiEvent(0, FMidiMsg::CreateText(stringIndex, MidiConstants::kMeta_TrackName)));
+	AddEvent(FMidiEvent(0, FMidiMsg::CreateText(stringIndex, Harmonix::Midi::Constants::GMeta_TrackName)));
 }
 
 const FMidiEventList& FMidiTrack::GetEvents() const
@@ -48,7 +48,7 @@ void FMidiTrack::SetName(const FString& InName)
 		if (Event.GetMsg().MsgType() == FMidiMsg::EType::Text)
 		{
 			FMidiMsg& Message = Event.GetMsg();
-			if (Message.GetTextType() == MidiConstants::kMeta_TrackName)
+			if (Message.GetTextType() == Harmonix::Midi::Constants::GMeta_TrackName)
 			{
 				Strings[Message.GetTextIndex()] = InName;
 				return;
@@ -65,7 +65,7 @@ const FString* FMidiTrack::GetName() const
 		if (Event.GetMsg().MsgType() == FMidiMsg::EType::Text)
 		{
 			const FMidiMsg& Message = Event.GetMsg();
-			if (Message.GetTextType() == MidiConstants::kMeta_TrackName)
+			if (Message.GetTextType() == Harmonix::Midi::Constants::GMeta_TrackName)
 			{
 				return &GetTextAtIndex(Message.GetTextIndex());
 			}
@@ -87,7 +87,7 @@ void FMidiTrack::AddEvent(const FMidiEvent& Event)
 	}
 
 	// Last name event added at tick zero wins.
-	if (Event.GetTick() == 0 && Event.GetMsg().MsgType() == FMidiMsg::EType::Text && Event.GetMsg().GetTextType() == MidiConstants::kMeta_TrackName)
+	if (Event.GetTick() == 0 && Event.GetMsg().MsgType() == FMidiMsg::EType::Text && Event.GetMsg().GetTextType() == Harmonix::Midi::Constants::GMeta_TrackName)
 	{
 		for (auto& ExistingEvent : Events)
 		{
@@ -96,7 +96,7 @@ void FMidiTrack::AddEvent(const FMidiEvent& Event)
 				break;
 			}
 			const FMidiMsg& Message = ExistingEvent.GetMsg();
-			if (Message.MsgType() == FMidiMsg::EType::Text && Message.GetTextType() == MidiConstants::kMeta_TrackName)
+			if (Message.MsgType() == FMidiMsg::EType::Text && Message.GetTextType() == Harmonix::Midi::Constants::GMeta_TrackName)
 			{
 				ExistingEvent = Event;
 				return;
@@ -188,8 +188,8 @@ int32 FMidiTrack::CopyEvents(FMidiTrack& SourceTrack,	int32 FromTick, int32 Thru
 
 		const FMidiMsg& SourceMessage = SourceEvent.GetMsg();
 		if (SourceMessage.MsgType() == FMidiMsg::EType::Std &&
-			(SourceMessage.GetStdStatusType() == MidiConstants::kNoteOn ||
-				SourceMessage.GetStdStatusType() == MidiConstants::kNoteOff))
+			(SourceMessage.GetStdStatusType() == Harmonix::Midi::Constants::GNoteOn ||
+				SourceMessage.GetStdStatusType() == Harmonix::Midi::Constants::GNoteOff))
 		{
 			// we may have to filter or transpose...
 			int32 NoteNum = SourceMessage.GetStdData1();
@@ -205,7 +205,7 @@ int32 FMidiTrack::CopyEvents(FMidiTrack& SourceTrack,	int32 FromTick, int32 Thru
 			// we have to do something a little special for text messages...
 
 			// Is this the track name, and are we filtering the track name?
-			if (SourceMessage.GetTextType() == MidiConstants::kMeta_TrackName && FilterTrackName)
+			if (SourceMessage.GetTextType() == Harmonix::Midi::Constants::GMeta_TrackName && FilterTrackName)
 			{
 				continue;
 			}
diff --git a/Engine/Plugins/Runtime/Harmonix/Source/HarmonixMidi/Private/HarmonixMidi/MidiWriter.cpp b/Engine/Plugins/Runtime/Harmonix/Source/HarmonixMidi/Private/HarmonixMidi/MidiWriter.cpp
--- a/Engine/Plugins/Runtime/Harmonix/Source/HarmonixMidi/Private/HarmonixMidi/MidiWriter.cpp
+++ b/Engine/Plugins/Runtime/Harmonix/Source/HarmonixMidi/Private/HarmonixMidi/MidiWriter.cpp
@@ -40,8 +40,8 @@ void FMidiWriter::Close()
 void FMidiWriter::EndOfTrack()
 {
 	ProcessTick(CurTick);
-	_WRITE_TRACK_MIDI_BYTE(MidiConstants::kFile_Meta);
-	_WRITE_TRACK_MIDI_BYTE(MidiConstants::kMeta_EndOfTrack);
+	_WRITE_TRACK_MIDI_BYTE(Harmonix::Midi::Constants::GFile_Meta);
+	_WRITE_TRACK_MIDI_BYTE(Harmonix::Midi::Constants::GMeta_EndOfTrack);
 	_WRITE_TRACK_MIDI_BYTE(0);
 	TSharedPtr<FBufferWriter> Writer = MakeShared<FBufferWriter>(FMemory::Malloc(1024), 1024, EBufferWriterFlags::TakeOwnership | EBufferWriterFlags::AllowResize);
 	TrackWriters.Add(Writer);
@@ -51,24 +51,25 @@ void FMidiWriter::EndOfTrack()
 
 void FMidiWriter::MidiMessage(int32 Tick, uint8 Status, uint8 Data1, uint8 Data2)
 {
+	using namespace Harmonix::Midi::Constants;
 	ProcessTick(Tick);
 
 	*CurrentWriter << Status;
 
 	// messages are of varying length:
-	switch (Status & MidiConstants::kMessageTypeMask)
+	switch (Status & GMessageTypeMask)
 	{
-	case MidiConstants::kNoteOn: // 3-byte
-	case MidiConstants::kNoteOff:
-	case MidiConstants::kControl:
-	case MidiConstants::kPitch:
-	case MidiConstants::kPolyPres:
+	case GNoteOn: // 3-byte
+	case GNoteOff:
+	case GControl:
+	case GPitch:
+	case GPolyPres:
 		*CurrentWriter << Data1 << Data2;
 		break;
 
 		// 2 byte messages
-	case MidiConstants::kProgram:
-	case MidiConstants::kChanPres:
+	case GProgram:
+	case GChanPres:
 		*CurrentWriter << Data1;
 		break;
 
@@ -80,10 +81,11 @@ void FMidiWriter::MidiMessage(int32 Tick, uint8 Status, uint8 Data1, uint8 Data2
 
 void FMidiWriter::Tempo(int32 Tick, int32 Tempo)
 {
+	using namespace Harmonix::Midi::Constants;
 	check(Tempo <= 0xFFFFFF);
 	ProcessTick(Tick);
-	_WRITE_TRACK_MIDI_BYTE(MidiConstants::kFile_Meta);  // meta-event
-	_WRITE_TRACK_MIDI_BYTE(MidiConstants::kMeta_Tempo); // tempo tag
+	_WRITE_TRACK_MIDI_BYTE(GFile_Meta);  // meta-event
+	_WRITE_TRACK_MIDI_BYTE(GMeta_Tempo); // tempo tag
 	_WRITE_TRACK_MIDI_BYTE(0x03);							// size of following
 	_WRITE_TRACK_MIDI_BYTE(uint8((Tempo >> 16) & 0xFF));  // 3 bytes of tempo
 	_WRITE_TRACK_MIDI_BYTE(uint8((Tempo >> 8) & 0xFF));
@@ -94,7 +96,7 @@ void FMidiWriter::Text(int32 Tick, const TCHAR* Str, uint8 Type)
 {
 	check(Type >= 0x01 && Type <= 0x07);
 	ProcessTick(Tick);
-	_WRITE_TRACK_MIDI_BYTE(MidiConstants::kFile_Meta); // meta-event
+	_WRITE_TRACK_MIDI_BYTE(Harmonix::Midi::Constants::GFile_Meta); // meta-event
 	_WRITE_TRACK_MIDI_BYTE(Type);						   // tag
 
 	auto TextAsUtf8 = StringCast<UTF8CHAR>(Str);
@@ -105,11 +107,12 @@ void FMidiWriter::Text(int32 Tick, const TCHAR* Str, uint8 Type)
 
 void FMidiWriter::TimeSignature(int32 Tick, int32 Numerator, int32 Denominator)
 {
+	using namespace Harmonix::Midi::Constants;
 	check(Numerator > 0);
 	check(Denominator > 0);
 	ProcessTick(Tick);
-	_WRITE_TRACK_MIDI_BYTE(MidiConstants::kFile_Meta);    // meta-event
-	_WRITE_TRACK_MIDI_BYTE(MidiConstants::kMeta_TimeSig); // tag
+	_WRITE_TRACK_MIDI_BYTE(GFile_Meta);    // meta-event
+	_WRITE_TRACK_MIDI_BYTE(GMeta_TimeSig); // tag
 	_WRITE_TRACK_MIDI_BYTE(0x04);							  // size of following
 	_WRITE_TRACK_MIDI_BYTE(uint8(Numerator));               // numerator of time signature
 
diff --git a/Engine/Plugins/Runtime/Harmonix/Source/HarmonixMidi/Private/HarmonixMidi/MusicTimeSpan.cpp b/Engine/Plugins/Runtime/Harmonix/Source/HarmonixMidi/Private/HarmonixMidi/MusicTimeSpan.cpp
--- a/Engine/Plugins/Runtime/Harmonix/Source/HarmonixMidi/Private/HarmonixMidi/MusicTimeSpan.cpp
+++ b/Engine/Plugins/Runtime/Harmonix/Source/HarmonixMidi/Private/HarmonixMidi/MusicTimeSpan.cpp
@@ -22,6 +22,7 @@ float FMusicalTimeSpan::CalcPositionInSpan(const FMidiSongPos& Position, const F
 
 float FMusicalTimeSpan::CalcPositionInSpanWithOffset(const FMidiSongPos& Position, const FSongMaps& Maps) const
 {
+	using namespace Harmonix::Midi::Constants;
 	FMidiSongPos OffsetPosition;
 
 	int32 UnadjustedTick = Maps.MsToTick(Position.SecondsIncludingCountIn * 1000.0f);
@@ -47,49 +48,49 @@ float FMusicalTimeSpan::CalcPositionInSpanWithOffset(const FMidiSongPos& Positio
 		}
 		break;
 	case EMusicTimeSpanOffsetUnits::ThirtySecondNotes:
-		AdjustedTick = UnadjustedTick - MidiConstants::kTicksPerQuarterNoteInt * Offset / 8;
+		AdjustedTick = UnadjustedTick - GTicksPerQuarterNoteInt * Offset / 8;
 		break;
 	case EMusicTimeSpanOffsetUnits::SixteenthNotes:
-		AdjustedTick = UnadjustedTick - MidiConstants::kTicksPerQuarterNoteInt * Offset / 4;
+		AdjustedTick = UnadjustedTick - GTicksPerQuarterNoteInt * Offset / 4;
 		break;
 	case EMusicTimeSpanOffsetUnits::EighthNotes:
-		AdjustedTick = UnadjustedTick - MidiConstants::kTicksPerQuarterNoteInt * Offset / 2;
+		AdjustedTick = UnadjustedTick - GTicksPerQuarterNoteInt * Offset / 2;
 		break;
 	case EMusicTimeSpanOffsetUnits::QuarterNotes:
-		AdjustedTick = UnadjustedTick - MidiConstants::kTicksPerQuarterNoteInt * Offset;
+		AdjustedTick = UnadjustedTick - GTicksPerQuarterNoteInt * Offset;
 		break;
 	case EMusicTimeSpanOffsetUnits::HalfNotes:
-		AdjustedTick = UnadjustedTick - MidiConstants::kTicksPerQuarterNoteInt * Offset * 2;
+		AdjustedTick = UnadjustedTick - GTicksPerQuarterNoteInt * Offset * 2;
 		break;
 	case EMusicTimeSpanOffsetUnits::WholeNotes:
-		AdjustedTick = UnadjustedTick - MidiConstants::kTicksPerQuarterNoteInt * Offset * 4;
+		AdjustedTick = UnadjustedTick - GTicksPerQuarterNoteInt * Offset * 4;
 		break;
 	case EMusicTimeSpanOffsetUnits::DottedSixteenthNotes:
-		AdjustedTick = UnadjustedTick - MidiConstants::kTicksPerQuarterNoteInt * Offset * 3 / 8;
+		AdjustedTick = UnadjustedTick - GTicksPerQuarterNoteInt * Offset * 3 / 8;
 		break;
 	case EMusicTimeSpanOffsetUnits::DottedEighthNotes:
-		AdjustedTick = UnadjustedTick - MidiConstants::kTicksPerQuarterNoteInt * Offset * 3 / 4;
+		AdjustedTick = UnadjustedTick - GTicksPerQuarterNoteInt * Offset * 3 / 4;
 		break;
 	case EMusicTimeSpanOffsetUnits::DottedQuarterNotes:
-		AdjustedTick = UnadjustedTick - MidiConstants::kTicksPerQuarterNoteInt * Offset * 3 / 2;
+		AdjustedTick = UnadjustedTick - GTicksPerQuarterNoteInt * Offset * 3 / 2;
 		break;
 	case EMusicTimeSpanOffsetUnits::DottedHalfNotes:
-		AdjustedTick = UnadjustedTick - MidiConstants::kTicksPerQuarterNoteInt * Offset * 3;
+		AdjustedTick = UnadjustedTick - GTicksPerQuarterNoteInt * Offset * 3;
 		break;
 	case EMusicTimeSpanOffsetUnits::DottedWholeNotes:
-		AdjustedTick = UnadjustedTick - MidiConstants::kTicksPerQuarterNoteInt * Offset * 6;
+		AdjustedTick = UnadjustedTick - GTicksPerQuarterNoteInt * Offset * 6;
 		break;
 	case EMusicTimeSpanOffsetUnits::SixteenthNoteTriplets:
-		AdjustedTick = UnadjustedTick - MidiConstants::kTicksPerQuarterNoteInt * Offset / 6;
+		AdjustedTick = UnadjustedTick - GTicksPerQuarterNoteInt * Offset / 6;
 		break;
 	case EMusicTimeSpanOffsetUnits::EighthNoteTriplets:
-		AdjustedTick = UnadjustedTick - MidiConstants::kTicksPerQuarterNoteInt * Offset / 3;
+		AdjustedTick = UnadjustedTick - GTicksPerQuarterNoteInt * Offset / 3;
 		break;
 	case EMusicTimeSpanOffsetUnits::QuarterNoteTriplets:
-		AdjustedTick = UnadjustedTick - MidiConstants::kTicksPerQuarterNoteInt * Offset * 2 / 3;
+		AdjustedTick = UnadjustedTick - GTicksPerQuarterNoteInt * Offset * 2 / 3;
 		break;
 	case EMusicTimeSpanOffsetUnits::HalfNoteTriplets:
-		AdjustedTick = UnadjustedTick - MidiConstants::kTicksPerQuarterNoteInt * Offset * 4 / 3;
+		AdjustedTick = UnadjustedTick - GTicksPerQuarterNoteInt * Offset * 4 / 3;
 		break;
 	default:
 		checkNoEntry();
diff --git a/Engine/Plugins/Runtime/Harmonix/Source/HarmonixMidi/Private/HarmonixMidi/SongMapReceiver.cpp b/Engine/Plugins/Runtime/Harmonix/Source/HarmonixMidi/Private/HarmonixMidi/SongMapReceiver.cpp
--- a/Engine/Plugins/Runtime/Harmonix/Source/HarmonixMidi/Private/HarmonixMidi/SongMapReceiver.cpp
+++ b/Engine/Plugins/Runtime/Harmonix/Source/HarmonixMidi/Private/HarmonixMidi/SongMapReceiver.cpp
@@ -57,8 +57,10 @@ void FSongMapReceiver::OnMidiMessage(int32 Tick, uint8 Status, uint8 Data1, uint
 	{
 		return;
 	}
+	
+	using namespace Harmonix::Midi::Constants;
 
-	if (MidiConstants::GetType(Status) != MidiConstants::kNoteOn)
+	if (GetType(Status) != GNoteOn)
 	{
 		return;
 	}
@@ -73,9 +75,9 @@ void FSongMapReceiver::OnMidiMessage(int32 Tick, uint8 Status, uint8 Data1, uint
 	if ((LastBeatTick != -1) && !bHaveBeatFailure)
 	{
 		// Don't allow beat tracks faster than 16th note denominators.
-		if ((Tick - LastBeatTick) < MidiConstants::kTicksPerQuarterNoteInt / 4)
+		if ((Tick - LastBeatTick) < GTicksPerQuarterNoteInt / 4)
 		{
-			_LOG_SONGMAP_ERROR("Beat track cannot be faster than 16th notes; beats are less than %d ticks apart at %s", MidiConstants::kTicksPerQuarterNoteInt / 4, *FmtTick(Tick));
+			_LOG_SONGMAP_ERROR("Beat track cannot be faster than 16th notes; beats are less than %d ticks apart at %s", GTicksPerQuarterNoteInt / 4, *FmtTick(Tick));
 			bHaveBeatFailure = true;
 		}
 
@@ -102,11 +104,11 @@ void FSongMapReceiver::OnMidiMessage(int32 Tick, uint8 Status, uint8 Data1, uint
 void FSongMapReceiver::OnText(int32 Tick, const FString& Str, uint8 Type)
 {
 	// we only care about track name
-	if (Type == MidiConstants::kMeta_TrackName)
+	if (Type == Harmonix::Midi::Constants::GMeta_TrackName)
 	{
 		OnTrackName(Str);
 	}
-	else if (Type == MidiConstants::kMeta_Text)
+	else if (Type == Harmonix::Midi::Constants::GMeta_Text)
 	{
 		switch (CurrentTrack)
 		{
@@ -130,7 +132,6 @@ void FSongMapReceiver::OnTempo(int32 Tick, int32 Tempo)
 
 void FSongMapReceiver::OnTimeSignature(int32 Tick, int32 Numerator, int32 Denominator, bool FailOnError)
 {
-	FMusicTimestamp Timestamp;
 	check(Tick == 0 || SongMaps->BarMap.GetNumTimeSignaturePoints() > 0);
 	int32 BarIndex = SongMaps->BarMap.TickToBarIncludingCountIn(Tick);
 
diff --git a/Engine/Plugins/Runtime/Harmonix/Source/HarmonixMidi/Private/HarmonixMidi/TempoMap.cpp b/Engine/Plugins/Runtime/Harmonix/Source/HarmonixMidi/Private/HarmonixMidi/TempoMap.cpp
--- a/Engine/Plugins/Runtime/Harmonix/Source/HarmonixMidi/Private/HarmonixMidi/TempoMap.cpp
+++ b/Engine/Plugins/Runtime/Harmonix/Source/HarmonixMidi/Private/HarmonixMidi/TempoMap.cpp
@@ -35,7 +35,7 @@ int32 FTempoMap::GetMicrosecondsPerQuarterNoteAtTick(int32 Tick) const
 
 float FTempoMap::GetTempoAtTick(int32 Tick) const
 {
-	return MidiConstants::MidiTempoToBPM(GetMicrosecondsPerQuarterNoteAtTick(Tick));
+	return Harmonix::Midi::Constants::MidiTempoToBPM(GetMicrosecondsPerQuarterNoteAtTick(Tick));
 }
 
 void FTempoMap::Empty()
diff --git a/Engine/Plugins/Runtime/Harmonix/Source/HarmonixMidi/Public/HarmonixMidi/BarMap.h b/Engine/Plugins/Runtime/Harmonix/Source/HarmonixMidi/Public/HarmonixMidi/BarMap.h
--- a/Engine/Plugins/Runtime/Harmonix/Source/HarmonixMidi/Public/HarmonixMidi/BarMap.h
+++ b/Engine/Plugins/Runtime/Harmonix/Source/HarmonixMidi/Public/HarmonixMidi/BarMap.h
@@ -143,7 +143,7 @@ struct HARMONIXMIDI_API FBarMap
 
 public:
 	FBarMap()
-		: TicksPerQuarterNote(MidiConstants::kTicksPerQuarterNoteInt)
+		: TicksPerQuarterNote(Harmonix::Midi::Constants::GTicksPerQuarterNoteInt)
 	{}
 	friend bool operator==(const FBarMap& Left, const FBarMap& Right);
 
diff --git a/Engine/Plugins/Runtime/Harmonix/Source/HarmonixMidi/Public/HarmonixMidi/BeatMap.h b/Engine/Plugins/Runtime/Harmonix/Source/HarmonixMidi/Public/HarmonixMidi/BeatMap.h
--- a/Engine/Plugins/Runtime/Harmonix/Source/HarmonixMidi/Public/HarmonixMidi/BeatMap.h
+++ b/Engine/Plugins/Runtime/Harmonix/Source/HarmonixMidi/Public/HarmonixMidi/BeatMap.h
@@ -99,7 +99,7 @@ struct HARMONIXMIDI_API FBeatMap
 
 public:
 	FBeatMap()
-		: TicksPerQuarterNote(MidiConstants::kTicksPerQuarterNoteInt)
+		: TicksPerQuarterNote(Harmonix::Midi::Constants::GTicksPerQuarterNoteInt)
 	{}
 	friend bool operator==(const FBeatMap& Left, const FBeatMap& Right);
 
diff --git a/Engine/Plugins/Runtime/Harmonix/Source/HarmonixMidi/Public/HarmonixMidi/ChordMap.h b/Engine/Plugins/Runtime/Harmonix/Source/HarmonixMidi/Public/HarmonixMidi/ChordMap.h
--- a/Engine/Plugins/Runtime/Harmonix/Source/HarmonixMidi/Public/HarmonixMidi/ChordMap.h
+++ b/Engine/Plugins/Runtime/Harmonix/Source/HarmonixMidi/Public/HarmonixMidi/ChordMap.h
@@ -47,7 +47,7 @@ struct HARMONIXMIDI_API FChordProgressionMap
 
 public:
 	FChordProgressionMap()
-		: TicksPerQuarterNote(MidiConstants::kTicksPerQuarterNoteInt)
+		: TicksPerQuarterNote(Harmonix::Midi::Constants::GTicksPerQuarterNoteInt)
 	{}
 	friend bool operator==(const FChordProgressionMap& Left, const FChordProgressionMap& Right);
 
diff --git a/Engine/Plugins/Runtime/Harmonix/Source/HarmonixMidi/Public/HarmonixMidi/MidiConstants.h b/Engine/Plugins/Runtime/Harmonix/Source/HarmonixMidi/Public/HarmonixMidi/MidiConstants.h
--- a/Engine/Plugins/Runtime/Harmonix/Source/HarmonixMidi/Public/HarmonixMidi/MidiConstants.h
+++ b/Engine/Plugins/Runtime/Harmonix/Source/HarmonixMidi/Public/HarmonixMidi/MidiConstants.h
@@ -1,31 +1,19 @@
 // Copyright Epic Games, Inc. All Rights Reserved.
+
 #pragma once
+
 #include "Logging/LogMacros.h"
 
-DECLARE_LOG_CATEGORY_EXTERN(LogMidi, Log, All);
+HARMONIXMIDI_API DECLARE_LOG_CATEGORY_EXTERN(LogMidi, Log, All);
 
-class HARMONIXMIDI_API MidiConstants
+namespace Harmonix::Midi::Constants
 {
-public:
 	enum class EMidiTextEventEncoding
 	{
 		Latin1,
 		UTF8
 	};
 
-	// helper enum for common time subdivisions
-	enum class ESubdivisions
-	{
-		Downbeat      = 1,
-		Half          = 2,
-		ThreeOverFour = 3,
-		Quarter       = 4,
-		Eighth        = 8,
-		Triplet       = 12,
-		Sixteenth     = 16,
-		ThirtySecond  = 32,
-	};
-
 	/////////////////////////////////////////////////////////////////////////////// 
 	// Constants and functions for MIDI 1.0 standard
 	//
@@ -33,97 +21,96 @@ public:
 	//  http://www.midi.org/about-midi/table1.shtml
 	//  http://www.borg.com/~jglatt/tech/midispec.htm
 
-	static const uint8 kNumChannels    = 16;
-	static const uint8 kNotesPerOctave = 12;
-	static const uint8 kMinNote        = 0;
-	static const uint8 kMaxNote        = 127;
-	static const uint8 kMaxNumNotes    = 128;
-	static const uint8 kMiddleC        = 60;
-	static const uint8 kMinVelocity    = 0;
-	static const uint8 kMaxVelocity    = 127;
+	inline constexpr uint8 GNumChannels    = 16;
+	inline constexpr uint8 GNotesPerOctave = 12;
+	inline constexpr uint8 GMinNote        = 0;
+	inline constexpr uint8 GMaxNote        = 127;
+	inline constexpr uint8 GMaxNumNotes    = 128;
+	inline constexpr uint8 GMinVelocity    = 0;
+	inline constexpr uint8 GMaxVelocity    = 127;
 
-	static const int32   kTicksPerQuarterNoteInt  = 960;
-	static constexpr float kTicksPerQuarterNote = 960.0f;
-	static constexpr float kQuarterNotesPerTick = 1.0f / kTicksPerQuarterNote;
+	inline constexpr int32   GTicksPerQuarterNoteInt  = 960;
+	inline constexpr float GTicksPerQuarterNote = 960.0f;
+	inline constexpr float GQuarterNotesPerTick = 1.0f / GTicksPerQuarterNote;
 
 	//////////////////////////////////////////////////////////////////////////
 	// Constants for the SMF (Standard Midi File) format.
 
 	// Codes for special handling of System Exclusive messages:
-	static const uint8 kFile_Escape      = 0xf7;
-	static const uint8 kFile_SysEx       = 0xf0; 
+	inline constexpr uint8 GFile_Escape      = 0xf7;
+	inline constexpr uint8 GFile_SysEx       = 0xf0; 
 
 	// Status code for SMF meta-events:
-	static const uint8 kFile_Meta        = 0xff;
+	inline constexpr uint8 GFile_Meta        = 0xff;
 
 	// Meta-event IDs:
-	static const uint8 kMeta_Text            = 0x01;
-	static const uint8 kMeta_Copyright       = 0x02;
-	static const uint8 kMeta_TrackName       = 0x03;
-	static const uint8 kMeta_InstrumentName  = 0x04;
-	static const uint8 kMeta_Lyric           = 0x05;
-	static const uint8 kMeta_Marker          = 0x06;
-	static const uint8 kMeta_CuePoint        = 0x07;
-	static const uint8 kMeta_ChannelPrefix   = 0x20;
-	static const uint8 kMeta_Port            = 0x21; // obsolete
-	static const uint8 kMeta_EndOfTrack      = 0x2f;
-	static const uint8 kMeta_Tempo           = 0x51;
-	static const uint8 kMeta_SMPTE           = 0x54;
-	static const uint8 kMeta_TimeSig         = 0x58;
-	static const uint8 kMeta_KeySig          = 0x59;
-	static const uint8 kMeta_Special         = 0x7f;
+	inline constexpr uint8 GMeta_Text            = 0x01;
+	inline constexpr uint8 GMeta_Copyright       = 0x02;
+	inline constexpr uint8 GMeta_TrackName       = 0x03;
+	inline constexpr uint8 GMeta_InstrumentName  = 0x04;
+	inline constexpr uint8 GMeta_Lyric           = 0x05;
+	inline constexpr uint8 GMeta_Marker          = 0x06;
+	inline constexpr uint8 GMeta_CuePoint        = 0x07;
+	inline constexpr uint8 GMeta_ChannelPrefix   = 0x20;
+	inline constexpr uint8 GMeta_Port            = 0x21; // obsolete
+	inline constexpr uint8 GMeta_EndOfTrack      = 0x2f;
+	inline constexpr uint8 GMeta_Tempo           = 0x51;
+	inline constexpr uint8 GMeta_SMPTE           = 0x54;
+	inline constexpr uint8 GMeta_TimeSig         = 0x58;
+	inline constexpr uint8 GMeta_KeySig          = 0x59;
+	inline constexpr uint8 GMeta_Special         = 0x7f;
 
 	// if FMidiMsg type is "Runtime" these are the possible status byte.
 	// note: runtime messages are not serialized!
-	static const uint8 kRuntimeAllNotesOffStatus = 0x01;  // Receiver should feel free to allow notes to "adsr release"
-	static const uint8 kRuntimeAllNotesKillStatus = 0x02; // Receiver is expected to immediately free resources associated with any sounding notes and stop playing ASAP!
+	inline constexpr uint8 GRuntimeAllNotesOffStatus = 0x01;  // Receiver should feel free to allow notes to "adsr release"
+	inline constexpr uint8 GRuntimeAllNotesKillStatus = 0x02; // Receiver is expected to immediately free resources associated with any sounding notes and stop playing ASAP!
 
-	static FString GetTextTypeName(uint8 TextType);
+	HARMONIXMIDI_API FString GetTextTypeName(uint8 TextType);
 
 	///////////////////////////////////////////////////////////////////////////////
 	// Status bytes indicate the type of message (3 bytes long except where noted).
 
 	// Masks for disassembling MIDI Status bytes
-	static const uint8 kStatusBitMask   = 0x80; // hi bit always set on Status byte
-	static const uint8 kMessageTypeMask = 0xf0; // mask for message type
-	static const uint8 kChannelMask     = 0x0f; // mask for channel
-	static const uint8 kRealTimeMask    = 0xf8; // indicates real-time Status message
-
-	static const uint8 kNoteOff  = 0x80; // Note Off
-	static const uint8 kNoteOn   = 0x90; // Note On
-	static const uint8 kPolyPres = 0xa0; // Polyphonic Key Pressure 
-	static const uint8 kControl  = 0xb0; // Control Change
-	static const uint8 kProgram  = 0xc0; // Program Change (2 bytes)
-	static const uint8 kChanPres = 0xd0; // Channel Pressure (2 bytes)
-	static const uint8 kPitch    = 0xe0; // Pitch Wheel Change
-	static const uint8 kSystem   = 0xf0; // System
-
-
-	static inline bool  IsStatus(uint8 Byte)     { return (Byte & kStatusBitMask) > 0;  }
-	static inline uint8 GetChannel(uint8 Status) { check(IsStatus(Status));	check(!IsSystem(Status)); return Status & kChannelMask; }
-	static inline uint8 GetType(uint8 Status)    { check(IsStatus(Status)); return Status & kMessageTypeMask; }
-
-	static inline bool  IsNoteOff(uint8 Status)  { return GetType(Status) == kNoteOff;  }
-	static inline bool  IsNoteOn(uint8 Status)   { return GetType(Status) == kNoteOn;   }
-	static inline bool  IsPolyPres(uint8 Status) { return GetType(Status) == kPolyPres; }
-	static inline bool  IsControl(uint8 Status)  { return GetType(Status) == kControl;  }
-	static inline bool  IsProgram(uint8 Status)  { return GetType(Status) == kProgram;  }
-	static inline bool  IsChanPres(uint8 Status) { return GetType(Status) == kChanPres; }
-	static inline bool  IsPitch(uint8 Status)    { return GetType(Status) == kPitch;    }
-	static inline bool  IsSystem(uint8 Status)   { return GetType(Status) == kSystem;   }
+	inline constexpr uint8 GStatusBitMask   = 0x80; // hi bit always set on Status byte
+	inline constexpr uint8 GMessageTypeMask = 0xf0; // mask for message type
+	inline constexpr uint8 GChannelMask     = 0x0f; // mask for channel
+	inline constexpr uint8 GRealTimeMask    = 0xf8; // indicates real-time Status message
+
+	inline constexpr uint8 GNoteOff  = 0x80; // Note Off
+	inline constexpr uint8 GNoteOn   = 0x90; // Note On
+	inline constexpr uint8 GPolyPres = 0xa0; // Polyphonic Key Pressure 
+	inline constexpr uint8 GControl  = 0xb0; // Control Change
+	inline constexpr uint8 GProgram  = 0xc0; // Program Change (2 bytes)
+	inline constexpr uint8 GChanPres = 0xd0; // Channel Pressure (2 bytes)
+	inline constexpr uint8 GPitch    = 0xe0; // Pitch Wheel Change
+	inline constexpr uint8 GSystem   = 0xf0; // System
+
+	HARMONIXMIDI_API inline bool  IsStatus(uint8 Byte)     { return (Byte & GStatusBitMask) > 0;  }
+	HARMONIXMIDI_API inline uint8 GetType(uint8 Status)    { check(IsStatus(Status)); return Status & GMessageTypeMask; }
+
+	HARMONIXMIDI_API inline bool  IsNoteOff(uint8 Status)  { return GetType(Status) == GNoteOff;  }
+	HARMONIXMIDI_API inline bool  IsNoteOn(uint8 Status)   { return GetType(Status) == GNoteOn;   }
+	HARMONIXMIDI_API inline bool  IsPolyPres(uint8 Status) { return GetType(Status) == GPolyPres; }
+	HARMONIXMIDI_API inline bool  IsControl(uint8 Status)  { return GetType(Status) == GControl;  }
+	HARMONIXMIDI_API inline bool  IsProgram(uint8 Status)  { return GetType(Status) == GProgram;  }
+	HARMONIXMIDI_API inline bool  IsChanPres(uint8 Status) { return GetType(Status) == GChanPres; }
+	HARMONIXMIDI_API inline bool  IsPitch(uint8 Status)    { return GetType(Status) == GPitch;    }
+	HARMONIXMIDI_API inline bool  IsSystem(uint8 Status)   { return GetType(Status) == GSystem;   }
+
+	HARMONIXMIDI_API inline uint8 GetChannel(uint8 Status) { check(IsStatus(Status));	check(!IsSystem(Status)); return Status & GChannelMask; }
 
 	// turns MIDI style tempo (microseconds per quarter note) into BPM...
-	static inline float MidiTempoToBPM(int32 UsPerQuarterNote) { return UsPerQuarterNote== 0 ? 0 : 60000000.0f/(float)UsPerQuarterNote; }
-	static inline int32 BPMToMidiTempo(float Bpm)       { return Bpm == 0.0 ? 0 : (int32)(60000000.0f / (float)Bpm); }
+	HARMONIXMIDI_API inline float MidiTempoToBPM(int32 UsPerQuarterNote) { return UsPerQuarterNote== 0 ? 0 : 60000000.0f/(float)UsPerQuarterNote; }
+	HARMONIXMIDI_API inline int32 BPMToMidiTempo(float Bpm)       { return Bpm == 0.0 ? 0 : (int32)(60000000.0f / (float)Bpm); }
 
-	static float RoundToStandardBeatPrecision(float InBeat, int TimeSIgnatureDenominator);
+	HARMONIXMIDI_API float RoundToStandardBeatPrecision(float InBeat, int TimeSignatureDenominator);
 
 	///////////////////////////////////////////////////////////////////////////////
 	// Controllers
 	// 
 	// In a Control Change message, data byte 1 is the controller ID, and data
 	// byte 2 is the controller value.
-	enum EControllerID
+	enum class EControllerID : uint8
 	{ // from midi specification...
 	   // NOTE: If you add any here also add them to the GetControllerName function!
 		BankSelection = 0,
@@ -198,42 +185,36 @@ public:
 		AllNotesOff = 123
 	};
 
-	static FString GetControllerName(EControllerID ControllerId);
-	static void    GetControllerNames(TArray<FString>& Names);
+	HARMONIXMIDI_API FString GetControllerName(EControllerID ControllerId);
+	HARMONIXMIDI_API void GetControllerNames(TArray<FString>& Names);
 
 	// allows choice between various naming conventions for enharmonic notes
-	enum class ENoteNameEnharmonicStyle
+	enum class ENoteNameEnharmonicStyle : uint8
 	{
 		Sharp,
 		Flat,
 		SharpAndFlat
 	};
 
-	static const int8 GetNoteNumberFromNoteName(const char* InName);
-	static const char* GetNoteNameFromNoteNumber(uint8 MidiNoteNumber, ENoteNameEnharmonicStyle Style = ENoteNameEnharmonicStyle::SharpAndFlat);
-	static int8 GetNoteOctaveFromNoteNumber(uint8 MidiNoteNumber);
-
-	///////////////////////////////////////////////////////////////////////////////
-	// RPNs (registered paramater numbers)
-	static const uint8 kRPN_BendRangeC   = 0; 
-	static const uint8 kRPN_BendRangeF   = 0; 
-
+	HARMONIXMIDI_API int8 GetNoteNumberFromNoteName(const char* InName);
+	HARMONIXMIDI_API const char* GetNoteNameFromNoteNumber(uint8 MidiNoteNumber, ENoteNameEnharmonicStyle Style = ENoteNameEnharmonicStyle::SharpAndFlat);
+	HARMONIXMIDI_API int8 GetNoteOctaveFromNoteNumber(uint8 MidiNoteNumber);
 
 	///////////////////////////////////////////////////////////////////////////////
 	// System messages:
 
 	// System Common:
-	static const uint8 kSys_MTC         = 0xf1; // 2 byte message
-	static const uint8 kSys_SongPos     = 0xf2; // 3 byte message
-	static const uint8 kSys_SongSelect  = 0xf3; // 2 byte message
-	static const uint8 kSys_TuneRequest = 0xf6; // 1 byte message
-	static const uint8 kSys_EOX         = 0xf7; // 1 byte message
-
+	inline constexpr uint8 GSys_Mtc         = 0xf1; // 2 byte message
+	inline constexpr uint8 GSys_SongPos     = 0xf2; // 3 byte message
+	inline constexpr uint8 GSys_SongSelect  = 0xf3; // 2 byte message
+	inline constexpr uint8 GSys_TuneRequest = 0xf6; // 1 byte message
+	inline constexpr uint8 GSys_Eox         = 0xf7; // 1 byte message
+	
 	// System RealTime messages (all 1 byte)
-	static const uint8 kSys_TimingClock = 0xf8;
-	static const uint8 kSys_Start       = 0xfa;
-	static const uint8 kSys_Continue    = 0xfb;
-	static const uint8 kSys_Stop        = 0xfc;
-	static const uint8 kSys_ActiveSense = 0xfe;
-	static const uint8 kSys_Reset       = 0xff;
+	inline constexpr uint8 GSys_TimingClock = 0xf8;
+	inline constexpr uint8 GSys_Start       = 0xfa;
+	inline constexpr uint8 GSys_Continue    = 0xfb;
+	inline constexpr uint8 GSys_Stop        = 0xfc;
+	inline constexpr uint8 GSys_ActiveSense = 0xfe;
+	inline constexpr uint8 GSys_Reset       = 0xff;
 };
diff --git a/Engine/Plugins/Runtime/Harmonix/Source/HarmonixMidi/Public/HarmonixMidi/MidiFile.h b/Engine/Plugins/Runtime/Harmonix/Source/HarmonixMidi/Public/HarmonixMidi/MidiFile.h
--- a/Engine/Plugins/Runtime/Harmonix/Source/HarmonixMidi/Public/HarmonixMidi/MidiFile.h
+++ b/Engine/Plugins/Runtime/Harmonix/Source/HarmonixMidi/Public/HarmonixMidi/MidiFile.h
@@ -30,7 +30,7 @@ struct HARMONIXMIDI_API FMidiFileData
 	GENERATED_BODY()
 
 	FMidiFileData()
-		: TicksPerQuarterNote(MidiConstants::kTicksPerQuarterNoteInt)
+		: TicksPerQuarterNote(Harmonix::Midi::Constants::GTicksPerQuarterNoteInt)
 		, LastEventTick(0)
 	{}
 	UPROPERTY(BlueprintReadOnly, Category = "MidiFile")
@@ -102,14 +102,23 @@ public:
 	virtual void PostInitProperties() override;
 
 	/** A method for importing a standard midi file, with the option of providing a pointer to an FSongMaps instance that will be populated during the load. */
-	void LoadStdMidiFile(const FString& FilePath, int32 DesiredTicksPerQuarterNote = MidiConstants::kTicksPerQuarterNoteInt,
-		MidiConstants::EMidiTextEventEncoding InTextEncoding = MidiConstants::EMidiTextEventEncoding::UTF8);
+	void LoadStdMidiFile(
+		const FString& FilePath,
+		int32 DesiredTicksPerQuarterNote = Harmonix::Midi::Constants::GTicksPerQuarterNoteInt,
+		Harmonix::Midi::Constants::EMidiTextEventEncoding InTextEncoding = Harmonix::Midi::Constants::EMidiTextEventEncoding::UTF8);
 	/** A method for importing a standard midi file, with the option of providing a pointer to an FSongMaps instance that will be populated during the load. */
-	void LoadStdMidiFile(void* Buffer, int32 BufferSize, const FString& FileName, int32 DesiredTicksPerQuarterNote = MidiConstants::kTicksPerQuarterNoteInt,
-		MidiConstants::EMidiTextEventEncoding InTextEncoding = MidiConstants::EMidiTextEventEncoding::UTF8);
+	void LoadStdMidiFile(
+		void* Buffer,
+		int32 BufferSize,
+		const FString& FileName,
+		int32 DesiredTicksPerQuarterNote = Harmonix::Midi::Constants::GTicksPerQuarterNoteInt,
+		Harmonix::Midi::Constants::EMidiTextEventEncoding InTextEncoding = Harmonix::Midi::Constants::EMidiTextEventEncoding::UTF8);
 	/** A method for importing a standard midi file, with the option of providing a pointer to an FSongMaps instance that will be populated during the load. */
-	void LoadStdMidiFile(TSharedPtr<FArchive> Archive, const FString& Filename, int32 DesiredTicksPerQuarterNote = MidiConstants::kTicksPerQuarterNoteInt,
-		MidiConstants::EMidiTextEventEncoding InTextEncoding = MidiConstants::EMidiTextEventEncoding::UTF8);
+	void LoadStdMidiFile(
+		TSharedPtr<FArchive> Archive,
+		const FString& Filename,
+		int32 DesiredTicksPerQuarterNote = Harmonix::Midi::Constants::GTicksPerQuarterNoteInt,
+		Harmonix::Midi::Constants::EMidiTextEventEncoding InTextEncoding = Harmonix::Midi::Constants::EMidiTextEventEncoding::UTF8);
 
 	/** A method for exporting the midi track data to a standard midi file. */
 	void SaveStdMidiFile(const FString& FilePath);
diff --git a/Engine/Plugins/Runtime/Harmonix/Source/HarmonixMidi/Public/HarmonixMidi/MidiMsg.h b/Engine/Plugins/Runtime/Harmonix/Source/HarmonixMidi/Public/HarmonixMidi/MidiMsg.h
--- a/Engine/Plugins/Runtime/Harmonix/Source/HarmonixMidi/Public/HarmonixMidi/MidiMsg.h
+++ b/Engine/Plugins/Runtime/Harmonix/Source/HarmonixMidi/Public/HarmonixMidi/MidiMsg.h
@@ -41,24 +41,31 @@ public:
 	bool  IsRuntime() const { return Type == EType::Runtime; }
 	bool  IsText() const    { return Type == EType::Text;    }
 	bool  IsNoteMessage() const { return IsNoteOn() || IsNoteOff() || IsAllNotesOff() || IsAllNotesKill(); }
-	bool  IsNoteOn() const  { return Type == EType::Std && MidiConstants::IsNoteOn(Status);  }
-	bool  IsNoteOff() const { return Type == EType::Std && (MidiConstants::IsNoteOff(Status) || (MidiConstants::IsNoteOn(Status) && Data2 == 0)); }
-	bool  IsAllNotesOff() const { return Type == EType::Runtime && Status == MidiConstants::kRuntimeAllNotesOffStatus; }
-	bool  IsAllNotesKill() const { return Type == EType::Runtime && Status == MidiConstants::kRuntimeAllNotesKillStatus; }
-	bool  IsControlChange() const { return IsStd() && MidiConstants::IsControl(Status); }
+	bool  IsNoteOn() const  { return Type == EType::Std && Harmonix::Midi::Constants::IsNoteOn(Status);  }
+	bool  IsNoteOff() const
+	{
+		return Type == EType::Std && (Harmonix::Midi::Constants::IsNoteOff(Status) || (Harmonix::Midi::Constants::IsNoteOn(Status) && Data2 == 0));
+	}
+	bool  IsAllNotesOff() const { return Type == EType::Runtime && Status == Harmonix::Midi::Constants::GRuntimeAllNotesOffStatus; }
+	bool  IsAllNotesKill() const { return Type == EType::Runtime && Status == Harmonix::Midi::Constants::GRuntimeAllNotesKillStatus; }
+	bool  IsControlChange() const { return IsStd() && Harmonix::Midi::Constants::IsControl(Status); }
 
 	/** Construct a standard (std) midi message */
 	FMidiMsg(uint8 InStatus, uint8 InData1, uint8 InData2);
 	uint8 GetStdStatus()     const { check(Type == EType::Std); return Status; }
 	uint8 GetStdData1()      const { check(Type == EType::Std); return Data1; }
 	uint8 GetStdData2()      const { check(Type == EType::Std); return Data2; }
-	uint8 GetStdChannel()    const { check(Type == EType::Std); return MidiConstants::GetChannel(Status); }
-	uint8 GetStdStatusType() const { check(Type == EType::Std); return MidiConstants::GetType(Status); }
+	uint8 GetStdChannel()    const { check(Type == EType::Std); return Harmonix::Midi::Constants::GetChannel(Status); }
+	uint8 GetStdStatusType() const { check(Type == EType::Std); return Harmonix::Midi::Constants::GetType(Status); }
 	float GetPitchBendFromData() const;
 
 	/** Construct a midi tempo message */
 	FMidiMsg(int32 MicrosecPerQuarterNote);
-	int32 GetMicrosecPerQuarterNote() const { check(Type == EType::Tempo); return int32((uint32)MicsPerQuarterNoteH << 16 | (uint32)MicsPerQuarterNoteL); }
+	int32 GetMicrosecPerQuarterNote() const
+	{
+		check(Type == EType::Tempo);
+		return int32((uint32)MicsPerQuarterNoteH << 16 | (uint32)MicsPerQuarterNoteL);
+	}
 
 	/** Construct a midi time signature message */
 	FMidiMsg(uint8 Numerator, uint8 Denominator);
@@ -91,14 +98,14 @@ public:
 	{
 		FMidiMsg NewMsg;
 		NewMsg.Type = EType::Runtime;
-		NewMsg.Status = MidiConstants::kRuntimeAllNotesOffStatus;
+		NewMsg.Status = Harmonix::Midi::Constants::GRuntimeAllNotesOffStatus;
 		return NewMsg;
 	}
 	static FMidiMsg CreateAllNotesKill()
 	{
 		FMidiMsg NewMsg;
 		NewMsg.Type = EType::Runtime;
-		NewMsg.Status = MidiConstants::kRuntimeAllNotesKillStatus;
+		NewMsg.Status = Harmonix::Midi::Constants::GRuntimeAllNotesKillStatus;
 		return NewMsg;
 	}
 
diff --git a/Engine/Plugins/Runtime/Harmonix/Source/HarmonixMidi/Public/HarmonixMidi/MidiReader.h b/Engine/Plugins/Runtime/Harmonix/Source/HarmonixMidi/Public/HarmonixMidi/MidiReader.h
--- a/Engine/Plugins/Runtime/Harmonix/Source/HarmonixMidi/Public/HarmonixMidi/MidiReader.h
+++ b/Engine/Plugins/Runtime/Harmonix/Source/HarmonixMidi/Public/HarmonixMidi/MidiReader.h
@@ -43,16 +43,32 @@ public:
 class FStdMidiFileReader : public  IMidiReader
 {
 public:
-
-	FStdMidiFileReader(TSharedPtr<FArchive> Archive, const FString& InFilename, IMidiReceiver* InReceiver, 
-			int32 TicksPerQuarterNote = MidiConstants::kTicksPerQuarterNoteInt, MidiConstants::EMidiTextEventEncoding InTextEncoding = MidiConstants::EMidiTextEventEncoding::UTF8);
-	FStdMidiFileReader(const FString& FilePath, IMidiReceiver* Receiver,
-			int32 TicksPerQuarterNote = MidiConstants::kTicksPerQuarterNoteInt, MidiConstants::EMidiTextEventEncoding InTextEncoding = MidiConstants::EMidiTextEventEncoding::UTF8);
-	FStdMidiFileReader(void* Buffer, int32 BufferSize, const FString& FileName, IMidiReceiver* Receiver,
-			int32 TicksPerQuarterNote = MidiConstants::kTicksPerQuarterNoteInt, MidiConstants::EMidiTextEventEncoding InTextEncoding = MidiConstants::EMidiTextEventEncoding::UTF8);
-
-	static bool ReadStdMidiFileForReceiver(void* Buffer, int32 BufferSize, const FString& FileName, IMidiReceiver* Receiver,
-			int32 TicksPerQuarterNote = MidiConstants::kTicksPerQuarterNoteInt, MidiConstants::EMidiTextEventEncoding InTextEncoding = MidiConstants::EMidiTextEventEncoding::UTF8);
+	FStdMidiFileReader(
+		TSharedPtr<FArchive> Archive,
+		const FString& InFilename,
+		IMidiReceiver* InReceiver, 
+		int32 TicksPerQuarterNote = Harmonix::Midi::Constants::GTicksPerQuarterNoteInt,
+		Harmonix::Midi::Constants::EMidiTextEventEncoding InTextEncoding = Harmonix::Midi::Constants::EMidiTextEventEncoding::UTF8);
+	FStdMidiFileReader(
+		const FString& FilePath,
+		IMidiReceiver* Receiver,
+		int32 TicksPerQuarterNote = Harmonix::Midi::Constants::GTicksPerQuarterNoteInt,
+		Harmonix::Midi::Constants::EMidiTextEventEncoding InTextEncoding = Harmonix::Midi::Constants::EMidiTextEventEncoding::UTF8);
+	FStdMidiFileReader(
+		void* Buffer,
+		int32 BufferSize,
+		const FString& FileName,
+		IMidiReceiver* Receiver,
+		int32 TicksPerQuarterNote = Harmonix::Midi::Constants::GTicksPerQuarterNoteInt,
+		Harmonix::Midi::Constants::EMidiTextEventEncoding InTextEncoding = Harmonix::Midi::Constants::EMidiTextEventEncoding::UTF8);
+
+	static bool ReadStdMidiFileForReceiver(
+		void* Buffer,
+		int32 BufferSize,
+		const FString& FileName,
+		IMidiReceiver* Receiver,
+		int32 TicksPerQuarterNote = Harmonix::Midi::Constants::GTicksPerQuarterNoteInt,
+		Harmonix::Midi::Constants::EMidiTextEventEncoding InTextEncoding = Harmonix::Midi::Constants::EMidiTextEventEncoding::UTF8);
 
 	virtual ~FStdMidiFileReader() {}
 
@@ -105,10 +121,10 @@ private:
 	TSharedPtr<FArchiveFileReaderGeneric> FileArchive;
 
 	IMidiReceiver* Receiver = nullptr;
-	MidiConstants::EMidiTextEventEncoding TextEncoding = MidiConstants::EMidiTextEventEncoding::Latin1;
+	Harmonix::Midi::Constants::EMidiTextEventEncoding TextEncoding = Harmonix::Midi::Constants::EMidiTextEventEncoding::Latin1;
 	EState State = EState::Start;
 
-	int32 DestinationTicksPerQuarterNote = MidiConstants::kTicksPerQuarterNoteInt;
+	int32 DestinationTicksPerQuarterNote = Harmonix::Midi::Constants::GTicksPerQuarterNoteInt;
 	float TickConversionFactor = 1.0f; // factor to convert from file's tick-per-quarternote to kTicksPerQuarterNote 
 
 	int16 Format = 1;             // Format 0 or Format 1 standard midi file?
@@ -144,16 +160,17 @@ private:
 	{
 		inline int32 MidiRank(uint8 Status) const
 		{
-			switch (Status & MidiConstants::kMessageTypeMask)
+			using namespace Harmonix::Midi::Constants;
+			switch (Status & GMessageTypeMask)
 			{
-			case MidiConstants::kNoteOff:  return 1;
-			case MidiConstants::kControl:  return 2;
-			case MidiConstants::kProgram:  return 3;
-			case MidiConstants::kChanPres: return 4;
-			case MidiConstants::kPitch:    return 5;
-			case MidiConstants::kPolyPres: return 6;
-			case MidiConstants::kNoteOn:   return 7;
-			default:             return 8;
+			case GNoteOff:  return 1;
+			case GControl:  return 2;
+			case GProgram:  return 3;
+			case GChanPres: return 4;
+			case GPitch:    return 5;
+			case GPolyPres: return 6;
+			case GNoteOn:   return 7;
+			default:        return 8;
 			}
 		}
 		bool operator()(const FStdMidiFileReader::FRawMidiMsg& lhs, const FStdMidiFileReader::FRawMidiMsg& rhs) const
diff --git a/Engine/Plugins/Runtime/Harmonix/Source/HarmonixMidi/Public/HarmonixMidi/MidiWriter.h b/Engine/Plugins/Runtime/Harmonix/Source/HarmonixMidi/Public/HarmonixMidi/MidiWriter.h
--- a/Engine/Plugins/Runtime/Harmonix/Source/HarmonixMidi/Public/HarmonixMidi/MidiWriter.h
+++ b/Engine/Plugins/Runtime/Harmonix/Source/HarmonixMidi/Public/HarmonixMidi/MidiWriter.h
@@ -19,7 +19,7 @@
 class HARMONIXMIDI_API FMidiWriter
 {
 public:
-	FMidiWriter(FArchive& OutputArhive, int32 TicksPerQuarterNote = MidiConstants::kTicksPerQuarterNoteInt);
+	FMidiWriter(FArchive& OutputArhive, int32 TicksPerQuarterNote = Harmonix::Midi::Constants::GTicksPerQuarterNoteInt);
 	~FMidiWriter();
 
 	/** Flushes data to file; also called by dtor */
diff --git a/Engine/Plugins/Runtime/Harmonix/Source/HarmonixMidi/Public/HarmonixMidi/SectionMap.h b/Engine/Plugins/Runtime/Harmonix/Source/HarmonixMidi/Public/HarmonixMidi/SectionMap.h
--- a/Engine/Plugins/Runtime/Harmonix/Source/HarmonixMidi/Public/HarmonixMidi/SectionMap.h
+++ b/Engine/Plugins/Runtime/Harmonix/Source/HarmonixMidi/Public/HarmonixMidi/SectionMap.h
@@ -42,7 +42,7 @@ struct HARMONIXMIDI_API FSectionMap
 
 public:
 	FSectionMap()
-		: TicksPerQuarterNote(MidiConstants::kTicksPerQuarterNoteInt)
+		: TicksPerQuarterNote(Harmonix::Midi::Constants::GTicksPerQuarterNoteInt)
 	{}
 	friend bool operator==(const FSectionMap& Left, const FSectionMap& Right);
 
diff --git a/Engine/Plugins/Runtime/Harmonix/Source/HarmonixMidi/Public/HarmonixMidi/SongMaps.h b/Engine/Plugins/Runtime/Harmonix/Source/HarmonixMidi/Public/HarmonixMidi/SongMaps.h
--- a/Engine/Plugins/Runtime/Harmonix/Source/HarmonixMidi/Public/HarmonixMidi/SongMaps.h
+++ b/Engine/Plugins/Runtime/Harmonix/Source/HarmonixMidi/Public/HarmonixMidi/SongMaps.h
@@ -155,7 +155,7 @@ public:
 protected:
 	friend class FSongMapReceiver;
 	UPROPERTY()
-	int32 TicksPerQuarterNote = MidiConstants::kTicksPerQuarterNoteInt;
+	int32 TicksPerQuarterNote = Harmonix::Midi::Constants::GTicksPerQuarterNoteInt;
 	UPROPERTY()
 	FTempoMap TempoMap;
 	UPROPERTY()
diff --git a/Engine/Plugins/Runtime/Harmonix/Source/HarmonixMidi/Public/HarmonixMidi/TempoMap.h b/Engine/Plugins/Runtime/Harmonix/Source/HarmonixMidi/Public/HarmonixMidi/TempoMap.h
--- a/Engine/Plugins/Runtime/Harmonix/Source/HarmonixMidi/Public/HarmonixMidi/TempoMap.h
+++ b/Engine/Plugins/Runtime/Harmonix/Source/HarmonixMidi/Public/HarmonixMidi/TempoMap.h
@@ -31,7 +31,7 @@ public:
 		return Ms == Other.Ms && MidiTempo == Other.MidiTempo;
 	}
 
-	float GetBPM() const { return MidiConstants::MidiTempoToBPM(MidiTempo); }
+	float GetBPM() const { return Harmonix::Midi::Constants::MidiTempoToBPM(MidiTempo); }
 
 	UPROPERTY()
 	float Ms = 0;        // The time at which the tempo is changing
@@ -61,7 +61,7 @@ struct HARMONIXMIDI_API FTempoMap
 public:
 
 	FTempoMap() 
-		: TicksPerQuarterNote(MidiConstants::kTicksPerQuarterNoteInt)
+		: TicksPerQuarterNote(Harmonix::Midi::Constants::GTicksPerQuarterNoteInt)
 	{}
 	virtual ~FTempoMap() {}
 
diff --git a/Engine/Plugins/Runtime/Harmonix/Source/HarmonixMidiEditor/Private/MidiFileFactory.cpp b/Engine/Plugins/Runtime/Harmonix/Source/HarmonixMidiEditor/Private/MidiFileFactory.cpp
--- a/Engine/Plugins/Runtime/Harmonix/Source/HarmonixMidiEditor/Private/MidiFileFactory.cpp
+++ b/Engine/Plugins/Runtime/Harmonix/Source/HarmonixMidiEditor/Private/MidiFileFactory.cpp
@@ -36,7 +36,7 @@ UObject* UMidiFileFactory::FactoryCreateFile(UClass* InClass, UObject* InParent,
 	{
 		MidiFileAsset = NewObject<UMidiFile>(InParent, InClass, InName, Flags);
 	}
-	MidiFileAsset->LoadStdMidiFile(Filename, MidiConstants::kTicksPerQuarterNoteInt);
+	MidiFileAsset->LoadStdMidiFile(Filename, Harmonix::Midi::Constants::GTicksPerQuarterNoteInt);
 	
 	//add the current file to the file array (for handling multi-batch import)
 	ImportedFiles.Add(MidiFileAsset);
@@ -186,7 +186,7 @@ EReimportResult::Type UMidiFileFactory::Reimport(UObject* Obj)
 		return EReimportResult::Failed;
 	}
 	
-	AsMidiFile->LoadStdMidiFile(ReimportPath, MidiConstants::kTicksPerQuarterNoteInt);
+	AsMidiFile->LoadStdMidiFile(ReimportPath, Harmonix::Midi::Constants::GTicksPerQuarterNoteInt);
 
 	//Allow file length to be reconformed after reimport 
 	AsMidiFile->bLengthRoundedDown = false;
diff --git a/Engine/Plugins/Runtime/Harmonix/Source/HarmonixMidiEditor/Private/MidiNoteCustomization.cpp b/Engine/Plugins/Runtime/Harmonix/Source/HarmonixMidiEditor/Private/MidiNoteCustomization.cpp
--- a/Engine/Plugins/Runtime/Harmonix/Source/HarmonixMidiEditor/Private/MidiNoteCustomization.cpp
+++ b/Engine/Plugins/Runtime/Harmonix/Source/HarmonixMidiEditor/Private/MidiNoteCustomization.cpp
@@ -53,7 +53,7 @@ void FMidiNoteCustomization::OnValueStringSelected(const FString& SelectedString
 
 void FMidiNoteCustomization::OnGetStrings(TArray< TSharedPtr<FString> >& OutStrings, TArray<TSharedPtr<SToolTip>>& OutToolTips, TArray<bool>& OutRestrictedItems) const
 {
-	for (uint8 NoteNumber = MidiConstants::kMinNote; NoteNumber < MidiConstants::kMaxNumNotes; ++NoteNumber)
+	for (uint8 NoteNumber = Harmonix::Midi::Constants::GMinNote; NoteNumber < Harmonix::Midi::Constants::GMaxNumNotes; ++NoteNumber)
 	{
 		OutStrings.Add(MakeShared<FString>(FMidiNote(NoteNumber).ToEditorString()));
 		OutRestrictedItems.Add(false);
diff --git a/Engine/Plugins/Runtime/Harmonix/Source/HarmonixMidiEditor/Private/Pins/MidiNotePin.cpp b/Engine/Plugins/Runtime/Harmonix/Source/HarmonixMidiEditor/Private/Pins/MidiNotePin.cpp
--- a/Engine/Plugins/Runtime/Harmonix/Source/HarmonixMidiEditor/Private/Pins/MidiNotePin.cpp
+++ b/Engine/Plugins/Runtime/Harmonix/Source/HarmonixMidiEditor/Private/Pins/MidiNotePin.cpp
@@ -37,7 +37,7 @@ TSharedRef<SWidget> SMidiNotePin::GetDefaultValueWidget()
 
 void SMidiNotePin::OnGetStrings(TArray< TSharedPtr<FString> >& OutStrings, TArray<TSharedPtr<SToolTip>>& OutToolTips, TArray<bool>& OutRestrictedItems) const
 {
-	for (uint8 noteNumber = MidiConstants::kMinNote; noteNumber < MidiConstants::kMaxNumNotes; ++noteNumber)
+	for (uint8 noteNumber = Harmonix::Midi::Constants::GMinNote; noteNumber < Harmonix::Midi::Constants::GMaxNumNotes; ++noteNumber)
 	{
 		OutStrings.Add(MakeShared<FString>(FMidiNote(noteNumber).ToEditorString()));
 		OutRestrictedItems.Add(false);
diff --git a/Engine/Plugins/Runtime/Harmonix/Source/HarmonixMidiTests/Private/ConformMidiFileLengthTests.cpp b/Engine/Plugins/Runtime/Harmonix/Source/HarmonixMidiTests/Private/ConformMidiFileLengthTests.cpp
--- a/Engine/Plugins/Runtime/Harmonix/Source/HarmonixMidiTests/Private/ConformMidiFileLengthTests.cpp
+++ b/Engine/Plugins/Runtime/Harmonix/Source/HarmonixMidiTests/Private/ConformMidiFileLengthTests.cpp
@@ -122,7 +122,7 @@ namespace HarmonixMidiTests::ConformMidiFileLength
 				}
 
 				//check for pitch bend events on the last tick
-				if (MsgStatus == MidiConstants::kPitch)
+				if (MsgStatus == Harmonix::Midi::Constants::GPitch)
 				{
 					//check if there exist multiple events with same status on the last tick
 					for (int32 i = EventIndex - 1; i >= 0 && Events[i].GetTick() == ConformedLastEventTick; --i)
@@ -135,7 +135,7 @@ namespace HarmonixMidiTests::ConformMidiFileLength
 				}
 
 				//check for CC events and Poly Press on the last tick with the same controller ID
-				if (MsgStatus == MidiConstants::kControl || MsgStatus == MidiConstants::kPolyPres)
+				if (MsgStatus == Harmonix::Midi::Constants::GControl || MsgStatus == Harmonix::Midi::Constants::GPolyPres)
 				{
 					uint8 CurrentEventControllerId = CurrentMsg.Data1;
 					for (int32 i = EventIndex - 1; i >= 0 && Events[i].GetTick() == ConformedLastEventTick; --i)
diff --git a/Engine/Plugins/Runtime/Harmonix/Source/HarmonixMidiTests/Private/MidiTestUtility.cpp b/Engine/Plugins/Runtime/Harmonix/Source/HarmonixMidiTests/Private/MidiTestUtility.cpp
--- a/Engine/Plugins/Runtime/Harmonix/Source/HarmonixMidiTests/Private/MidiTestUtility.cpp
+++ b/Engine/Plugins/Runtime/Harmonix/Source/HarmonixMidiTests/Private/MidiTestUtility.cpp
@@ -18,7 +18,7 @@ namespace Harmonix::Testing::Utility::MidiTestUtility
 		constexpr uint8 DefaultPolyPresValue = 127;
 
 		// Set initial tempo and time signature
-		TheMidiFile->GetSongMaps()->GetTempoMap().AddTempoInfoPoint(MidiConstants::BPMToMidiTempo(InTempo), 0);
+		TheMidiFile->GetSongMaps()->GetTempoMap().AddTempoInfoPoint(Midi::Constants::BPMToMidiTempo(InTempo), 0);
 		TheMidiFile->GetSongMaps()->GetBarMap().AddTimeSignatureAtBarIncludingCountIn(0, InTimeSigNum, InTimeSigDenom);
 		TheMidiFile->GetSongMaps()->GetBarMap().SetTicksPerQuarterNote(DefaultTicksPerQuarter);
 
@@ -39,7 +39,7 @@ namespace Harmonix::Testing::Utility::MidiTestUtility
 					//(InFileLengthBars will end up being the very beginning of the next bar)
 					LastEventTickInFile -= 1;
 				}
-				FMidiEvent TextEventAtTheEnd(LastEventTickInFile, FMidiMsg::CreateText(TextIndex, MidiConstants::kMeta_Text));
+				FMidiEvent TextEventAtTheEnd(LastEventTickInFile, FMidiMsg::CreateText(TextIndex, Midi::Constants::GMeta_Text));
 				CurrentTrack->AddEvent(TextEventAtTheEnd);
 			}
 
@@ -80,7 +80,7 @@ namespace Harmonix::Testing::Utility::MidiTestUtility
 		int32 CurrentBarTick = InFile->GetSongMaps()->GetBarMap().BarIncludingCountInToTick(InBarIndex);
 		//Add Control Events to each channels/tracks
 		FMidiEvent CurrentCCEvent(CurrentBarTick,
-			FMidiMsg(MidiConstants::kControl + InChannel, InControllerID, InControlValue)
+			FMidiMsg(Midi::Constants::GControl + InChannel, InControllerID, InControlValue)
 		);
 		CurrentTrack->AddEvent(CurrentCCEvent);
 	}
@@ -92,7 +92,7 @@ namespace Harmonix::Testing::Utility::MidiTestUtility
 		int32 CurrentBarTick = InFile->GetSongMaps()->GetBarMap().BarIncludingCountInToTick(InBarIndex);
 		//text events
 		uint16 TextIndex = CurrentTrack->AddText(InText);
-		FMidiEvent CurrentTextEvent(CurrentBarTick, FMidiMsg::CreateText(TextIndex, MidiConstants::kMeta_Text));
+		FMidiEvent CurrentTextEvent(CurrentBarTick, FMidiMsg::CreateText(TextIndex, Midi::Constants::GMeta_Text));
 		CurrentTrack->AddEvent(CurrentTextEvent);
 	}
 
@@ -103,7 +103,7 @@ namespace Harmonix::Testing::Utility::MidiTestUtility
 		int32 CurrentBarTick = InFile->GetSongMaps()->GetBarMap().BarIncludingCountInToTick(InBarIndex);
 		//Add pitch bend events to channels/tracks
 		FMidiEvent CurrentPitchEvent(CurrentBarTick,
-			FMidiMsg(MidiConstants::kPitch + InChannel, InPitchValueLSB, InPitchValueMSB)
+			FMidiMsg(Midi::Constants::GPitch + InChannel, InPitchValueLSB, InPitchValueMSB)
 		);
 		CurrentTrack->AddEvent(CurrentPitchEvent);
 	}
@@ -116,7 +116,7 @@ namespace Harmonix::Testing::Utility::MidiTestUtility
 		int32 CurrentBarTick = InFile->GetSongMaps()->GetBarMap().BarIncludingCountInToTick(InBarIndex);
 		//Add Poly Pres events to channels/tracks
 		FMidiEvent CurrentPolyPresEvent(CurrentBarTick,
-			FMidiMsg(MidiConstants::kPolyPres + InChannel, InNoteNumber, InPolyPresValue)
+			FMidiMsg(Midi::Constants::GPolyPres + InChannel, InNoteNumber, InPolyPresValue)
 		);
 		CurrentTrack->AddEvent(CurrentPolyPresEvent);
 	}
