[Harmonix] Move FMidiStreamEventTrackChannelFilter to internal, since it's only used by deprecated nodes.

#tests automated tests

diff --git a/Engine/Plugins/Runtime/Harmonix/Source/HarmonixMetasound/Private/HarmonixMetasound/DataTypes/MidiStream.cpp b/Engine/Plugins/Runtime/Harmonix/Source/HarmonixMetasound/Private/HarmonixMetasound/DataTypes/MidiStream.cpp
--- a/Engine/Plugins/Runtime/Harmonix/Source/HarmonixMetasound/Private/HarmonixMetasound/DataTypes/MidiStream.cpp
+++ b/Engine/Plugins/Runtime/Harmonix/Source/HarmonixMetasound/Private/HarmonixMetasound/DataTypes/MidiStream.cpp
@@ -402,96 +402,6 @@ namespace HarmonixMetasound
 		}
 	}
 
-	uint64 FMidiStreamEventTrackChannelFilter::StringToBitfield(const FString& InString, FString& OutErrorMessage)
-	{
-		OutErrorMessage.Empty();
-		uint64 Bitfield = 0;
-		auto SetBit = [&](uint32 bit) { Bitfield |= (uint64(1) << bit); };
-
-		if (InString.Equals(TEXT("*")))
-		{
-			Bitfield = std::numeric_limits<uint64>::max();
-			return Bitfield;
-		}
-
-		TArray<FString> CommaParts;
-		InString.ParseIntoArray(CommaParts, TEXT(","));
-		for (const FString& CommaPart : CommaParts)
-		{
-			TArray<FString> DashParts;
-			CommaPart.ParseIntoArray(DashParts, TEXT("-"));
-			if (DashParts.IsEmpty() || DashParts.Num() > 2)
-			{
-				OutErrorMessage = FString::Printf(TEXT("Unexpected filter expression -> \"%s\""), *CommaPart);
-				return 0;
-			}
-
-			if (DashParts.Num() == 1)
-			{
-				int32 AsNumber;
-				if (!LexTryParseString(AsNumber, *DashParts[0]))
-				{
-					OutErrorMessage = FString::Printf(TEXT("Error parsing part of filter string -> \"%s\""), *DashParts[0]);
-					return 0;
-				}
-
-				if (AsNumber <= 0)
-				{
-					OutErrorMessage = TEXT("Track and Channel filter indexes should be '1' based!");
-					return 0;
-				}
-
-				SetBit(AsNumber - 1);
-			}
-			else if (DashParts.Num() == 2)
-			{
-				int32 AsNumberStart;
-				int32 AsNumberEnd;
-				if (!LexTryParseString(AsNumberStart, *DashParts[0]))
-				{
-					OutErrorMessage = FString::Printf(TEXT("Error parsing part of filter string -> \"%s\""), *DashParts[0]);
-					return 0;
-				}
-
-				if (!LexTryParseString(AsNumberEnd, *DashParts[1]))
-				{
-					OutErrorMessage = FString::Printf(TEXT("Error parsing part of filter string -> \"%s\""), *DashParts[1]);
-					return 0;
-				}
-
-				if (AsNumberStart <= 0)
-				{
-					OutErrorMessage = TEXT("Track and Channel filter indexes should be '1' based!");
-					return 0;
-				}
-
-				if (AsNumberEnd < AsNumberStart)
-				{
-					OutErrorMessage = FString::Printf(TEXT("End of range must be greater than start of range (\"%s\")!"), *CommaPart);
-					return 0;
-				}
-
-				for (int32 i = AsNumberStart; i < (AsNumberEnd + 1); ++i)
-				{
-					SetBit(i - 1);
-				}
-			}
-		}
-		return Bitfield;
-	}
-
-	bool FMidiStreamEventTrackChannelFilter::SetMidiChannelFilterFromString(const FString& FilterString, FString& OutErrorMessage)
-	{
-		MidiChannelFilter = (uint16)StringToBitfield(FilterString, OutErrorMessage);
-		return OutErrorMessage.IsEmpty();
-	}
-
-	bool FMidiStreamEventTrackChannelFilter::SetTrackFilterFromString(const FString& FilterString, FString& OutErrorMessage)
-	{
-		TrackFilter = (uint64)StringToBitfield(FilterString, OutErrorMessage);
-		return OutErrorMessage.IsEmpty();
-	}
-
 	void FMidiVoiceTracker::Process(const FMidiStream& MidiStream, const FKillVoiceFn& KillVoiceFn)
 	{
 		// Remove active voices that have note offs this block
diff --git a/Engine/Plugins/Runtime/Harmonix/Source/HarmonixMetasound/Public/HarmonixMetasound/DataTypes/MidiStream.h b/Engine/Plugins/Runtime/Harmonix/Source/HarmonixMetasound/Public/HarmonixMetasound/DataTypes/MidiStream.h
--- a/Engine/Plugins/Runtime/Harmonix/Source/HarmonixMetasound/Public/HarmonixMetasound/DataTypes/MidiStream.h
+++ b/Engine/Plugins/Runtime/Harmonix/Source/HarmonixMetasound/Public/HarmonixMetasound/DataTypes/MidiStream.h
@@ -13,7 +13,6 @@
 
 #include "MidiClock.h"
 
-#include <functional>
 #include <atomic>
 
 namespace HarmonixMetasound
@@ -49,38 +48,6 @@ namespace HarmonixMetasound
 		FMidiVoiceId VoiceId;
 	};
 
-	struct HARMONIXMETASOUND_API FMidiStreamEventTrackChannelFilter
-	{
-		uint16 MidiChannelFilter;
-		uint64 TrackFilter;
-
-		/** 
-		* Parses the FilterString to a Bitfield and assigns it to the MidiChannelFilter
-		* if it fails, will assign MidiChannelFilter to 0
-		*/
-		bool SetMidiChannelFilterFromString(const FString& FilterString, FString& OutErrorMessage);
-
-		/**
-		* Parses the FilterString to a Bitfield and assigns it to the TrackFilter
-		* if it fails, will assign TrackFilter to 0
-		*/
-		bool SetTrackFilterFromString(const FString& FilterString, FString& OutErrorMessage);
-
-		bool operator()(const FMidiStreamEvent& Event) const
-		{
-			if (!(TrackFilter & (uint64(1) << Event.TrackIndex)))
-			{
-				return false;
-			}
-
-			return !Event.MidiMessage.IsStd() || (MidiChannelFilter & (uint16(1) << Event.MidiMessage.GetStdChannel()));
-		}
-
-	private:
-		uint64 StringToBitfield(const FString& InString, FString& OutErrorMesage);
-	};
-
-
 	class HARMONIXMETASOUND_API FMidiStream
 	{
 	public:
