[Harmonix] Node categorization pass

#tests automated tests

diff --git a/Engine/Plugins/Runtime/Harmonix/Source/HarmonixMetasound/Private/HarmonixMetasound/Common.cpp b/Engine/Plugins/Runtime/Harmonix/Source/HarmonixMetasound/Private/HarmonixMetasound/Common.cpp
--- a/Engine/Plugins/Runtime/Harmonix/Source/HarmonixMetasound/Private/HarmonixMetasound/Common.cpp
+++ b/Engine/Plugins/Runtime/Harmonix/Source/HarmonixMetasound/Private/HarmonixMetasound/Common.cpp
@@ -14,6 +14,7 @@ namespace HarmonixMetasound
 	{
 		const FText Harmonix = METASOUND_LOCTEXT("Metasound_HarmonixCategory", "Harmonix");
 		const FText Modulation = METASOUND_LOCTEXT("Metasound_ModulationCategory", "Modulation");
+		const FText Analysis = METASOUND_LOCTEXT("Metasound_AnalsisCategory", "Analysis");
 	}
 	
 	namespace CommonPinNames
diff --git a/Engine/Plugins/Runtime/Harmonix/Source/HarmonixMetasound/Private/HarmonixMetasound/Nodes/DelayNode.cpp b/Engine/Plugins/Runtime/Harmonix/Source/HarmonixMetasound/Private/HarmonixMetasound/Nodes/DelayNode.cpp
--- a/Engine/Plugins/Runtime/Harmonix/Source/HarmonixMetasound/Private/HarmonixMetasound/Nodes/DelayNode.cpp
+++ b/Engine/Plugins/Runtime/Harmonix/Source/HarmonixMetasound/Private/HarmonixMetasound/Nodes/DelayNode.cpp
@@ -180,7 +180,7 @@ namespace HarmonixMetasound::DelayNode
 				Info.MajorVersion = 0;
 				Info.MinorVersion = 1;
 				Info.DisplayName = METASOUND_LOCTEXT("Delay_DisplayName", "Clock-Synced Delay");
-				Info.CategoryHierarchy = { NodeCategories::Delays };
+				Info.CategoryHierarchy = { MetasoundNodeCategories::Harmonix, NodeCategories::Delays };
 				Info.Description = METASOUND_LOCTEXT("Delay_Description", "Delay effect with optional music clock sync.");
 				Info.Author = PluginAuthor;
 				Info.PromptIfMissing = PluginNodeMissingPrompt;
diff --git a/Engine/Plugins/Runtime/Harmonix/Source/HarmonixMetasound/Private/HarmonixMetasound/Nodes/DistortionNode.cpp b/Engine/Plugins/Runtime/Harmonix/Source/HarmonixMetasound/Private/HarmonixMetasound/Nodes/DistortionNode.cpp
--- a/Engine/Plugins/Runtime/Harmonix/Source/HarmonixMetasound/Private/HarmonixMetasound/Nodes/DistortionNode.cpp
+++ b/Engine/Plugins/Runtime/Harmonix/Source/HarmonixMetasound/Private/HarmonixMetasound/Nodes/DistortionNode.cpp
@@ -234,7 +234,7 @@ namespace HarmonixMetasound
 				Info.MajorVersion = 0;
 				Info.MinorVersion = 1;
 				Info.DisplayName = METASOUND_LOCTEXT("Distortion_DisplayName", "Distortion");
-				Info.CategoryHierarchy = { NodeCategories::Dynamics };
+				Info.CategoryHierarchy = { MetasoundNodeCategories::Harmonix, NodeCategories::Dynamics };
 				Info.Description = METASOUND_LOCTEXT("Distortion_Description", "Distortion audio effect.");
 				Info.Author = PluginAuthor;
 				Info.PromptIfMissing = PluginNodeMissingPrompt;
diff --git a/Engine/Plugins/Runtime/Harmonix/Source/HarmonixMetasound/Private/HarmonixMetasound/Nodes/FFTAnalyzerNode.cpp b/Engine/Plugins/Runtime/Harmonix/Source/HarmonixMetasound/Private/HarmonixMetasound/Nodes/FFTAnalyzerNode.cpp
--- a/Engine/Plugins/Runtime/Harmonix/Source/HarmonixMetasound/Private/HarmonixMetasound/Nodes/FFTAnalyzerNode.cpp
+++ b/Engine/Plugins/Runtime/Harmonix/Source/HarmonixMetasound/Private/HarmonixMetasound/Nodes/FFTAnalyzerNode.cpp
@@ -116,7 +116,7 @@ namespace HarmonixMetasound
 				Info.MajorVersion = 0;
 				Info.MinorVersion = 1;
 				Info.DisplayName = METASOUND_LOCTEXT("FFTAnalyzer_DisplayName", "FFT Analyzer");
-				Info.CategoryHierarchy = { NodeCategories::Dynamics };
+				Info.CategoryHierarchy = { MetasoundNodeCategories::Harmonix, MetasoundNodeCategories::Analysis };
 				Info.Description = METASOUND_LOCTEXT("FFTAnalyzer_Description", "Fast Fourier Transform Analysis.");
 				Info.Author = PluginAuthor;
 				Info.PromptIfMissing = PluginNodeMissingPrompt;
diff --git a/Engine/Plugins/Runtime/Harmonix/Source/HarmonixMetasound/Private/HarmonixMetasound/Nodes/FusionSamplerNode.cpp b/Engine/Plugins/Runtime/Harmonix/Source/HarmonixMetasound/Private/HarmonixMetasound/Nodes/FusionSamplerNode.cpp
--- a/Engine/Plugins/Runtime/Harmonix/Source/HarmonixMetasound/Private/HarmonixMetasound/Nodes/FusionSamplerNode.cpp
+++ b/Engine/Plugins/Runtime/Harmonix/Source/HarmonixMetasound/Private/HarmonixMetasound/Nodes/FusionSamplerNode.cpp
@@ -329,7 +329,7 @@ namespace HarmonixMetasound
 				Info.Author = PluginAuthor;
 				Info.PromptIfMissing = PluginNodeMissingPrompt;
 				Info.DefaultInterface = GetVertexInterface();
-				Info.CategoryHierarchy.Emplace(NodeCategories::Music);
+				Info.CategoryHierarchy = { MetasoundNodeCategories::Harmonix, NodeCategories::Generators };
 				return Info;
 			};
 
diff --git a/Engine/Plugins/Runtime/Harmonix/Source/HarmonixMetasound/Private/HarmonixMetasound/Nodes/MetronomeNode.cpp b/Engine/Plugins/Runtime/Harmonix/Source/HarmonixMetasound/Private/HarmonixMetasound/Nodes/MetronomeNode.cpp
--- a/Engine/Plugins/Runtime/Harmonix/Source/HarmonixMetasound/Private/HarmonixMetasound/Nodes/MetronomeNode.cpp
+++ b/Engine/Plugins/Runtime/Harmonix/Source/HarmonixMetasound/Private/HarmonixMetasound/Nodes/MetronomeNode.cpp
@@ -118,7 +118,7 @@ namespace HarmonixMetasound
 			Info.Author           = PluginAuthor;
 			Info.PromptIfMissing  = PluginNodeMissingPrompt;
 			Info.DefaultInterface = GetVertexInterface();
-			Info.CategoryHierarchy.Emplace(NodeCategories::Music);
+			Info.CategoryHierarchy = { MetasoundNodeCategories::Harmonix, NodeCategories::Music };
 
 			return Info;
 		};
diff --git a/Engine/Plugins/Runtime/Harmonix/Source/HarmonixMetasound/Private/HarmonixMetasound/Nodes/MidiCCTriggerNode.cpp b/Engine/Plugins/Runtime/Harmonix/Source/HarmonixMetasound/Private/HarmonixMetasound/Nodes/MidiCCTriggerNode.cpp
--- a/Engine/Plugins/Runtime/Harmonix/Source/HarmonixMetasound/Private/HarmonixMetasound/Nodes/MidiCCTriggerNode.cpp
+++ b/Engine/Plugins/Runtime/Harmonix/Source/HarmonixMetasound/Private/HarmonixMetasound/Nodes/MidiCCTriggerNode.cpp
@@ -118,7 +118,7 @@ namespace HarmonixMetasound::Nodes::MidiCCTriggerNode
 				Info.Author = PluginAuthor;
 				Info.PromptIfMissing = PluginNodeMissingPrompt;
 				Info.DefaultInterface = GetVertexInterface();
-				Info.CategoryHierarchy.Emplace(NodeCategories::Music);
+				Info.CategoryHierarchy = { MetasoundNodeCategories::Harmonix, NodeCategories::Music };
 				return Info;
 			};
 
diff --git a/Engine/Plugins/Runtime/Harmonix/Source/HarmonixMetasound/Private/HarmonixMetasound/Nodes/MidiClockStateNode.cpp b/Engine/Plugins/Runtime/Harmonix/Source/HarmonixMetasound/Private/HarmonixMetasound/Nodes/MidiClockStateNode.cpp
--- a/Engine/Plugins/Runtime/Harmonix/Source/HarmonixMetasound/Private/HarmonixMetasound/Nodes/MidiClockStateNode.cpp
+++ b/Engine/Plugins/Runtime/Harmonix/Source/HarmonixMetasound/Private/HarmonixMetasound/Nodes/MidiClockStateNode.cpp
@@ -86,7 +86,7 @@ namespace HarmonixMetasound
 			Info.Author           = PluginAuthor;
 			Info.PromptIfMissing  = PluginNodeMissingPrompt;
 			Info.DefaultInterface = GetVertexInterface();
-			Info.CategoryHierarchy.Emplace(NodeCategories::Music);
+			Info.CategoryHierarchy = { MetasoundNodeCategories::Harmonix, NodeCategories::Music };
 			return Info;
 		};
 
diff --git a/Engine/Plugins/Runtime/Harmonix/Source/HarmonixMetasound/Private/HarmonixMetasound/Nodes/MidiNoteTriggerNode.cpp b/Engine/Plugins/Runtime/Harmonix/Source/HarmonixMetasound/Private/HarmonixMetasound/Nodes/MidiNoteTriggerNode.cpp
--- a/Engine/Plugins/Runtime/Harmonix/Source/HarmonixMetasound/Private/HarmonixMetasound/Nodes/MidiNoteTriggerNode.cpp
+++ b/Engine/Plugins/Runtime/Harmonix/Source/HarmonixMetasound/Private/HarmonixMetasound/Nodes/MidiNoteTriggerNode.cpp
@@ -62,7 +62,7 @@ namespace HarmonixMetasound::Nodes::MidiNoteTriggerNode
 				Info.Author           = PluginAuthor;
 				Info.PromptIfMissing  = PluginNodeMissingPrompt;
 				Info.DefaultInterface = GetVertexInterface();
-				Info.CategoryHierarchy.Emplace(NodeCategories::Music);
+				Info.CategoryHierarchy = { MetasoundNodeCategories::Harmonix, NodeCategories::Music };
 				return Info;
 			};
 
diff --git a/Engine/Plugins/Runtime/Harmonix/Source/HarmonixMetasound/Private/HarmonixMetasound/Nodes/MidiPlayerNode.cpp b/Engine/Plugins/Runtime/Harmonix/Source/HarmonixMetasound/Private/HarmonixMetasound/Nodes/MidiPlayerNode.cpp
--- a/Engine/Plugins/Runtime/Harmonix/Source/HarmonixMetasound/Private/HarmonixMetasound/Nodes/MidiPlayerNode.cpp
+++ b/Engine/Plugins/Runtime/Harmonix/Source/HarmonixMetasound/Private/HarmonixMetasound/Nodes/MidiPlayerNode.cpp
@@ -180,7 +180,7 @@ namespace HarmonixMetasound
 			Info.Author           = PluginAuthor;
 			Info.PromptIfMissing  = PluginNodeMissingPrompt;
 			Info.DefaultInterface = GetVertexInterface();
-			Info.CategoryHierarchy.Emplace(NodeCategories::Music);
+			Info.CategoryHierarchy = { MetasoundNodeCategories::Harmonix, NodeCategories::Music };
 
 			return Info;
 		};
diff --git a/Engine/Plugins/Runtime/Harmonix/Source/HarmonixMetasound/Private/HarmonixMetasound/Nodes/MidiPulseGeneratorNode.cpp b/Engine/Plugins/Runtime/Harmonix/Source/HarmonixMetasound/Private/HarmonixMetasound/Nodes/MidiPulseGeneratorNode.cpp
--- a/Engine/Plugins/Runtime/Harmonix/Source/HarmonixMetasound/Private/HarmonixMetasound/Nodes/MidiPulseGeneratorNode.cpp
+++ b/Engine/Plugins/Runtime/Harmonix/Source/HarmonixMetasound/Private/HarmonixMetasound/Nodes/MidiPulseGeneratorNode.cpp
@@ -61,7 +61,7 @@ namespace HarmonixMetasound::Nodes::MidiPulseGeneratorNode
 				Info.Author           = PluginAuthor;
 				Info.PromptIfMissing  = PluginNodeMissingPrompt;
 				Info.DefaultInterface = GetVertexInterface();
-				Info.CategoryHierarchy.Emplace(NodeCategories::Music);
+				Info.CategoryHierarchy = { MetasoundNodeCategories::Harmonix, NodeCategories::Music };
 				return Info;
 			};
 
diff --git a/Engine/Plugins/Runtime/Harmonix/Source/HarmonixMetasound/Private/HarmonixMetasound/Nodes/MidiStreamLoggerNode.cpp b/Engine/Plugins/Runtime/Harmonix/Source/HarmonixMetasound/Private/HarmonixMetasound/Nodes/MidiStreamLoggerNode.cpp
--- a/Engine/Plugins/Runtime/Harmonix/Source/HarmonixMetasound/Private/HarmonixMetasound/Nodes/MidiStreamLoggerNode.cpp
+++ b/Engine/Plugins/Runtime/Harmonix/Source/HarmonixMetasound/Private/HarmonixMetasound/Nodes/MidiStreamLoggerNode.cpp
@@ -74,7 +74,7 @@ namespace HarmonixMetasound
 			Info.Author           = PluginAuthor;
 			Info.PromptIfMissing  = PluginNodeMissingPrompt;
 			Info.DefaultInterface = GetVertexInterface();
-			Info.CategoryHierarchy.Emplace(NodeCategories::Music);
+			Info.CategoryHierarchy = { MetasoundNodeCategories::Harmonix, NodeCategories::Music, NodeCategories::Debug };
 			return Info;
 		};
 
diff --git a/Engine/Plugins/Runtime/Harmonix/Source/HarmonixMetasound/Private/HarmonixMetasound/Nodes/MidiStreamMergeNode.cpp b/Engine/Plugins/Runtime/Harmonix/Source/HarmonixMetasound/Private/HarmonixMetasound/Nodes/MidiStreamMergeNode.cpp
--- a/Engine/Plugins/Runtime/Harmonix/Source/HarmonixMetasound/Private/HarmonixMetasound/Nodes/MidiStreamMergeNode.cpp
+++ b/Engine/Plugins/Runtime/Harmonix/Source/HarmonixMetasound/Private/HarmonixMetasound/Nodes/MidiStreamMergeNode.cpp
@@ -71,7 +71,7 @@ namespace HarmonixMetasound
 			Info.Author           = PluginAuthor;
 			Info.PromptIfMissing  = PluginNodeMissingPrompt;
 			Info.DefaultInterface = GetVertexInterface();
-			Info.CategoryHierarchy.Emplace(NodeCategories::Music);
+			Info.CategoryHierarchy = { MetasoundNodeCategories::Harmonix, NodeCategories::Music };
 			return Info;
 		};
 
diff --git a/Engine/Plugins/Runtime/Harmonix/Source/HarmonixMetasound/Private/HarmonixMetasound/Nodes/MidiStreamSelectNode.cpp b/Engine/Plugins/Runtime/Harmonix/Source/HarmonixMetasound/Private/HarmonixMetasound/Nodes/MidiStreamSelectNode.cpp
--- a/Engine/Plugins/Runtime/Harmonix/Source/HarmonixMetasound/Private/HarmonixMetasound/Nodes/MidiStreamSelectNode.cpp
+++ b/Engine/Plugins/Runtime/Harmonix/Source/HarmonixMetasound/Private/HarmonixMetasound/Nodes/MidiStreamSelectNode.cpp
@@ -86,7 +86,7 @@ namespace HarmonixMetasound
 			Info.Author = PluginAuthor;
 			Info.PromptIfMissing = PluginNodeMissingPrompt;
 			Info.DefaultInterface = GetVertexInterface();
-			Info.CategoryHierarchy.Emplace(NodeCategories::Music);
+			Info.CategoryHierarchy = { MetasoundNodeCategories::Harmonix, NodeCategories::Music };
 			return Info;
 		};
 
diff --git a/Engine/Plugins/Runtime/Harmonix/Source/HarmonixMetasound/Private/HarmonixMetasound/Nodes/MidiStreamTrackFilterNode.cpp b/Engine/Plugins/Runtime/Harmonix/Source/HarmonixMetasound/Private/HarmonixMetasound/Nodes/MidiStreamTrackFilterNode.cpp
--- a/Engine/Plugins/Runtime/Harmonix/Source/HarmonixMetasound/Private/HarmonixMetasound/Nodes/MidiStreamTrackFilterNode.cpp
+++ b/Engine/Plugins/Runtime/Harmonix/Source/HarmonixMetasound/Private/HarmonixMetasound/Nodes/MidiStreamTrackFilterNode.cpp
@@ -66,7 +66,7 @@ namespace HarmonixMetasound::Nodes::MidiTrackFilter
 				Info.Author           = PluginAuthor;
 				Info.PromptIfMissing  = PluginNodeMissingPrompt;
 				Info.DefaultInterface = GetVertexInterface();
-				Info.CategoryHierarchy.Emplace(NodeCategories::Music);
+				Info.CategoryHierarchy = { MetasoundNodeCategories::Harmonix, NodeCategories::Music };
 				return Info;
 			};
 
diff --git a/Engine/Plugins/Runtime/Harmonix/Source/HarmonixMetasound/Private/HarmonixMetasound/Nodes/MidiStreamTransposerNode.cpp b/Engine/Plugins/Runtime/Harmonix/Source/HarmonixMetasound/Private/HarmonixMetasound/Nodes/MidiStreamTransposerNode.cpp
--- a/Engine/Plugins/Runtime/Harmonix/Source/HarmonixMetasound/Private/HarmonixMetasound/Nodes/MidiStreamTransposerNode.cpp
+++ b/Engine/Plugins/Runtime/Harmonix/Source/HarmonixMetasound/Private/HarmonixMetasound/Nodes/MidiStreamTransposerNode.cpp
@@ -59,7 +59,7 @@ namespace HarmonixMetasound::Nodes::MidiNoteTranspose
 				Info.Author           = PluginAuthor;
 				Info.PromptIfMissing  = PluginNodeMissingPrompt;
 				Info.DefaultInterface = GetVertexInterface();
-				Info.CategoryHierarchy.Emplace(NodeCategories::Music);
+				Info.CategoryHierarchy = { MetasoundNodeCategories::Harmonix, NodeCategories::Music };
 				return Info;
 			};
 
diff --git a/Engine/Plugins/Runtime/Harmonix/Source/HarmonixMetasound/Private/HarmonixMetasound/Nodes/MidiTextTriggerNode.cpp b/Engine/Plugins/Runtime/Harmonix/Source/HarmonixMetasound/Private/HarmonixMetasound/Nodes/MidiTextTriggerNode.cpp
--- a/Engine/Plugins/Runtime/Harmonix/Source/HarmonixMetasound/Private/HarmonixMetasound/Nodes/MidiTextTriggerNode.cpp
+++ b/Engine/Plugins/Runtime/Harmonix/Source/HarmonixMetasound/Private/HarmonixMetasound/Nodes/MidiTextTriggerNode.cpp
@@ -41,7 +41,7 @@ namespace HarmonixMetasound
 				Info.Author           = PluginAuthor;
 				Info.PromptIfMissing  = PluginNodeMissingPrompt;
 				Info.DefaultInterface = GetVertexInterface();
-				Info.CategoryHierarchy.Emplace(NodeCategories::Music);
+				Info.CategoryHierarchy = { MetasoundNodeCategories::Harmonix, NodeCategories::Music };
 				return Info;
 			};
 
diff --git a/Engine/Plugins/Runtime/Harmonix/Source/HarmonixMetasound/Private/HarmonixMetasound/Nodes/MultiThreadedFusionSynchronizerNode.cpp b/Engine/Plugins/Runtime/Harmonix/Source/HarmonixMetasound/Private/HarmonixMetasound/Nodes/MultiThreadedFusionSynchronizerNode.cpp
--- a/Engine/Plugins/Runtime/Harmonix/Source/HarmonixMetasound/Private/HarmonixMetasound/Nodes/MultiThreadedFusionSynchronizerNode.cpp
+++ b/Engine/Plugins/Runtime/Harmonix/Source/HarmonixMetasound/Private/HarmonixMetasound/Nodes/MultiThreadedFusionSynchronizerNode.cpp
@@ -95,7 +95,7 @@ namespace HarmonixMetasound
 					Info.Author = PluginAuthor;
 					Info.PromptIfMissing = PluginNodeMissingPrompt;
 					Info.DefaultInterface = GetVertexInterface();
-					Info.CategoryHierarchy.Emplace(NodeCategories::Music);
+					Info.CategoryHierarchy = { MetasoundNodeCategories::Harmonix, NodeCategories::Generators };
 					return Info;
 				};
 
diff --git a/Engine/Plugins/Runtime/Harmonix/Source/HarmonixMetasound/Private/HarmonixMetasound/Nodes/MultibandAnalyzerNode.cpp b/Engine/Plugins/Runtime/Harmonix/Source/HarmonixMetasound/Private/HarmonixMetasound/Nodes/MultibandAnalyzerNode.cpp
--- a/Engine/Plugins/Runtime/Harmonix/Source/HarmonixMetasound/Private/HarmonixMetasound/Nodes/MultibandAnalyzerNode.cpp
+++ b/Engine/Plugins/Runtime/Harmonix/Source/HarmonixMetasound/Private/HarmonixMetasound/Nodes/MultibandAnalyzerNode.cpp
@@ -235,7 +235,7 @@ namespace HarmonixMetasound::Nodes::MultibandAnalyzer
 				Info.DisplayName = METASOUND_LOCTEXT("MultibandAnalyzer_DisplayName", "Multiband Analyzer");
 				Info.Description = METASOUND_LOCTEXT("MultibandAnalyzer_Description", "Reports levels in frequency bands for an audio signal.");
 				Info.Author = Metasound::PluginAuthor;
-				Info.CategoryHierarchy = { MetasoundNodeCategories::Harmonix, Metasound::NodeCategories::Envelopes };
+				Info.CategoryHierarchy = { MetasoundNodeCategories::Harmonix, MetasoundNodeCategories::Analysis };
 				Info.PromptIfMissing = Metasound::PluginNodeMissingPrompt;
 				Info.DefaultInterface = GetVertexInterface();
 
diff --git a/Engine/Plugins/Runtime/Harmonix/Source/HarmonixMetasound/Private/HarmonixMetasound/Nodes/PeakNode.cpp b/Engine/Plugins/Runtime/Harmonix/Source/HarmonixMetasound/Private/HarmonixMetasound/Nodes/PeakNode.cpp
--- a/Engine/Plugins/Runtime/Harmonix/Source/HarmonixMetasound/Private/HarmonixMetasound/Nodes/PeakNode.cpp
+++ b/Engine/Plugins/Runtime/Harmonix/Source/HarmonixMetasound/Private/HarmonixMetasound/Nodes/PeakNode.cpp
@@ -77,7 +77,7 @@ namespace HarmonixMetasound::Nodes::Peak
 				Info.Author = Metasound::PluginAuthor;
 				Info.PromptIfMissing = Metasound::PluginNodeMissingPrompt;
 				Info.DefaultInterface = GetVertexInterface();
-				Info.CategoryHierarchy = { MetasoundNodeCategories::Harmonix, Metasound::NodeCategories::Envelopes };
+				Info.CategoryHierarchy = { MetasoundNodeCategories::Harmonix, MetasoundNodeCategories::Analysis };
 				return Info;
 			};
 
diff --git a/Engine/Plugins/Runtime/Harmonix/Source/HarmonixMetasound/Private/HarmonixMetasound/Nodes/StepSequencePlayerNode.cpp b/Engine/Plugins/Runtime/Harmonix/Source/HarmonixMetasound/Private/HarmonixMetasound/Nodes/StepSequencePlayerNode.cpp
--- a/Engine/Plugins/Runtime/Harmonix/Source/HarmonixMetasound/Private/HarmonixMetasound/Nodes/StepSequencePlayerNode.cpp
+++ b/Engine/Plugins/Runtime/Harmonix/Source/HarmonixMetasound/Private/HarmonixMetasound/Nodes/StepSequencePlayerNode.cpp
@@ -141,7 +141,7 @@ namespace HarmonixMetasound
 			Info.Author           = PluginAuthor;
 			Info.PromptIfMissing  = PluginNodeMissingPrompt;
 			Info.DefaultInterface = GetVertexInterface();
-			Info.CategoryHierarchy.Emplace(NodeCategories::Music);
+			Info.CategoryHierarchy = { MetasoundNodeCategories::Harmonix, NodeCategories::Music };
 
 			return Info;
 		};
diff --git a/Engine/Plugins/Runtime/Harmonix/Source/HarmonixMetasound/Private/HarmonixMetasound/Nodes/TransportToTriggerNode.cpp b/Engine/Plugins/Runtime/Harmonix/Source/HarmonixMetasound/Private/HarmonixMetasound/Nodes/TransportToTriggerNode.cpp
--- a/Engine/Plugins/Runtime/Harmonix/Source/HarmonixMetasound/Private/HarmonixMetasound/Nodes/TransportToTriggerNode.cpp
+++ b/Engine/Plugins/Runtime/Harmonix/Source/HarmonixMetasound/Private/HarmonixMetasound/Nodes/TransportToTriggerNode.cpp
@@ -66,7 +66,7 @@ namespace HarmonixMetasound
 				Info.Author = PluginAuthor;
 				Info.PromptIfMissing = PluginNodeMissingPrompt;
 				Info.DefaultInterface = GetVertexInterface();
-				Info.CategoryHierarchy.Emplace(NodeCategories::Music);
+				Info.CategoryHierarchy = { MetasoundNodeCategories::Harmonix, NodeCategories::Music };
 				return Info;
 			};
 
diff --git a/Engine/Plugins/Runtime/Harmonix/Source/HarmonixMetasound/Private/HarmonixMetasound/Nodes/TransportWavePlayerControllerNode.cpp b/Engine/Plugins/Runtime/Harmonix/Source/HarmonixMetasound/Private/HarmonixMetasound/Nodes/TransportWavePlayerControllerNode.cpp
--- a/Engine/Plugins/Runtime/Harmonix/Source/HarmonixMetasound/Private/HarmonixMetasound/Nodes/TransportWavePlayerControllerNode.cpp
+++ b/Engine/Plugins/Runtime/Harmonix/Source/HarmonixMetasound/Private/HarmonixMetasound/Nodes/TransportWavePlayerControllerNode.cpp
@@ -72,7 +72,7 @@ namespace HarmonixMetasound
 				Info.Author = PluginAuthor;
 				Info.PromptIfMissing = PluginNodeMissingPrompt;
 				Info.DefaultInterface = GetVertexInterface();
-				Info.CategoryHierarchy.Emplace(NodeCategories::Music);
+				Info.CategoryHierarchy = { MetasoundNodeCategories::Harmonix, NodeCategories::Generators };
 				return Info;
 			};
 
diff --git a/Engine/Plugins/Runtime/Harmonix/Source/HarmonixMetasound/Private/HarmonixMetasound/Nodes/TriggerToTransportNode.cpp b/Engine/Plugins/Runtime/Harmonix/Source/HarmonixMetasound/Private/HarmonixMetasound/Nodes/TriggerToTransportNode.cpp
--- a/Engine/Plugins/Runtime/Harmonix/Source/HarmonixMetasound/Private/HarmonixMetasound/Nodes/TriggerToTransportNode.cpp
+++ b/Engine/Plugins/Runtime/Harmonix/Source/HarmonixMetasound/Private/HarmonixMetasound/Nodes/TriggerToTransportNode.cpp
@@ -84,7 +84,7 @@ namespace HarmonixMetasound
 			Info.Author           = PluginAuthor;
 			Info.PromptIfMissing  = PluginNodeMissingPrompt;
 			Info.DefaultInterface = GetVertexInterface();
-			Info.CategoryHierarchy.Emplace(NodeCategories::Music);
+			Info.CategoryHierarchy = { MetasoundNodeCategories::Harmonix, NodeCategories::Music };
 			return Info;
 		};
 
diff --git a/Engine/Plugins/Runtime/Harmonix/Source/HarmonixMetasound/Public/HarmonixMetasound/Common.h b/Engine/Plugins/Runtime/Harmonix/Source/HarmonixMetasound/Public/HarmonixMetasound/Common.h
--- a/Engine/Plugins/Runtime/Harmonix/Source/HarmonixMetasound/Public/HarmonixMetasound/Common.h
+++ b/Engine/Plugins/Runtime/Harmonix/Source/HarmonixMetasound/Public/HarmonixMetasound/Common.h
@@ -20,6 +20,7 @@ namespace HarmonixMetasound
 	{
 		HARMONIXMETASOUND_API extern const FText Harmonix;
 		HARMONIXMETASOUND_API extern const FText Modulation;
+		HARMONIXMETASOUND_API extern const FText Analysis;
 	}
 
 #define DECLARE_METASOUND_PARAM_EXTERN(NAME)      \
