midi -> MIDI in displayed strings

diff --git a/Engine/Plugins/Runtime/Harmonix/Source/HarmonixMetasound/Private/HarmonixMetasound/Analysis/MidiStreamVertexAnalyzer.cpp b/Engine/Plugins/Runtime/Harmonix/Source/HarmonixMetasound/Private/HarmonixMetasound/Analysis/MidiStreamVertexAnalyzer.cpp
--- a/Engine/Plugins/Runtime/Harmonix/Source/HarmonixMetasound/Private/HarmonixMetasound/Analysis/MidiStreamVertexAnalyzer.cpp
+++ b/Engine/Plugins/Runtime/Harmonix/Source/HarmonixMetasound/Private/HarmonixMetasound/Analysis/MidiStreamVertexAnalyzer.cpp
@@ -19,7 +19,7 @@ namespace HarmonixMetasound::Analysis
 
 	const Metasound::Frontend::FAnalyzerOutput& FMidiStreamVertexAnalyzer::FOutputs::GetValue()
 	{
-		static Metasound::Frontend::FAnalyzerOutput Value = { "LastMidiEvent", Metasound::GetMetasoundDataTypeName<FMidiEventInfo>() };
+		static Metasound::Frontend::FAnalyzerOutput Value = { "LastMIDIEvent", Metasound::GetMetasoundDataTypeName<FMidiEventInfo>() };
 		return Value;
 	}
 
diff --git a/Engine/Plugins/Runtime/Harmonix/Source/HarmonixMetasound/Private/HarmonixMetasound/Common.cpp b/Engine/Plugins/Runtime/Harmonix/Source/HarmonixMetasound/Private/HarmonixMetasound/Common.cpp
--- a/Engine/Plugins/Runtime/Harmonix/Source/HarmonixMetasound/Private/HarmonixMetasound/Common.cpp
+++ b/Engine/Plugins/Runtime/Harmonix/Source/HarmonixMetasound/Private/HarmonixMetasound/Common.cpp
@@ -22,18 +22,18 @@ namespace HarmonixMetasound
 		namespace Inputs
 		{
 			DEFINE_INPUT_METASOUND_PARAM(Enable,                        "Enable",                             "Enables processing.")
-			DEFINE_INPUT_METASOUND_PARAM(MidiStream,                    "Midi Stream",						  "Midi event stream.")
-			DEFINE_INPUT_METASOUND_PARAM(MidiChannelNumber,             "Midi Channel",                       "Midi channel to render (1-16, 0 for all).")
-			DEFINE_INPUT_METASOUND_PARAM(MidiChannelFilterSpecifier,    "Midi Channel Filter",                "Midi channels to process. '*' for all, or a string like '1,3,4-8' to identify specific channels.")
+			DEFINE_INPUT_METASOUND_PARAM(MidiStream,                    "MIDI Stream",						  "MIDI event stream.")
+			DEFINE_INPUT_METASOUND_PARAM(MidiChannelNumber,             "MIDI Channel",                       "MIDI channel to render (1-16, 0 for all).")
+			DEFINE_INPUT_METASOUND_PARAM(MidiChannelFilterSpecifier,    "MIDI Channel Filter",                "MIDI channels to process. '*' for all, or a string like '1,3,4-8' to identify specific channels.")
 			DEFINE_INPUT_METASOUND_PARAM(MidiTrackNumber,               "Track Number",                       "Track number (1 based).")
-			DEFINE_INPUT_METASOUND_PARAM(MidiTrackIndexFilterSpecifier, "Midi Track Filter",                  "Midi tracks to process. '*' for all, or a string like '1,3,4-8' to identify specific tracks.")
-			DEFINE_INPUT_METASOUND_PARAM(MinMidiNote,                   "Min Note #",                         "Minimum midi note number (0-127).")
-			DEFINE_INPUT_METASOUND_PARAM(MaxMidiNote,                   "Max Note #",                         "Maximum midi note number (0-127).")
+			DEFINE_INPUT_METASOUND_PARAM(MidiTrackIndexFilterSpecifier, "MIDI Track Filter",                  "MIDI tracks to process. '*' for all, or a string like '1,3,4-8' to identify specific tracks.")
+			DEFINE_INPUT_METASOUND_PARAM(MinMidiNote,                   "Min Note #",                         "Minimum MIDI note number (0-127).")
+			DEFINE_INPUT_METASOUND_PARAM(MaxMidiNote,                   "Max Note #",                         "Maximum MIDI note number (0-127).")
 			DEFINE_INPUT_METASOUND_PARAM(MinMidiVelocity,               "Min Velocity",                       "Minimum Velocity (0-127).")
 			DEFINE_INPUT_METASOUND_PARAM(MaxMidiVelocity,               "Max Velocity",                       "Maximum Velocity (0-127).")
-			DEFINE_INPUT_METASOUND_PARAM(MidiFileAsset,                 "Midi File",                          "Standard Midi file.")
+			DEFINE_INPUT_METASOUND_PARAM(MidiFileAsset,                 "MIDI File",                          "Standard MIDI file.")
 			DEFINE_INPUT_METASOUND_PARAM(Transport,                     "Transport",                          "Play, pause, continue, stop, etc.")
-			DEFINE_INPUT_METASOUND_PARAM(MidiClock,                     "Midi Clock",                         "Midi timing information.")
+			DEFINE_INPUT_METASOUND_PARAM(MidiClock,                     "MIDI Clock",                         "MIDI timing information.")
 			DEFINE_INPUT_METASOUND_PARAM(Tempo,                         "Tempo",                              "BPM. This is MIDI bpm... ALWAYS QUARTER NOTES PER MINUTE regardless of time signature!")
 			DEFINE_INPUT_METASOUND_PARAM(TimeSigNumerator,	            "Time Sig. Numerator",                "Time Signature Numerator")
 			DEFINE_INPUT_METASOUND_PARAM(TimeSigDenominator,            "Time Sig. Denominator",              "TimeSignature Denominator")
@@ -48,7 +48,7 @@ namespace HarmonixMetasound
 			DEFINE_INPUT_METASOUND_PARAM(TransportKill,                 "Kill",                               "Causes a transport kill request.")
 			DEFINE_INPUT_METASOUND_PARAM(TriggerSeek,                   "Trigger Seek",                       "Causes a transport seek request. Seek destination provided on a separate input.")
 			DEFINE_INPUT_METASOUND_PARAM(SeekDestination,               "Seek Target",                        "Destination for the next seek when it is triggered.")
-			DEFINE_INPUT_METASOUND_PARAM(PrerollBars,                   "Pre-roll Bars",                      "Number of bars to look back when seeking midi to find note-on messages that should hold over.")
+			DEFINE_INPUT_METASOUND_PARAM(PrerollBars,                   "Pre-roll Bars",                      "Number of bars to look back when seeking MIDI to find note-on messages that should hold over.")
 			DEFINE_INPUT_METASOUND_PARAM(GridSizeUnits,                 "Base Grid Size",                     "Base size of the qrid squares in musical terms.")
 			DEFINE_INPUT_METASOUND_PARAM(GridSizeMult,                  "Grid Size Multiplier",               "Grid size multiplier.")
 			DEFINE_INPUT_METASOUND_PARAM(OffsetUnits,                   "Base Offset Size ",                  "Base size of the qrid offset.")
@@ -57,8 +57,8 @@ namespace HarmonixMetasound
 			DEFINE_INPUT_METASOUND_PARAM(FloatBeat,			            "Beat",                               "Musical beat index. 0 based and in units of the time signature denomninator! The value is rounded to the nearest 32nd note.")
 			DEFINE_INPUT_METASOUND_PARAM(SynthPatch,                    "Patch",                              "Synthesizer patch asset.")
 			DEFINE_INPUT_METASOUND_PARAM(Transposition,                 "Transposition",                      "Transposition in semitones.")
-			DEFINE_INPUT_METASOUND_PARAM(ClockSpeedToPitch,             "Midi Clock Speed Affects Pitch",     "When true the clock speed of the incoming midi stream will affect the pitch of the synth.")
-			DEFINE_INPUT_METASOUND_PARAM(ClockSpeedToFrequency,         "Midi Clock Speed Affects Frequency", "When true the clock speed of the incoming midi stream will affect the calculated frequency.")
+			DEFINE_INPUT_METASOUND_PARAM(ClockSpeedToPitch,             "MIDI Clock Speed Affects Pitch",     "When true the clock speed of the incoming MIDI stream will affect the pitch of the synth.")
+			DEFINE_INPUT_METASOUND_PARAM(ClockSpeedToFrequency,         "MIDI Clock Speed Affects Frequency", "When true the clock speed of the incoming MIDI stream will affect the calculated frequency.")
 			DEFINE_INPUT_METASOUND_PARAM(LFOSyncType,					"LFO Sync Type",					  "Specifies if and how the LFO should sync to a clock");
 			DEFINE_INPUT_METASOUND_PARAM(LFOFrequency,				    "LFO Frequency",					  "The period of the LFO. If Sync Type is TempoSync, the unit is cycles per quarter note. Otherwise, the unit is Hz (cycles per second).");
 			DEFINE_INPUT_METASOUND_PARAM(LFOInvert,					    "Invert LFO",						  "Toggle to invert the LFO");
@@ -66,18 +66,21 @@ namespace HarmonixMetasound
 			DEFINE_INPUT_METASOUND_PARAM(AudioMono,					    "Audio Mono",					      "The mono audio.");
 			DEFINE_INPUT_METASOUND_PARAM(AudioLeft,					    "Audio Left",					      "The left channel audio.");
 			DEFINE_INPUT_METASOUND_PARAM(AudioRight,					"Audio Right",					      "The right channel audio.");
+			DEFINE_INPUT_METASOUND_PARAM(Timestamp,                     "Music Timestamp",                    "The Bar and Beat of interest.");
+			DEFINE_INPUT_METASOUND_PARAM(ShouldQuantizeTimestamp,       "Quantize",                           "Quantize the provide timestamp (Bar & Beat) to a musical subdivision.");
+			DEFINE_INPUT_METASOUND_PARAM(QuantizationUnit,              "Quantization Unit",                  "The music subdivision that a music timestamp should be quantized to.");
 		}
 		namespace Outputs
 		{
 			DEFINE_OUTPUT_METASOUND_PARAM(Transport,          "Transport",                   "Play, pause, continue, stop, etc.")
-			DEFINE_OUTPUT_METASOUND_PARAM(MidiClock,          "Midi Clock",                  "Midi timing information.")
-			DEFINE_OUTPUT_METASOUND_PARAM(MidiFileAsset,      "Midi File",                   "Standard Midi file.")
-			DEFINE_OUTPUT_METASOUND_PARAM(MidiStream,         "Midi Stream",				 "Midi event stream.")
-			DEFINE_OUTPUT_METASOUND_PARAM(NoteOn,             "Note On",                     "Midi note on.")
-			DEFINE_OUTPUT_METASOUND_PARAM(NoteOff,            "Note Off",                    "Midi note off.")
-			DEFINE_OUTPUT_METASOUND_PARAM(MidiNoteNumber,     "Midi Note #",                 "Midi note number (0 - 127).")
-			DEFINE_OUTPUT_METASOUND_PARAM(Frequency,          "Frequency",                   "Midi note number converted to frequency.")
-			DEFINE_OUTPUT_METASOUND_PARAM(MidiVelocity,       "Velocity",                    "Midi velocity (0 - 127).")
+			DEFINE_OUTPUT_METASOUND_PARAM(MidiClock,          "MIDI Clock",                  "MIDI timing information.")
+			DEFINE_OUTPUT_METASOUND_PARAM(MidiFileAsset,      "MIDI File",                   "Standard MIDI file.")
+			DEFINE_OUTPUT_METASOUND_PARAM(MidiStream,         "MIDI Stream",				 "MIDI event stream.")
+			DEFINE_OUTPUT_METASOUND_PARAM(NoteOn,             "Note On",                     "MIDI note on.")
+			DEFINE_OUTPUT_METASOUND_PARAM(NoteOff,            "Note Off",                    "MIDI note off.")
+			DEFINE_OUTPUT_METASOUND_PARAM(MidiNoteNumber,     "MIDI Note #",                 "MIDI note number (0 - 127).")
+			DEFINE_OUTPUT_METASOUND_PARAM(Frequency,          "Frequency",                   "MIDI note number converted to frequency.")
+			DEFINE_OUTPUT_METASOUND_PARAM(MidiVelocity,       "Velocity",                    "MIDI velocity (0 - 127).")
 			DEFINE_OUTPUT_METASOUND_PARAM(NormalizedVelocity, "Normalized Velocity",         "Velocity (0.0 - 1.0).")
 			DEFINE_OUTPUT_METASOUND_PARAM(TransportPrepare,   "Prepare",                     "Causes receiving nodes to prepare to play (pre-buffer, cache, etc.)")
 			DEFINE_OUTPUT_METASOUND_PARAM(TransportPlay,      "Play",                        "Causes a transport play request.")
diff --git a/Engine/Plugins/Runtime/Harmonix/Source/HarmonixMetasound/Private/HarmonixMetasound/Components/MetasoundMusicClockDriver.cpp b/Engine/Plugins/Runtime/Harmonix/Source/HarmonixMetasound/Private/HarmonixMetasound/Components/MetasoundMusicClockDriver.cpp
--- a/Engine/Plugins/Runtime/Harmonix/Source/HarmonixMetasound/Private/HarmonixMetasound/Components/MetasoundMusicClockDriver.cpp
+++ b/Engine/Plugins/Runtime/Harmonix/Source/HarmonixMetasound/Private/HarmonixMetasound/Components/MetasoundMusicClockDriver.cpp
@@ -200,7 +200,7 @@ void FMetasoundMusicClockDriver::TryToRegisterPlayCursor()
 			}
 			else
 			{
-				UE_LOG(LogMusicClock, Verbose, TEXT("Didn't find Midi Clock output named \"%s\" in the Metasound!"), *MetasoundOutputName.ToString());
+				UE_LOG(LogMusicClock, Verbose, TEXT("Didn't find MIDI Clock output named \"%s\" in the Metasound!"), *MetasoundOutputName.ToString());
 			}
 		}
 	}
diff --git a/Engine/Plugins/Runtime/Harmonix/Source/HarmonixMetasound/Private/HarmonixMetasound/Components/MetasoundMusicClockDriver.h b/Engine/Plugins/Runtime/Harmonix/Source/HarmonixMetasound/Private/HarmonixMetasound/Components/MetasoundMusicClockDriver.h
--- a/Engine/Plugins/Runtime/Harmonix/Source/HarmonixMetasound/Private/HarmonixMetasound/Components/MetasoundMusicClockDriver.h
+++ b/Engine/Plugins/Runtime/Harmonix/Source/HarmonixMetasound/Private/HarmonixMetasound/Components/MetasoundMusicClockDriver.h
@@ -22,7 +22,7 @@ public:
 	virtual const FSongMaps* GetCurrentSongMaps() const override;
 
 
-	bool ConnectToAudioComponentsMetasound(UAudioComponent* InAudioComponent, FName MetasoundOuputPinName = "Midi Clock");
+	bool ConnectToAudioComponentsMetasound(UAudioComponent* InAudioComponent, FName MetasoundOuputPinName = "MIDI Clock");
 
 protected:
 	void OnGeneratorAttached();
diff --git a/Engine/Plugins/Runtime/Harmonix/Source/HarmonixMetasound/Private/HarmonixMetasound/DataTypes/MidiAsset.cpp b/Engine/Plugins/Runtime/Harmonix/Source/HarmonixMetasound/Private/HarmonixMetasound/DataTypes/MidiAsset.cpp
--- a/Engine/Plugins/Runtime/Harmonix/Source/HarmonixMetasound/Private/HarmonixMetasound/DataTypes/MidiAsset.cpp
+++ b/Engine/Plugins/Runtime/Harmonix/Source/HarmonixMetasound/Private/HarmonixMetasound/DataTypes/MidiAsset.cpp
@@ -4,7 +4,7 @@
 
 #include "MetasoundDataTypeRegistrationMacro.h"
 
-REGISTER_METASOUND_DATATYPE(HarmonixMetasound::FMidiAsset, "MidiAsset", Metasound::ELiteralType::UObjectProxy, UMidiFile);
+REGISTER_METASOUND_DATATYPE(HarmonixMetasound::FMidiAsset, "MIDIAsset", Metasound::ELiteralType::UObjectProxy, UMidiFile);
 
 namespace HarmonixMetasound
 {
diff --git a/Engine/Plugins/Runtime/Harmonix/Source/HarmonixMetasound/Private/HarmonixMetasound/DataTypes/MidiClock.cpp b/Engine/Plugins/Runtime/Harmonix/Source/HarmonixMetasound/Private/HarmonixMetasound/DataTypes/MidiClock.cpp
--- a/Engine/Plugins/Runtime/Harmonix/Source/HarmonixMetasound/Private/HarmonixMetasound/DataTypes/MidiClock.cpp
+++ b/Engine/Plugins/Runtime/Harmonix/Source/HarmonixMetasound/Private/HarmonixMetasound/DataTypes/MidiClock.cpp
@@ -123,7 +123,7 @@ int32 SubdivisionToMidiTicks(const EMidiClockSubdivisionQuantization Division, c
 	}
 }
 
-REGISTER_METASOUND_DATATYPE(HarmonixMetasound::FMidiClock, "MidiClock")
+REGISTER_METASOUND_DATATYPE(HarmonixMetasound::FMidiClock, "MIDIClock")
 
 namespace HarmonixMetasound
 {
diff --git a/Engine/Plugins/Runtime/Harmonix/Source/HarmonixMetasound/Private/HarmonixMetasound/DataTypes/MidiControllerID.cpp b/Engine/Plugins/Runtime/Harmonix/Source/HarmonixMetasound/Private/HarmonixMetasound/DataTypes/MidiControllerID.cpp
--- a/Engine/Plugins/Runtime/Harmonix/Source/HarmonixMetasound/Private/HarmonixMetasound/DataTypes/MidiControllerID.cpp
+++ b/Engine/Plugins/Runtime/Harmonix/Source/HarmonixMetasound/Private/HarmonixMetasound/DataTypes/MidiControllerID.cpp
@@ -5,7 +5,7 @@
 #define LOCTEXT_NAMESPACE "HarmonixMetaSound"
 
 namespace Metasound {
-	DEFINE_METASOUND_ENUM_BEGIN(EStdMidiControllerID, FEnumStdMidiControllerID, "StdMidiControllerID")
+	DEFINE_METASOUND_ENUM_BEGIN(EStdMidiControllerID, FEnumStdMidiControllerID, "StdMIDIControllerID")
 		DEFINE_METASOUND_ENUM_ENTRY(EStdMidiControllerID::BankSelection, "EStdMidiControllerID_BankSelection_DisplayName", "0  Bank Selection", "EStdMidiControllerID_BankSelection_TT", "Allow switching between groups of 128 programs"),
 		DEFINE_METASOUND_ENUM_ENTRY(EStdMidiControllerID::ModWheel, "EStdMidiControllerID_ModWheel_DisplayName", "1  Modulation Wheel", "EStdMidiControllerID_ModWheel_TT", "Modulation Wheel (pitch, loudness, brighness etc.)"),
 		DEFINE_METASOUND_ENUM_ENTRY(EStdMidiControllerID::Breath, "EStdMidiControllerID_Breath_DisplayName", "2  Breath", "EStdMidiControllerID_Breath_TT", "Breath Controller"),
diff --git a/Engine/Plugins/Runtime/Harmonix/Source/HarmonixMetasound/Private/HarmonixMetasound/DataTypes/MidiEventInfo.cpp b/Engine/Plugins/Runtime/Harmonix/Source/HarmonixMetasound/Private/HarmonixMetasound/DataTypes/MidiEventInfo.cpp
--- a/Engine/Plugins/Runtime/Harmonix/Source/HarmonixMetasound/Private/HarmonixMetasound/DataTypes/MidiEventInfo.cpp
+++ b/Engine/Plugins/Runtime/Harmonix/Source/HarmonixMetasound/Private/HarmonixMetasound/DataTypes/MidiEventInfo.cpp
@@ -5,7 +5,7 @@
 #include "MetasoundDataTypeRegistrationMacro.h"
 #include "MetasoundOutput.h"
 
-REGISTER_METASOUND_DATATYPE(FMidiEventInfo, "MidiEventInfo");
+REGISTER_METASOUND_DATATYPE(FMidiEventInfo, "MIDIEventInfo");
 
 uint8 FMidiEventInfo::GetChannel() const
 {
diff --git a/Engine/Plugins/Runtime/Harmonix/Source/HarmonixMetasound/Private/HarmonixMetasound/DataTypes/MidiStepSequence.cpp b/Engine/Plugins/Runtime/Harmonix/Source/HarmonixMetasound/Private/HarmonixMetasound/DataTypes/MidiStepSequence.cpp
--- a/Engine/Plugins/Runtime/Harmonix/Source/HarmonixMetasound/Private/HarmonixMetasound/DataTypes/MidiStepSequence.cpp
+++ b/Engine/Plugins/Runtime/Harmonix/Source/HarmonixMetasound/Private/HarmonixMetasound/DataTypes/MidiStepSequence.cpp
@@ -5,9 +5,9 @@
 
 #include "Harmonix/PropertyUtility.h"
 
-DEFINE_LOG_CATEGORY_STATIC(LogMidiStepSequence, Log, All);
+DEFINE_LOG_CATEGORY_STATIC(LogMIDIStepSequence, Log, All);
 
-DEFINE_AUDIORENDERABLE_ASSET(HarmonixMetasound, FMidiStepSequenceAsset, MidiStepSequenceAsset, UMidiStepSequence)
+DEFINE_AUDIORENDERABLE_ASSET(HarmonixMetasound, FMidiStepSequenceAsset, MIDIStepSequenceAsset, UMidiStepSequence)
 
 bool FStepSequenceTable::IsPageBlank(int32 PageIdx) const
 {
diff --git a/Engine/Plugins/Runtime/Harmonix/Source/HarmonixMetasound/Private/HarmonixMetasound/DataTypes/MidiStream.cpp b/Engine/Plugins/Runtime/Harmonix/Source/HarmonixMetasound/Private/HarmonixMetasound/DataTypes/MidiStream.cpp
--- a/Engine/Plugins/Runtime/Harmonix/Source/HarmonixMetasound/Private/HarmonixMetasound/DataTypes/MidiStream.cpp
+++ b/Engine/Plugins/Runtime/Harmonix/Source/HarmonixMetasound/Private/HarmonixMetasound/DataTypes/MidiStream.cpp
@@ -6,7 +6,7 @@
 
 #include "MetasoundDataTypeRegistrationMacro.h"
 
-REGISTER_METASOUND_DATATYPE(HarmonixMetasound::FMidiStream, "MidiStream")
+REGISTER_METASOUND_DATATYPE(HarmonixMetasound::FMidiStream, "MIDIStream")
 
 namespace HarmonixMetasound
 {
@@ -146,7 +146,7 @@ namespace HarmonixMetasound
 	void FMidiStream::CopyMidiEvents(const FMidiStream& InStream, const FCopyMidiEventsPredicate& Predicate)
 	{
 		const TArray<FMidiStreamEvent>& MidiEvents = InStream.GetEventsInBlock();
-		if (!ensureAlwaysMsgf(EventsInBlock.Num()==0, TEXT("Existing midi events in midi stream are being destroyed during copy!")))
+		if (!ensureAlwaysMsgf(EventsInBlock.Num()==0, TEXT("Existing MIDI events in MIDI stream are being destroyed during copy!")))
 		{
 			EventsInBlock.Empty(FMath::Max(MidiEvents.Num(), 32));
 			ActiveVoices.Reset();
@@ -164,7 +164,7 @@ namespace HarmonixMetasound
 	{
 		if (GetMidiClockSource() && GetMidiClockSource()->Get() != InStream->GetMidiClockSource()->Get())
 		{
-			UE_LOG(LogMidiStreamDataType, Warning, TEXT("Cannot copy midi events from one stream to another if they have different clocks."));
+			UE_LOG(LogMidiStreamDataType, Warning, TEXT("Cannot copy MIDI events from one stream to another if they have different clocks."));
 			return;
 		}
 		if (EventsOnly)
@@ -194,7 +194,7 @@ namespace HarmonixMetasound
 			}
 			else
 			{
-				UE_LOG(LogMidiStreamDataType, Warning, TEXT("Cannot copy midi events from one stream to another if they have different clocks."));
+				UE_LOG(LogMidiStreamDataType, Warning, TEXT("Cannot copy MIDI events from one stream to another if they have different clocks."));
 			}
 		}
 	}
diff --git a/Engine/Plugins/Runtime/Harmonix/Source/HarmonixMetasound/Private/HarmonixMetasound/Nodes/FusionSamplerNode.cpp b/Engine/Plugins/Runtime/Harmonix/Source/HarmonixMetasound/Private/HarmonixMetasound/Nodes/FusionSamplerNode.cpp
--- a/Engine/Plugins/Runtime/Harmonix/Source/HarmonixMetasound/Private/HarmonixMetasound/Nodes/FusionSamplerNode.cpp
+++ b/Engine/Plugins/Runtime/Harmonix/Source/HarmonixMetasound/Private/HarmonixMetasound/Nodes/FusionSamplerNode.cpp
@@ -325,7 +325,7 @@ namespace HarmonixMetasound
 				Info.MajorVersion = 0;
 				Info.MinorVersion = 1;
 				Info.DisplayName = GetDisplayName();
-				Info.Description = METASOUND_LOCTEXT("FusionSamplerNode_Description", "Renders incoming midi stream using the Fusion Sampler configured with the specified patch.");
+				Info.Description = METASOUND_LOCTEXT("FusionSamplerNode_Description", "Renders incoming MIDI stream using the Fusion Sampler configured with the specified patch.");
 				Info.Author = PluginAuthor;
 				Info.PromptIfMissing = PluginNodeMissingPrompt;
 				Info.DefaultInterface = GetVertexInterface();
diff --git a/Engine/Plugins/Runtime/Harmonix/Source/HarmonixMetasound/Private/HarmonixMetasound/Nodes/MetronomeNode.cpp b/Engine/Plugins/Runtime/Harmonix/Source/HarmonixMetasound/Private/HarmonixMetasound/Nodes/MetronomeNode.cpp
--- a/Engine/Plugins/Runtime/Harmonix/Source/HarmonixMetasound/Private/HarmonixMetasound/Nodes/MetronomeNode.cpp
+++ b/Engine/Plugins/Runtime/Harmonix/Source/HarmonixMetasound/Private/HarmonixMetasound/Nodes/MetronomeNode.cpp
@@ -113,8 +113,8 @@ namespace HarmonixMetasound
 			Info.ClassName        = { HarmonixNodeNamespace, TEXT("Metronome"), TEXT("") };
 			Info.MajorVersion     = 0;
 			Info.MinorVersion     = 1;
-			Info.DisplayName      = METASOUND_LOCTEXT("MetronomeNode_DisplayName", "Metronome Midi Clock Generator");
-			Info.Description      = METASOUND_LOCTEXT("MetronomeNode_Description", "Provides a midi clock at the specified tempo and speed.");
+			Info.DisplayName      = METASOUND_LOCTEXT("MetronomeNode_DisplayName", "Metronome MIDI Clock Generator");
+			Info.Description      = METASOUND_LOCTEXT("MetronomeNode_Description", "Provides a MIDI clock at the specified tempo and speed.");
 			Info.Author           = PluginAuthor;
 			Info.PromptIfMissing  = PluginNodeMissingPrompt;
 			Info.DefaultInterface = GetVertexInterface();
diff --git a/Engine/Plugins/Runtime/Harmonix/Source/HarmonixMetasound/Private/HarmonixMetasound/Nodes/MidiCCTriggerNode.cpp b/Engine/Plugins/Runtime/Harmonix/Source/HarmonixMetasound/Private/HarmonixMetasound/Nodes/MidiCCTriggerNode.cpp
--- a/Engine/Plugins/Runtime/Harmonix/Source/HarmonixMetasound/Private/HarmonixMetasound/Nodes/MidiCCTriggerNode.cpp
+++ b/Engine/Plugins/Runtime/Harmonix/Source/HarmonixMetasound/Private/HarmonixMetasound/Nodes/MidiCCTriggerNode.cpp
@@ -267,8 +267,8 @@ namespace HarmonixMetasound::Nodes::MidiCCTriggerNode
 				Info.ClassName = GetClassName();
 				Info.MajorVersion = 0;
 				Info.MinorVersion = 1;
-				Info.DisplayName = METASOUND_LOCTEXT("MidiCCTriggerNode_DisplayName", "Midi CC Trigger");
-				Info.Description = METASOUND_LOCTEXT("MidiCCTriggerNode_Description", "Find the Midi CC messages in a midi stream and output them as triggers, floats, ints by the specified Midi track, Midi channel, and Midi Controller ID.");
+				Info.DisplayName = METASOUND_LOCTEXT("MIDICCTriggerNode_DisplayName", "MIDI CC Trigger");
+				Info.Description = METASOUND_LOCTEXT("MIDICCTriggerNode_Description", "Find the MIDI CC messages in a MIDI stream and output them as triggers, floats, ints by the specified MIDI track, MIDI channel, and MIDI Controller ID.");
 				Info.Author = PluginAuthor;
 				Info.PromptIfMissing = PluginNodeMissingPrompt;
 				Info.DefaultInterface = GetVertexInterface();
diff --git a/Engine/Plugins/Runtime/Harmonix/Source/HarmonixMetasound/Private/HarmonixMetasound/Nodes/MidiNoteTriggerNode.cpp b/Engine/Plugins/Runtime/Harmonix/Source/HarmonixMetasound/Private/HarmonixMetasound/Nodes/MidiNoteTriggerNode.cpp
--- a/Engine/Plugins/Runtime/Harmonix/Source/HarmonixMetasound/Private/HarmonixMetasound/Nodes/MidiNoteTriggerNode.cpp
+++ b/Engine/Plugins/Runtime/Harmonix/Source/HarmonixMetasound/Private/HarmonixMetasound/Nodes/MidiNoteTriggerNode.cpp
@@ -57,8 +57,8 @@ namespace HarmonixMetasound::Nodes::MidiNoteTriggerNode
 				Info.ClassName        = GetClassName();
 				Info.MajorVersion     = 1;
 				Info.MinorVersion     = 0;
-				Info.DisplayName      = METASOUND_LOCTEXT("MidiNoteTriggerNodeV1_DisplayName", "MIDI Note Trigger");
-				Info.Description      = METASOUND_LOCTEXT("MidiNoteTriggerNodeV1_Description", "Outputs triggers and info for incoming MIDI notes.");
+				Info.DisplayName      = METASOUND_LOCTEXT("MIDINoteTriggerNodeV1_DisplayName", "MIDI Note Trigger");
+				Info.Description      = METASOUND_LOCTEXT("MIDINoteTriggerNodeV1_Description", "Outputs triggers and info for incoming MIDI notes.");
 				Info.Author           = PluginAuthor;
 				Info.PromptIfMissing  = PluginNodeMissingPrompt;
 				Info.DefaultInterface = GetVertexInterface();
diff --git a/Engine/Plugins/Runtime/Harmonix/Source/HarmonixMetasound/Private/HarmonixMetasound/Nodes/MidiPlayerNode.cpp b/Engine/Plugins/Runtime/Harmonix/Source/HarmonixMetasound/Private/HarmonixMetasound/Nodes/MidiPlayerNode.cpp
--- a/Engine/Plugins/Runtime/Harmonix/Source/HarmonixMetasound/Private/HarmonixMetasound/Nodes/MidiPlayerNode.cpp
+++ b/Engine/Plugins/Runtime/Harmonix/Source/HarmonixMetasound/Private/HarmonixMetasound/Nodes/MidiPlayerNode.cpp
@@ -26,8 +26,8 @@ namespace HarmonixMetasound
 
 	namespace MidiPlayerNodePinNames
 	{
-		METASOUND_PARAM(KillVoicesOnSeek, "Kill Voices On Seek", "If true, a \"Kill All Voices\" midi message will be sent when seeking. Otherwise an \"All Notes Off\" will be sent which allows to ADSR release phases.")
-		METASOUND_PARAM(KillVoicesOnMidiChange, "Kill Voices On Midi File Change", "If true, a \"Kill All Voices\" midi message will be sent when the midi file asset is changed. Otherwise an \"All Notes Off\" will be sent which allows to ADSR release phases.")
+		METASOUND_PARAM(KillVoicesOnSeek, "Kill Voices On Seek", "If true, a \"Kill All Voices\" MIDI message will be sent when seeking. Otherwise an \"All Notes Off\" will be sent which allows to ADSR release phases.")
+		METASOUND_PARAM(KillVoicesOnMidiChange, "Kill Voices On MIDI File Change", "If true, a \"Kill All Voices\" MIDI message will be sent when the MIDI file asset is changed. Otherwise an \"All Notes Off\" will be sent which allows to ADSR release phases.")
 	}
 
 	class FMidiPlayerOperator : public TExecutableOperator<FMidiPlayerOperator>, public FMidiPlayCursor, public FMidiVoiceGeneratorBase, public FMusicTransportControllable
@@ -172,11 +172,11 @@ namespace HarmonixMetasound
 		auto InitNodeInfo = []() -> FNodeClassMetadata
 		{
 			FNodeClassMetadata Info;
-			Info.ClassName        = { HarmonixNodeNamespace, TEXT("MidiPlayer"), TEXT("") };
+			Info.ClassName        = { HarmonixNodeNamespace, TEXT("MIDIPlayer"), TEXT("") };
 			Info.MajorVersion     = 0;
 			Info.MinorVersion     = 1;
-			Info.DisplayName      = METASOUND_LOCTEXT("MidiPlayerNode_DisplayName", "Midi Player");
-			Info.Description      = METASOUND_LOCTEXT("MidiPlayerNode_Description", "Plays a standard midi file.");
+			Info.DisplayName      = METASOUND_LOCTEXT("MIDIPlayerNode_DisplayName", "MIDI Player");
+			Info.Description      = METASOUND_LOCTEXT("MIDIPlayerNode_Description", "Plays a standard MIDI file.");
 			Info.Author           = PluginAuthor;
 			Info.PromptIfMissing  = PluginNodeMissingPrompt;
 			Info.DefaultInterface = GetVertexInterface();
diff --git a/Engine/Plugins/Runtime/Harmonix/Source/HarmonixMetasound/Private/HarmonixMetasound/Nodes/MidiStreamLoggerNode.cpp b/Engine/Plugins/Runtime/Harmonix/Source/HarmonixMetasound/Private/HarmonixMetasound/Nodes/MidiStreamLoggerNode.cpp
--- a/Engine/Plugins/Runtime/Harmonix/Source/HarmonixMetasound/Private/HarmonixMetasound/Nodes/MidiStreamLoggerNode.cpp
+++ b/Engine/Plugins/Runtime/Harmonix/Source/HarmonixMetasound/Private/HarmonixMetasound/Nodes/MidiStreamLoggerNode.cpp
@@ -12,7 +12,7 @@
 #include "HarmonixMetasound/DataTypes/MidiStream.h"
 #include "HarmonixMetasound/DataTypes/MusicTransport.h"
 
-DEFINE_LOG_CATEGORY_STATIC(LogMidiStreamLogger, Log, All);
+DEFINE_LOG_CATEGORY_STATIC(LogMIDIStreamLogger, Log, All);
 
 #define LOCTEXT_NAMESPACE "HarmonixMetaSound"
 
@@ -66,11 +66,11 @@ namespace HarmonixMetasound
 		auto InitNodeInfo = []() -> FNodeClassMetadata
 		{
 			FNodeClassMetadata Info;
-			Info.ClassName        = { HarmonixNodeNamespace, TEXT("MidiStreamLogger"), TEXT("")};
+			Info.ClassName        = { HarmonixNodeNamespace, TEXT("MIDIStreamLogger"), TEXT("")};
 			Info.MajorVersion     = 0;
 			Info.MinorVersion     = 1;
-			Info.DisplayName      = METASOUND_LOCTEXT("MidiStreamLoggerNode_DisplayName", "MIDI Logger");
-			Info.Description      = METASOUND_LOCTEXT("MidiStreamLoggerNode_Description", "Receives MIDI messages and writes them to the log.");
+			Info.DisplayName      = METASOUND_LOCTEXT("MIDIStreamLoggerNode_DisplayName", "MIDI Logger");
+			Info.Description      = METASOUND_LOCTEXT("MIDIStreamLoggerNode_Description", "Receives MIDI messages and writes them to the log.");
 			Info.Author           = PluginAuthor;
 			Info.PromptIfMissing  = PluginNodeMissingPrompt;
 			Info.DefaultInterface = GetVertexInterface();
@@ -194,12 +194,12 @@ namespace HarmonixMetasound
 
 	void FMidiStreamLoggerOperator::DumpTransportEvent(const FMidiTimestampTransportState& TransportEvent)
 	{
-		UE_LOG(LogMidiStreamLogger, Log, TEXT("[%d (%f)] Transport: %s"), TransportEvent.BlockSampleFrameIndex, TransportEvent.BlockSampleFrameOffset, *FMusicTransportControllable::StateToString(TransportEvent.TransportState));
+		UE_LOG(LogMIDIStreamLogger, Log, TEXT("[%d (%f)] Transport: %s"), TransportEvent.BlockSampleFrameIndex, TransportEvent.BlockSampleFrameOffset, *FMusicTransportControllable::StateToString(TransportEvent.TransportState));
 	}
 
 	void FMidiStreamLoggerOperator::DumpMidiEvent(const FMidiStreamEvent& MidiEvent)
 	{
-		UE_LOG(LogMidiStreamLogger, Log, TEXT("[%d (%f)] Track: %d, Auth Tick %d, Render Tick %d, Ms Offset %f, Message: %s"),
+		UE_LOG(LogMIDIStreamLogger, Log, TEXT("[%d (%f)] Track: %d, Auth Tick %d, Render Tick %d, Ms Offset %f, Message: %s"),
 			MidiEvent.BlockSampleFrameIndex,
 			MidiEvent.BlockSampleFrameOffset,
 			MidiEvent.TrackIndex,
diff --git a/Engine/Plugins/Runtime/Harmonix/Source/HarmonixMetasound/Private/HarmonixMetasound/Nodes/MidiStreamMergeNode.cpp b/Engine/Plugins/Runtime/Harmonix/Source/HarmonixMetasound/Private/HarmonixMetasound/Nodes/MidiStreamMergeNode.cpp
--- a/Engine/Plugins/Runtime/Harmonix/Source/HarmonixMetasound/Private/HarmonixMetasound/Nodes/MidiStreamMergeNode.cpp
+++ b/Engine/Plugins/Runtime/Harmonix/Source/HarmonixMetasound/Private/HarmonixMetasound/Nodes/MidiStreamMergeNode.cpp
@@ -66,8 +66,8 @@ namespace HarmonixMetasound
 			Info.ClassName        = { HarmonixNodeNamespace, TEXT("MidiStreamMerge"), TEXT("")};
 			Info.MajorVersion     = 0;
 			Info.MinorVersion     = 1;
-			Info.DisplayName      = METASOUND_LOCTEXT("MidiStreamMergeNode_DisplayName", "Midi Merge");
-			Info.Description      = METASOUND_LOCTEXT("MidiStreamMergeNode_Description", "Combines two midi streams into a single stream.");
+			Info.DisplayName      = METASOUND_LOCTEXT("MIDIStreamMergeNode_DisplayName", "MIDI Merge");
+			Info.Description      = METASOUND_LOCTEXT("MIDIStreamMergeNode_Description", "Combines two MIDI streams into a single stream.");
 			Info.Author           = PluginAuthor;
 			Info.PromptIfMissing  = PluginNodeMissingPrompt;
 			Info.DefaultInterface = GetVertexInterface();
@@ -82,8 +82,8 @@ namespace HarmonixMetasound
 
 	namespace MergePinNames
 	{
-		METASOUND_PARAM(InputMidiStreamA, "Midi Stream A", "The first midi stream to merge");
-		METASOUND_PARAM(InputMidiStreamB, "Midi Stream B", "The second midi stream to merge");
+		METASOUND_PARAM(InputMidiStreamA, "MIDI Stream A", "The first MIDI stream to merge");
+		METASOUND_PARAM(InputMidiStreamB, "MIDI Stream B", "The second MIDI stream to merge");
 	}
 
 	const FVertexInterface& FMidiStreamMergeOperator::GetVertexInterface()
diff --git a/Engine/Plugins/Runtime/Harmonix/Source/HarmonixMetasound/Private/HarmonixMetasound/Nodes/MidiStreamSelectNode.cpp b/Engine/Plugins/Runtime/Harmonix/Source/HarmonixMetasound/Private/HarmonixMetasound/Nodes/MidiStreamSelectNode.cpp
--- a/Engine/Plugins/Runtime/Harmonix/Source/HarmonixMetasound/Private/HarmonixMetasound/Nodes/MidiStreamSelectNode.cpp
+++ b/Engine/Plugins/Runtime/Harmonix/Source/HarmonixMetasound/Private/HarmonixMetasound/Nodes/MidiStreamSelectNode.cpp
@@ -13,7 +13,7 @@
 #include "HarmonixMetasound/DataTypes/MusicTransport.h"
 #include "HarmonixMidi/MidiMsg.h"
 
-DEFINE_LOG_CATEGORY_STATIC(LogMidiStreamSelect, Log, All);
+DEFINE_LOG_CATEGORY_STATIC(LogMIDIStreamSelect, Log, All);
 
 #define LOCTEXT_NAMESPACE "HaronixMetaSound"
 
@@ -81,8 +81,8 @@ namespace HarmonixMetasound
 			Info.ClassName = { HarmonixNodeNamespace, TEXT("MidiStreamSelect"), TEXT("") };
 			Info.MajorVersion = 0;
 			Info.MinorVersion = 1;
-			Info.DisplayName = METASOUND_LOCTEXT("MidiStreamSelectNode_DisplayName", "Midi Select");
-			Info.Description = METASOUND_LOCTEXT("MidiStreamSelectNode_Description", "Copies the output of one of the inputs to the output.");
+			Info.DisplayName = METASOUND_LOCTEXT("MIDIStreamSelectNode_DisplayName", "MIDI Select");
+			Info.Description = METASOUND_LOCTEXT("MIDIStreamSelectNode_Description", "Copies the output of one of the inputs to the output.");
 			Info.Author = PluginAuthor;
 			Info.PromptIfMissing = PluginNodeMissingPrompt;
 			Info.DefaultInterface = GetVertexInterface();
@@ -97,11 +97,11 @@ namespace HarmonixMetasound
 
 	namespace SelectPinNames
 	{
-		METASOUND_PARAM(InputMidiStreamA, "Midi Stream A", "The first midi stream to select.");
-		METASOUND_PARAM(InputMidiStreamB, "Midi Stream B", "The second midi stream to select.");
-		METASOUND_PARAM(InputStreamIndex, "Index", "The index of the midi stream to copy.");
-		METASOUND_PARAM(InputImmediateNoteOff, "Immediately Note Off", "Should all active midi in the stream turn off when the index changes?");
-		METASOUND_PARAM(InputCopyInactiveNoteOffs, "Copy Inactive Note Offs", "Should note offs for active midi be copied from inactive streams?");
+		METASOUND_PARAM(InputMidiStreamA, "MIDI Stream A", "The first MIDI stream to select.");
+		METASOUND_PARAM(InputMidiStreamB, "MIDI Stream B", "The second MIDI stream to select.");
+		METASOUND_PARAM(InputStreamIndex, "Index", "The index of the MIDI stream to copy.");
+		METASOUND_PARAM(InputImmediateNoteOff, "Immediately Note Off", "Should all active MIDI in the stream turn off when the index changes?");
+		METASOUND_PARAM(InputCopyInactiveNoteOffs, "Copy Inactive Note Offs", "Should note offs for active MIDI be copied from inactive streams?");
 	}
 
 	const FVertexInterface& FMidiStreamSelectOperator::GetVertexInterface()
diff --git a/Engine/Plugins/Runtime/Harmonix/Source/HarmonixMetasound/Private/HarmonixMetasound/Nodes/MidiStreamTrackFilterNode.cpp b/Engine/Plugins/Runtime/Harmonix/Source/HarmonixMetasound/Private/HarmonixMetasound/Nodes/MidiStreamTrackFilterNode.cpp
--- a/Engine/Plugins/Runtime/Harmonix/Source/HarmonixMetasound/Private/HarmonixMetasound/Nodes/MidiStreamTrackFilterNode.cpp
+++ b/Engine/Plugins/Runtime/Harmonix/Source/HarmonixMetasound/Private/HarmonixMetasound/Nodes/MidiStreamTrackFilterNode.cpp
@@ -61,8 +61,8 @@ namespace HarmonixMetasound::Nodes::MidiTrackFilter
 				Info.ClassName        = GetClassName();
 				Info.MajorVersion     = 1;
 				Info.MinorVersion     = 0;
-				Info.DisplayName      = METASOUND_LOCTEXT("MidiStreamTrackFilterNodeV1_DisplayName", "Midi Stream Track Filter");
-				Info.Description      = METASOUND_LOCTEXT("MidiStreamTrackFilterNodeV1_Description", "Duplicates the incoming midi stream to its output after filtering tracks from the input Midi Stream");
+				Info.DisplayName      = METASOUND_LOCTEXT("MIDIStreamTrackFilterNodeV1_DisplayName", "MIDI Stream Track Filter");
+				Info.Description      = METASOUND_LOCTEXT("MIDIStreamTrackFilterNodeV1_Description", "Duplicates the incoming MIDI stream to its output after filtering tracks from the input MIDI Stream");
 				Info.Author           = PluginAuthor;
 				Info.PromptIfMissing  = PluginNodeMissingPrompt;
 				Info.DefaultInterface = GetVertexInterface();
diff --git a/Engine/Plugins/Runtime/Harmonix/Source/HarmonixMetasound/Private/HarmonixMetasound/Nodes/MidiStreamTransposerNode.cpp b/Engine/Plugins/Runtime/Harmonix/Source/HarmonixMetasound/Private/HarmonixMetasound/Nodes/MidiStreamTransposerNode.cpp
--- a/Engine/Plugins/Runtime/Harmonix/Source/HarmonixMetasound/Private/HarmonixMetasound/Nodes/MidiStreamTransposerNode.cpp
+++ b/Engine/Plugins/Runtime/Harmonix/Source/HarmonixMetasound/Private/HarmonixMetasound/Nodes/MidiStreamTransposerNode.cpp
@@ -54,8 +54,8 @@ namespace HarmonixMetasound::Nodes::MidiNoteTranspose
 				Info.ClassName        = { HarmonixNodeNamespace, TEXT("MidiStreamTransposer"), TEXT("")};
 				Info.MajorVersion     = 1;
 				Info.MinorVersion     = 0;
-				Info.DisplayName      = METASOUND_LOCTEXT("MidiNoteTransposeNode_DisplayName", "MIDI Note Transpose");
-				Info.Description      = METASOUND_LOCTEXT("MidiNoteTransposeNode_Description", "Duplicates the incoming MIDI stream to its output with the note on/off messages transposed by the specified number of semitones.");
+				Info.DisplayName      = METASOUND_LOCTEXT("MIDINoteTransposeNode_DisplayName", "MIDI Note Transpose");
+				Info.Description      = METASOUND_LOCTEXT("MIDINoteTransposeNode_Description", "Duplicates the incoming MIDI stream to its output with the note on/off messages transposed by the specified number of semitones.");
 				Info.Author           = PluginAuthor;
 				Info.PromptIfMissing  = PluginNodeMissingPrompt;
 				Info.DefaultInterface = GetVertexInterface();
diff --git a/Engine/Plugins/Runtime/Harmonix/Source/HarmonixMetasound/Private/HarmonixMetasound/Nodes/MidiTextTriggerNode.cpp b/Engine/Plugins/Runtime/Harmonix/Source/HarmonixMetasound/Private/HarmonixMetasound/Nodes/MidiTextTriggerNode.cpp
--- a/Engine/Plugins/Runtime/Harmonix/Source/HarmonixMetasound/Private/HarmonixMetasound/Nodes/MidiTextTriggerNode.cpp
+++ b/Engine/Plugins/Runtime/Harmonix/Source/HarmonixMetasound/Private/HarmonixMetasound/Nodes/MidiTextTriggerNode.cpp
@@ -13,7 +13,7 @@
 #include "HarmonixMetasound/DataTypes/MidiStream.h"
 #include "HarmonixMetasound/DataTypes/MusicTransport.h"
 
-#define LOCTEXT_NAMESPACE "HarmonixMetaSound_MidiTextTriggerNode"
+#define LOCTEXT_NAMESPACE "HarmonixMetaSound_MIDITextTriggerNode"
 
 namespace HarmonixMetasound
 {
@@ -36,8 +36,8 @@ namespace HarmonixMetasound
 				Info.ClassName        = { HarmonixNodeNamespace, TEXT("MidiTextTrigger"), TEXT("")};
 				Info.MajorVersion     = 1;
 				Info.MinorVersion     = 0;
-				Info.DisplayName      = METASOUND_LOCTEXT("MidiTextTriggerNode_DisplayName", "MIDI Text Trigger");
-				Info.Description      = METASOUND_LOCTEXT("MidiTextTriggerNode_Description", "Receives a MIDI stream, filters for the desired messages, and outputs triggers.");
+				Info.DisplayName      = METASOUND_LOCTEXT("MIDITextTriggerNode_DisplayName", "MIDI Text Trigger");
+				Info.Description      = METASOUND_LOCTEXT("MIDITextTriggerNode_Description", "Receives a MIDI stream, filters for the desired messages, and outputs triggers.");
 				Info.Author           = PluginAuthor;
 				Info.PromptIfMissing  = PluginNodeMissingPrompt;
 				Info.DefaultInterface = GetVertexInterface();
diff --git a/Engine/Plugins/Runtime/Harmonix/Source/HarmonixMetasound/Public/HarmonixMetasound/Common.h b/Engine/Plugins/Runtime/Harmonix/Source/HarmonixMetasound/Public/HarmonixMetasound/Common.h
--- a/Engine/Plugins/Runtime/Harmonix/Source/HarmonixMetasound/Public/HarmonixMetasound/Common.h
+++ b/Engine/Plugins/Runtime/Harmonix/Source/HarmonixMetasound/Public/HarmonixMetasound/Common.h
@@ -105,6 +105,9 @@ namespace HarmonixMetasound
 			DECLARE_METASOUND_PARAM_EXTERN(AudioMono);
 			DECLARE_METASOUND_PARAM_EXTERN(AudioLeft);
 			DECLARE_METASOUND_PARAM_EXTERN(AudioRight);
+			DECLARE_METASOUND_PARAM_EXTERN(Timestamp);
+			DECLARE_METASOUND_PARAM_EXTERN(ShouldQuantizeTimestamp);
+			DECLARE_METASOUND_PARAM_EXTERN(QuantizationUnit);
 		}
 		namespace Outputs
 		{
diff --git a/Engine/Plugins/Runtime/Harmonix/Source/HarmonixMetasound/Public/HarmonixMetasound/Components/MusicClockComponent.h b/Engine/Plugins/Runtime/Harmonix/Source/HarmonixMetasound/Public/HarmonixMetasound/Components/MusicClockComponent.h
--- a/Engine/Plugins/Runtime/Harmonix/Source/HarmonixMetasound/Public/HarmonixMetasound/Components/MusicClockComponent.h
+++ b/Engine/Plugins/Runtime/Harmonix/Source/HarmonixMetasound/Public/HarmonixMetasound/Components/MusicClockComponent.h
@@ -56,7 +56,7 @@ public:
 	EMusicClockDriveMethod DriveMethod = EMusicClockDriveMethod::MetaSound;
 
 	UPROPERTY(EditDefaultsOnly, Category = "MusicClock", meta = (EditCondition = "DriveMethod == EMusicClockDriveMethod::MetaSound"))
-	FName MetasoundOutputName = "Midi Clock";
+	FName MetasoundOutputName = "MIDI Clock";
 
 	UPROPERTY(EditAnywhere, BlueprintReadWrite, Category = "MusicClock", meta = (EditCondition = "DriveMethod == EMusicClockDriveMethod::MetaSound"))
 	TObjectPtr<UAudioComponent> MetasoundsAudioComponent;
@@ -301,7 +301,7 @@ public:
 	float MeasureSpanProgress(const FMusicalTimeSpan& Span, ECalibratedMusicTimebase Timebase = ECalibratedMusicTimebase::VideoRenderTime) const;
 
 	UFUNCTION(BlueprintCallable, Category = "Audio|MusicClock", meta = (WorldContext = "WorldContextObject", DefaultToSelf = "WorldContextObject"))
-	static UMusicClockComponent* CreateMetasoundDrivenMusicClock(UObject* WorldContextObject, UAudioComponent* InAudioComponent, FName MetasoundOuputPinName = "Midi Clock", bool Start = true);
+	static UMusicClockComponent* CreateMetasoundDrivenMusicClock(UObject* WorldContextObject, UAudioComponent* InAudioComponent, FName MetasoundOuputPinName = "MIDI Clock", bool Start = true);
 
 	UFUNCTION(BlueprintCallable, Category = "Audio|MusicClock", meta = (WorldContext = "WorldContextObject", DefaultToSelf = "WorldContextObject"))
 	static UMusicClockComponent* CreateWallClockDrivenMusicClock(UObject* WorldContextObject, UMidiFile* WithTempoMap, bool Start = true);
diff --git a/Engine/Plugins/Runtime/Harmonix/Source/HarmonixMetasound/Public/HarmonixMetasound/DataTypes/MidiEventInfo.h b/Engine/Plugins/Runtime/Harmonix/Source/HarmonixMetasound/Public/HarmonixMetasound/DataTypes/MidiEventInfo.h
--- a/Engine/Plugins/Runtime/Harmonix/Source/HarmonixMetasound/Public/HarmonixMetasound/DataTypes/MidiEventInfo.h
+++ b/Engine/Plugins/Runtime/Harmonix/Source/HarmonixMetasound/Public/HarmonixMetasound/DataTypes/MidiEventInfo.h
@@ -9,7 +9,7 @@
 
 #include "MidiEventInfo.generated.h"
 
-USTRUCT(BlueprintType)
+USTRUCT(BlueprintType, Meta = (DisplayName = "MIDI Event Info"))
 struct HARMONIXMETASOUND_API FMidiEventInfo
 {
 	GENERATED_BODY()
diff --git a/Engine/Plugins/Runtime/Harmonix/Source/HarmonixMetasound/Public/HarmonixMetasound/DataTypes/MidiStepSequence.h b/Engine/Plugins/Runtime/Harmonix/Source/HarmonixMetasound/Public/HarmonixMetasound/DataTypes/MidiStepSequence.h
--- a/Engine/Plugins/Runtime/Harmonix/Source/HarmonixMetasound/Public/HarmonixMetasound/DataTypes/MidiStepSequence.h
+++ b/Engine/Plugins/Runtime/Harmonix/Source/HarmonixMetasound/Public/HarmonixMetasound/DataTypes/MidiStepSequence.h
@@ -205,7 +205,7 @@ DECLARE_AUDIORENDERABLE_ASSET(HarmonixMetasound, FMidiStepSequenceAsset, FStepSe
 
 // This class represents a step sequence table. It is used by the MetasSound Step Sequence node
 // to generate midi note on/off messages.
-UCLASS(BlueprintType)
+UCLASS(BlueprintType, Category = "Music", Meta = (DisplayName = "MIDI Step Sequence"))
 class HARMONIXMETASOUND_API UMidiStepSequence : public UObject, public IAudioProxyDataFactory
 {
 	GENERATED_BODY()
diff --git a/Engine/Plugins/Runtime/Harmonix/Source/HarmonixMetasoundEditor/Private/AssetDefinition_MidiStepSequence.cpp b/Engine/Plugins/Runtime/Harmonix/Source/HarmonixMetasoundEditor/Private/AssetDefinition_MidiStepSequence.cpp
--- a/Engine/Plugins/Runtime/Harmonix/Source/HarmonixMetasoundEditor/Private/AssetDefinition_MidiStepSequence.cpp
+++ b/Engine/Plugins/Runtime/Harmonix/Source/HarmonixMetasoundEditor/Private/AssetDefinition_MidiStepSequence.cpp
@@ -10,7 +10,7 @@ TSoftClassPtr<UObject> UAssetDefinition_MidiStepSequence::GetAssetClass() const
 
 FText UAssetDefinition_MidiStepSequence::GetAssetDisplayName() const
 {
-	return NSLOCTEXT("AssetTypeActions", "MidiStepSequenceDefinition", "Midi Step Sequence");
+	return NSLOCTEXT("AssetTypeActions", "MIDIStepSequenceDefinition", "MIDI Step Sequence");
 }
 
 FLinearColor  UAssetDefinition_MidiStepSequence::GetAssetColor() const
diff --git a/Engine/Plugins/Runtime/Harmonix/Source/HarmonixMetasoundEditor/Private/HarmonixMetasoundEditorModule.cpp b/Engine/Plugins/Runtime/Harmonix/Source/HarmonixMetasoundEditor/Private/HarmonixMetasoundEditorModule.cpp
--- a/Engine/Plugins/Runtime/Harmonix/Source/HarmonixMetasoundEditor/Private/HarmonixMetasoundEditorModule.cpp
+++ b/Engine/Plugins/Runtime/Harmonix/Source/HarmonixMetasoundEditor/Private/HarmonixMetasoundEditorModule.cpp
@@ -17,8 +17,8 @@ void FHarmonixMetasoundEditorModule::StartupModule()
 {
 	const HarmonixMetasoundEditor::FSlateStyle& Style = HarmonixMetasoundEditor::FSlateStyle::Get();
 	Metasound::Editor::IMetasoundEditorModule& MetasoundEditorModule = FModuleManager::GetModuleChecked<Metasound::Editor::IMetasoundEditorModule>("MetasoundEditor");
-	MetasoundEditorModule.RegisterPinType("MidiAsset");
-	MetasoundEditorModule.RegisterPinType("MidiStepSequenceAsset");
+	MetasoundEditorModule.RegisterPinType("MIDIAsset");
+	MetasoundEditorModule.RegisterPinType("MIDIStepSequenceAsset");
 	MetasoundEditorModule.RegisterPinType("FusionPatchAsset");
 	MetasoundEditorModule.RegisterPinType("Enum:SubdivisionQuantizationType","Int32");
 	MetasoundEditorModule.RegisterPinType("Enum:DelayFilterType", "Int32");
@@ -27,10 +27,10 @@ void FHarmonixMetasoundEditorModule::StartupModule()
 	MetasoundEditorModule.RegisterPinType("Enum:DistortionType", "Int32");
 	MetasoundEditorModule.RegisterPinType("Enum:Harmonix:BiquadFilterType", "Int32");
 	MetasoundEditorModule.RegisterPinType("Enum:Distortion:FilterPasses", "Int32");
-	MetasoundEditorModule.RegisterPinType("Enum:StdMidiControllerID", "Int32");
+	MetasoundEditorModule.RegisterPinType("Enum:StdMIDIControllerID", "Int32");
 
-	MetasoundEditorModule.RegisterPinType("MidiStream", {}, {}, Style.GetMidiStreamConnectedIcon(), Style.GetMidiStreamDisconnectedIcon());
-	MetasoundEditorModule.RegisterPinType("MidiClock", {}, {}, Style.GetMidiClockConnectedIcon(), Style.GetMidiClockDisconnectedIcon());
+	MetasoundEditorModule.RegisterPinType("MIDIStream", {}, {}, Style.GetMidiStreamConnectedIcon(), Style.GetMidiStreamDisconnectedIcon());
+	MetasoundEditorModule.RegisterPinType("MIDIClock", {}, {}, Style.GetMidiClockConnectedIcon(), Style.GetMidiClockDisconnectedIcon());
 	MetasoundEditorModule.RegisterPinType("MusicTransport", {}, {}, Style.GetTransportConnectedIcon(), Style.GetTransportDisconnectedIcon());
 
 	FPropertyEditorModule& PropertyModule = FModuleManager::LoadModuleChecked<FPropertyEditorModule>("PropertyEditor");
diff --git a/Engine/Plugins/Runtime/Harmonix/Source/HarmonixMetasoundEditor/Private/MidiStepSequenceFactory.cpp b/Engine/Plugins/Runtime/Harmonix/Source/HarmonixMetasoundEditor/Private/MidiStepSequenceFactory.cpp
--- a/Engine/Plugins/Runtime/Harmonix/Source/HarmonixMetasoundEditor/Private/MidiStepSequenceFactory.cpp
+++ b/Engine/Plugins/Runtime/Harmonix/Source/HarmonixMetasoundEditor/Private/MidiStepSequenceFactory.cpp
@@ -9,6 +9,11 @@ UMidiStepSequenceFactory::UMidiStepSequenceFactory()
 	bCreateNew = true;
 }
 
+FText UMidiStepSequenceFactory::GetDisplayName() const
+{
+	return NSLOCTEXT("MIDI", "MIDIStepSequenceFactoryName", "MIDI Step Sequence");
+}
+
 UObject* UMidiStepSequenceFactory::FactoryCreateNew(UClass* Class, UObject* InParent, FName Name, EObjectFlags Flags, UObject* Context, FFeedbackContext* Warn)
 {
 	return NewObject<UMidiStepSequence>(InParent, Class, Name, Flags, Context);
diff --git a/Engine/Plugins/Runtime/Harmonix/Source/HarmonixMetasoundEditor/Private/MidiStepSequenceFactory.h b/Engine/Plugins/Runtime/Harmonix/Source/HarmonixMetasoundEditor/Private/MidiStepSequenceFactory.h
--- a/Engine/Plugins/Runtime/Harmonix/Source/HarmonixMetasoundEditor/Private/MidiStepSequenceFactory.h
+++ b/Engine/Plugins/Runtime/Harmonix/Source/HarmonixMetasoundEditor/Private/MidiStepSequenceFactory.h
@@ -11,6 +11,9 @@ class UMidiStepSequenceFactory : public UFactory
 	GENERATED_BODY()
 public:
 	UMidiStepSequenceFactory();
-	UObject* FactoryCreateNew(UClass* Class, UObject* InParent, FName Name, EObjectFlags Flags, UObject* Context, FFeedbackContext* Warn);
+	//~ BEGIN UFactory interface
+	virtual FText GetDisplayName() const override;
+	virtual UObject* FactoryCreateNew(UClass* Class, UObject* InParent, FName Name, EObjectFlags Flags, UObject* Context, FFeedbackContext* Warn) override;
+	//~ END UFactory interface
 };
 
diff --git a/Engine/Plugins/Runtime/Harmonix/Source/HarmonixMidi/Public/HarmonixMidi/Blueprint/MidiNote.h b/Engine/Plugins/Runtime/Harmonix/Source/HarmonixMidi/Public/HarmonixMidi/Blueprint/MidiNote.h
--- a/Engine/Plugins/Runtime/Harmonix/Source/HarmonixMidi/Public/HarmonixMidi/Blueprint/MidiNote.h
+++ b/Engine/Plugins/Runtime/Harmonix/Source/HarmonixMidi/Public/HarmonixMidi/Blueprint/MidiNote.h
@@ -12,7 +12,7 @@
 * uint8 NoteNumber
 * ranging from Midi::kMinNote to Midi::kMaxNote [0 - 127]
 */
-USTRUCT(BlueprintType)
+USTRUCT(BlueprintType, Meta = (DisplayName = "MIDI Note"))
 struct HARMONIXMIDI_API FMidiNote
 {
 	GENERATED_BODY()
diff --git a/Engine/Plugins/Runtime/Harmonix/Source/HarmonixMidi/Public/HarmonixMidi/MidiFile.h b/Engine/Plugins/Runtime/Harmonix/Source/HarmonixMidi/Public/HarmonixMidi/MidiFile.h
--- a/Engine/Plugins/Runtime/Harmonix/Source/HarmonixMidi/Public/HarmonixMidi/MidiFile.h
+++ b/Engine/Plugins/Runtime/Harmonix/Source/HarmonixMidi/Public/HarmonixMidi/MidiFile.h
@@ -24,7 +24,7 @@ enum class EMidiFileLengthConformOption : uint8
 	Nearest
 };
 
-USTRUCT(BlueprintType)
+USTRUCT(BlueprintType, Meta = (DisplayName = "MIDI File Data"))
 struct HARMONIXMIDI_API FMidiFileData
 {
 	GENERATED_BODY()
@@ -83,7 +83,7 @@ struct TStructOpsTypeTraits<FMidiFileData> : public TStructOpsTypeTraitsBase2<FM
  * This class can handle loading and saving standard midi files, as well
  * as serializing itself to standard Unreal Engine FArchives.
  */
-UCLASS(BlueprintType, Category="Music")
+UCLASS(BlueprintType, Category="Music", Meta = (DisplayName = "Standard MIDI File"))
 class HARMONIXMIDI_API UMidiFile : public UObject, public IAudioProxyDataFactory
 {
 	GENERATED_BODY()
diff --git a/Engine/Plugins/Runtime/Harmonix/Source/HarmonixMidi/Public/HarmonixMidi/MidiSongPos.h b/Engine/Plugins/Runtime/Harmonix/Source/HarmonixMidi/Public/HarmonixMidi/MidiSongPos.h
--- a/Engine/Plugins/Runtime/Harmonix/Source/HarmonixMidi/Public/HarmonixMidi/MidiSongPos.h
+++ b/Engine/Plugins/Runtime/Harmonix/Source/HarmonixMidi/Public/HarmonixMidi/MidiSongPos.h
@@ -13,7 +13,7 @@ struct FSongMaps;
 /////////////////////////////////////////////////////////////////////////////
 // Position within a song (midi info)
 //
-USTRUCT(BlueprintType)
+USTRUCT(BlueprintType, Meta = (DisplayName = "MIDI Song Position"))
 struct HARMONIXMIDI_API FMidiSongPos
 {
 	GENERATED_BODY()
diff --git a/Engine/Plugins/Runtime/Harmonix/Source/HarmonixMidi/Public/HarmonixMidi/MidiTrack.h b/Engine/Plugins/Runtime/Harmonix/Source/HarmonixMidi/Public/HarmonixMidi/MidiTrack.h
--- a/Engine/Plugins/Runtime/Harmonix/Source/HarmonixMidi/Public/HarmonixMidi/MidiTrack.h
+++ b/Engine/Plugins/Runtime/Harmonix/Source/HarmonixMidi/Public/HarmonixMidi/MidiTrack.h
@@ -14,7 +14,7 @@ class FMidiWriter;
 	* It can be created dynamically or be the end result of importing a 
 	* standard midi file.
 	*/
-USTRUCT(BlueprintType)
+USTRUCT(BlueprintType, Meta = (DisplayName = "MIDI Track"))
 struct HARMONIXMIDI_API FMidiTrack
 {
 	GENERATED_BODY()
diff --git a/Engine/Plugins/Runtime/Harmonix/Source/HarmonixMidiEditor/Private/AssetDefinition_MidiFile.cpp b/Engine/Plugins/Runtime/Harmonix/Source/HarmonixMidiEditor/Private/AssetDefinition_MidiFile.cpp
--- a/Engine/Plugins/Runtime/Harmonix/Source/HarmonixMidiEditor/Private/AssetDefinition_MidiFile.cpp
+++ b/Engine/Plugins/Runtime/Harmonix/Source/HarmonixMidiEditor/Private/AssetDefinition_MidiFile.cpp
@@ -21,7 +21,7 @@ TSoftClassPtr<UObject> UAssetDefinition_MidiFile::GetAssetClass() const
 
 FText UAssetDefinition_MidiFile::GetAssetDisplayName() const
 {
-	return NSLOCTEXT("AssetTypeActions", "MidiFileDefinition", "Midi File");
+	return NSLOCTEXT("AssetTypeActions", "MIDIFileDefinition", "Standard MIDI File");
 }
 
 FLinearColor  UAssetDefinition_MidiFile::GetAssetColor() const
@@ -49,8 +49,8 @@ void UAssetDefinition_MidiFile::RegisterContextMenu()
 	Section.AddDynamicEntry("MidiFile_ExportMid",
 		FNewToolMenuSectionDelegate::CreateLambda([](FToolMenuSection& InSection)
 			{
-				const TAttribute<FText> Label = LOCTEXT("MidiFile_ExportMid", "Export Standard Midi File (.mid)");
-				const TAttribute<FText> ToolTip = LOCTEXT("MidiFile_ExportMidToolTip", "Exports standard midi file(s)");
+				const TAttribute<FText> Label = LOCTEXT("MidiFile_ExportMid", "Export Standard MIDI File (.mid)");
+				const TAttribute<FText> ToolTip = LOCTEXT("MidiFile_ExportMidToolTip", "Exports standard MIDI file(s)");
 				const FSlateIcon Icon = FSlateIcon(FAppStyle::GetAppStyleSetName(), "ClassIcon.MidiFile");
 				const FToolMenuExecuteAction UIAction = FToolMenuExecuteAction::CreateStatic(&UAssetDefinition_MidiFile::ExecuteExportMidiFile);
 				InSection.AddMenuEntry("MidiFile_ExportMid", Label, ToolTip, Icon, UIAction);
@@ -75,10 +75,10 @@ void UAssetDefinition_MidiFile::ExecuteExportMidiFile(const FToolMenuContext& Me
 				TArray<FString> SaveFileNames;
 				const bool bFileSelected = DesktopPlatform->SaveFileDialog(
 					FSlateApplication::Get().FindBestParentWindowHandleForDialogs(nullptr),
-					LOCTEXT("MidiFile_ExportMid_SaveFileDialogTitle", "Save standard midi file as...").ToString(),
+					LOCTEXT("MidiFile_ExportMid_SaveFileDialogTitle", "Save Standard MIDI File as...").ToString(),
 					LastMidiExportFolder.IsEmpty() ? FPaths::ProjectDir() : LastMidiExportFolder,
 					DefaultFileName,
-					TEXT("Standard Midi File (*.mid)|*.mid"),
+					TEXT("Standard MIDI File (*.mid)|*.mid"),
 					EFileDialogFlags::None,
 					SaveFileNames);
 
@@ -104,7 +104,7 @@ void UAssetDefinition_MidiFile::ExportAllMidiToFolder(const UContentBrowserAsset
 	if (DesktopPlatform)
 	{
 		FString SelectedFolderName;
-		if (DesktopPlatform->OpenDirectoryDialog(FSlateApplication::Get().FindBestParentWindowHandleForDialogs(nullptr), "Select destination for standard MIDI files...", LastMidiExportFolder.IsEmpty() ? FPaths::ProjectDir() : LastMidiExportFolder, SelectedFolderName))
+		if (DesktopPlatform->OpenDirectoryDialog(FSlateApplication::Get().FindBestParentWindowHandleForDialogs(nullptr), "Select destination for Standard MIDI Files...", LastMidiExportFolder.IsEmpty() ? FPaths::ProjectDir() : LastMidiExportFolder, SelectedFolderName))
 		{
 			LastMidiExportFolder = SelectedFolderName;
 			bool bKeepWarning = true;
diff --git a/Engine/Plugins/Runtime/Harmonix/Source/HarmonixMidiEditor/Private/MidiFileFactory.cpp b/Engine/Plugins/Runtime/Harmonix/Source/HarmonixMidiEditor/Private/MidiFileFactory.cpp
--- a/Engine/Plugins/Runtime/Harmonix/Source/HarmonixMidiEditor/Private/MidiFileFactory.cpp
+++ b/Engine/Plugins/Runtime/Harmonix/Source/HarmonixMidiEditor/Private/MidiFileFactory.cpp
@@ -16,12 +16,17 @@ UMidiFileFactory::UMidiFileFactory(const FObjectInitializer& ObjectInitializer)
 	bEditorImport = true;
 	bText = false;
 
-	Formats.Add(TEXT("mid;Standard Midi File"));
+	Formats.Add(TEXT("mid;Standard MIDI File"));
+}
+
+FText UMidiFileFactory::GetDisplayName() const
+{
+	return NSLOCTEXT("MIDI", "MIDIImporterFactoryName", "Standard MIDI File");
 }
 
 FText UMidiFileFactory::GetToolTip() const
 {
-	return NSLOCTEXT("Midi", "MidiImporterFactoryDescription", "Standard Midi Files exported from Digital Audio Workstations");
+	return NSLOCTEXT("MIDI", "MIDIImporterFactoryDescription", "Standard MIDI Files exported from Digital Audio Workstations");
 }
 
 bool UMidiFileFactory::FactoryCanImport(const FString& Filename)
@@ -97,7 +102,7 @@ void UMidiFileFactory::ShowConformMidiFileLengthDialog(int32 MidiFileAssetIndex)
 	//show the dialog widgets on a window
 	TSharedPtr<SWindow> ParentWindow;
 	SAssignNew(ParentWindow, SWindow)
-		.Title(FText::FromString(TEXT("Conform Midi File Length")))
+		.Title(FText::FromString(TEXT("Conform MIDI File Length")))
 		.ClientSize(FVector2D(490, 250));
 	
 	//the custom pop up dialog that asks for conforming midi file length
@@ -148,7 +153,7 @@ void UMidiFileFactory::ShowConformMidiFileLengthDialog(int32 MidiFileAssetIndex)
 	//set text for text block that asks user whether or not to conform midi file length
 	FString MidiFileName = MidiFileAsset->GetName();
 	ConformMidiFileLengthDialog->AskConformFileLengthText->SetText(FText::FromString(FString::Printf(
-		TEXT("Midi File Name: %s.mid, Length: %.5f bars"),
+		TEXT("MIDI File Name: %s.mid, Length: %.5f bars"),
 		*MidiFileName,
 		MidiFileFractionalLength)));
 
@@ -182,7 +187,7 @@ EReimportResult::Type UMidiFileFactory::Reimport(UObject* Obj)
 	FString ReimportPath = AsMidiFile->GetImportedSrcFilePath();
 	if (ReimportPath.IsEmpty())
 	{
-		UE_LOG(LogHarmonixMidiEditor, Warning, TEXT("Failed to reimport midi file: %s"), *AsMidiFile->GetFullName());
+		UE_LOG(LogHarmonixMidiEditor, Warning, TEXT("Failed to reimport MIDI file: %s"), *AsMidiFile->GetFullName());
 		return EReimportResult::Failed;
 	}
 	
@@ -225,8 +230,8 @@ void SConformMidiFileLengthDialog::Construct(const FArguments& InArgs)
 			[
 				SNew(STextBlock)
 				.Text(FText::FromString(TEXT(
-					"Currently, all midi asset lengths must be conformed to some full number of bars."
-					"\n\n      Some or all of the midi files you are importing have fractional bar counts."
+					"Currently, all MIDI asset lengths must be conformed to some full number of bars."
+					"\n\n      Some or all of the MIDI files you are importing have fractional bar counts."
 					"\n\n                                        How would you like them conformed?\n")))
 			]
 			+ SVerticalBox::Slot()
@@ -266,7 +271,7 @@ void SConformMidiFileLengthDialog::Construct(const FArguments& InArgs)
 					.Content()
 					[
 						SNew(STextBlock)
-						.ToolTipText(FText::FromString(TEXT("move Midi events in file that exceed the last integer bar to the last tick of that bar and remove excessive events on that tick")))
+						.ToolTipText(FText::FromString(TEXT("move MIDI events in file that exceed the last integer bar to the last tick of that bar and remove excessive events on that tick")))
 						.Text(FText::FromString(TEXT("Round Down")))
 					]
 				]
@@ -283,7 +288,7 @@ void SConformMidiFileLengthDialog::Construct(const FArguments& InArgs)
 					.Content()
 					[
 						SNew(STextBlock)
-						.ToolTipText(FText::FromString(TEXT("round Midi File length to the nearest integer bar (either rounding up or rounding down)")))
+						.ToolTipText(FText::FromString(TEXT("round MIDI File length to the nearest integer bar (either rounding up or rounding down)")))
 						.Text(FText::FromString(TEXT("Round to Nearest")))
 					]
 				]
diff --git a/Engine/Plugins/Runtime/Harmonix/Source/HarmonixMidiEditor/Private/MidiFileFactory.h b/Engine/Plugins/Runtime/Harmonix/Source/HarmonixMidiEditor/Private/MidiFileFactory.h
--- a/Engine/Plugins/Runtime/Harmonix/Source/HarmonixMidiEditor/Private/MidiFileFactory.h
+++ b/Engine/Plugins/Runtime/Harmonix/Source/HarmonixMidiEditor/Private/MidiFileFactory.h
@@ -27,6 +27,7 @@ class UMidiFileFactory : public UFactory, public FReimportHandler
 	GENERATED_UCLASS_BODY()
 
 	//~ BEGIN UFactory interface
+	virtual FText GetDisplayName() const override;
 	virtual FText GetToolTip() const override;
 	virtual bool FactoryCanImport(const FString& Filename) override;
 	virtual UObject* FactoryCreateFile(UClass* InClass, UObject* InParent, FName InName, EObjectFlags Flags, const FString& Filename, const TCHAR* Parms, FFeedbackContext* Warn, bool& bOutOperationCanceled) override;
