Move a few nodes to public

diff --git a/Engine/Plugins/Runtime/Harmonix/Source/HarmonixMetasound/Private/HarmonixMetasound/Nodes/MidiClockSubdivisionTriggerNode.cpp b/Engine/Plugins/Runtime/Harmonix/Source/HarmonixMetasound/Private/HarmonixMetasound/Nodes/MidiClockSubdivisionTriggerNode.cpp
new file mode 100644
--- /dev/null
+++ b/Engine/Plugins/Runtime/Harmonix/Source/HarmonixMetasound/Private/HarmonixMetasound/Nodes/MidiClockSubdivisionTriggerNode.cpp
@@ -0,0 +1,325 @@
+// Copyright Epic Games, Inc. All Rights Reserved.
+
+#include "Tickable.h"
+#include "Containers/Queue.h"
+
+#include <atomic>
+#include <limits>
+
+#include "MetasoundExecutableOperator.h"
+#include "MetasoundFacade.h"
+#include "MetasoundNodeInterface.h"
+#include "MetasoundParamHelper.h"
+#include "MetasoundSampleCounter.h"
+#include "MetasoundStandardNodesCategories.h"
+#include "MetasoundEnumRegistrationMacro.h"
+#include "MetasoundVertex.h"
+
+#include "HarmonixMetasound/Common.h"
+#include "HarmonixMetasound/DataTypes/MidiClock.h"
+
+DEFINE_LOG_CATEGORY_STATIC(LogMidiClockSubdivisionTrigger, Log, All);
+
+#define LOCTEXT_NAMESPACE "HarmonixMetaSound_MidiClockSubdivisionTriggerNode"
+
+namespace HarmonixMetasound
+{
+	using namespace Metasound;
+
+
+	class FMidiClockSubdivisionTriggerOperator : public TExecutableOperator<FMidiClockSubdivisionTriggerOperator>, public FMidiPlayCursor
+	{
+	public:
+		static const FNodeClassMetadata& GetNodeInfo();
+		static const FVertexInterface& GetVertexInterface();
+		static TUniquePtr<IOperator> CreateOperator(const FBuildOperatorParams& InParams, FBuildResults& OutResults);
+
+		FMidiClockSubdivisionTriggerOperator(const FBuildOperatorParams& InParams,
+		                            const FBoolReadRef&      InEnabled,
+		                            const FMidiClockReadRef& InMidiClock,
+									const FInt32ReadRef&     InGridSizeMult,
+									const FEnumMidiClockSubdivisionQuantizationReadRef& InGridUnits,
+									const FInt32ReadRef&	 InOffsetMult,
+									const FEnumMidiClockSubdivisionQuantizationReadRef& InOffsetUnits);
+		virtual ~FMidiClockSubdivisionTriggerOperator() override;
+
+		virtual void BindInputs(FInputVertexInterfaceData& InVertexData) override;
+		virtual void BindOutputs(FOutputVertexInterfaceData& InVertexData) override;
+
+		void Reset(const FResetParams& ResetParams);
+
+		void Execute();
+
+	private:
+		//** INPUTS
+		FMidiClockReadRef   MidiClockInPin;
+		FBoolReadRef        EnableInPin;
+		FInt32ReadRef		GridSizeMultInPin;
+		FEnumMidiClockSubdivisionQuantizationReadRef GridSizeUnitsInPin;
+		FInt32ReadRef		GridOffsetMultInPin;
+		FEnumMidiClockSubdivisionQuantizationReadRef GridOffsetUnitsInPin;
+
+		//** OUTPUTS
+		FTriggerWriteRef   TriggerOutPin;
+
+ 		//** DATA (current state)
+		EMidiClockSubdivisionQuantization GridSizeUnits;
+		int32 GridSizeMultiplier = 1;
+		EMidiClockSubdivisionQuantization GridOffsetUnits;
+		int32 GridOffsetMultiplier = 0;
+
+		int32 GridOffsetTicks = 0;
+		int32 GridSizeTicks	  = 0;
+		bool  AdvanceBlockNeeded = true;
+
+		//** BEGIN FMidiPlayCursor
+		virtual void OnTimeSig(int32 TrackIndex, int32 Tick, int32 Numerator, int32 Denominator, bool IsPreroll = false) override;
+		virtual void AdvanceThruTick(int32 Tick, bool IsPreRoll) override;
+		// We have to override this to disambiguate the FMidiPlayCursor reset and the MS operator one
+		virtual void Reset(bool ForceNoBroadcast) override { FMidiPlayCursor::Reset(ForceNoBroadcast); }
+		//** END FMidiPlayCursor
+	};
+
+	class FMidiClockSubdivisionTriggerNode : public FNodeFacade
+	{
+	public:
+		FMidiClockSubdivisionTriggerNode(const FNodeInitData& InInitData)
+			: FNodeFacade(InInitData.InstanceName, InInitData.InstanceID, TFacadeOperatorClass<FMidiClockSubdivisionTriggerOperator>())
+		{}
+		virtual ~FMidiClockSubdivisionTriggerNode() = default;
+	};
+
+	METASOUND_REGISTER_NODE(FMidiClockSubdivisionTriggerNode)
+		
+	const FNodeClassMetadata& FMidiClockSubdivisionTriggerOperator::GetNodeInfo()
+	{
+		auto InitNodeInfo = []() -> FNodeClassMetadata
+		{
+			FNodeClassMetadata Info;
+			Info.ClassName        = { HarmonixNodeNamespace, TEXT("MidiClockSubdivisionTrigger"), TEXT("")};
+			Info.MajorVersion     = 0;
+			Info.MinorVersion     = 1;
+			Info.DisplayName      = METASOUND_LOCTEXT("MIDIClockSubdivisionTriggerNode_DisplayName", "MIDI Clock Subdivision Trigger");
+			Info.Description      = METASOUND_LOCTEXT("MIDIClockSubdivisionTriggerNode_Description", "Watches a MIDI clock and outputs triggers at musical subdivisions.");
+			Info.Author           = PluginAuthor;
+			Info.PromptIfMissing  = PluginNodeMissingPrompt;
+			Info.DefaultInterface = GetVertexInterface();
+			Info.CategoryHierarchy = { MetasoundNodeCategories::Harmonix, NodeCategories::Music };
+			return Info;
+		};
+
+		static const FNodeClassMetadata Info = InitNodeInfo();
+
+		return Info;
+	}
+
+	namespace MidiClockSubdivisionTriggerPinNames
+	{
+		METASOUND_PARAM(TriggerOutput, "Trigger Out", "A series of triggers at the specified subdivision grid.")
+	}
+
+	const FVertexInterface& FMidiClockSubdivisionTriggerOperator::GetVertexInterface()
+	{
+		using namespace MidiClockSubdivisionTriggerPinNames;
+		using namespace CommonPinNames;
+
+		static const FVertexInterface Interface(
+			FInputVertexInterface(
+				TInputDataVertex<FMidiClock>(METASOUND_GET_PARAM_NAME_AND_METADATA(Inputs::MidiClock)),
+				TInputDataVertex<bool>(METASOUND_GET_PARAM_NAME_AND_METADATA(Inputs::Enable), true),
+				TInputDataVertex<FEnumMidiClockSubdivisionQuantizationType>(METASOUND_GET_PARAM_NAME_AND_METADATA(Inputs::GridSizeUnits), (int32)EMidiClockSubdivisionQuantization::Beat),
+				TInputDataVertex<int32>(METASOUND_GET_PARAM_NAME_AND_METADATA(Inputs::GridSizeMult), 1),
+				TInputDataVertex<FEnumMidiClockSubdivisionQuantizationType>(METASOUND_GET_PARAM_NAME_AND_METADATA(Inputs::OffsetUnits), (int32)EMidiClockSubdivisionQuantization::Beat),
+				TInputDataVertex<int32>(METASOUND_GET_PARAM_NAME_AND_METADATA(Inputs::OffsetMult), 0)
+				),
+			FOutputVertexInterface(
+				TOutputDataVertex<FTrigger>(METASOUND_GET_PARAM_NAME_AND_METADATA(TriggerOutput))
+			)
+		);
+
+		return Interface;
+	}
+
+	TUniquePtr<IOperator> FMidiClockSubdivisionTriggerOperator::CreateOperator(const FBuildOperatorParams& InParams, FBuildResults& OutResults)
+	{
+		using namespace MidiClockSubdivisionTriggerPinNames;
+		using namespace CommonPinNames;
+
+		const FMidiClockSubdivisionTriggerNode& TheNode = static_cast<const FMidiClockSubdivisionTriggerNode&>(InParams.Node);
+
+		const FInputVertexInterfaceData& InputData = InParams.InputData;
+		FMidiClockReadRef InMidiClock = InputData.GetOrConstructDataReadReference<FMidiClock>(METASOUND_GET_PARAM_NAME(Inputs::MidiClock), InParams.OperatorSettings);
+		FBoolReadRef InEnabled        = InputData.GetOrCreateDefaultDataReadReference<bool>(METASOUND_GET_PARAM_NAME(Inputs::Enable), InParams.OperatorSettings);
+		FEnumMidiClockSubdivisionQuantizationReadRef InGridSizeUnits = InputData.GetOrCreateDefaultDataReadReference<FEnumMidiClockSubdivisionQuantizationType>(METASOUND_GET_PARAM_NAME(Inputs::GridSizeUnits), InParams.OperatorSettings);
+		FInt32ReadRef InGridSizeMult = InputData.GetOrCreateDefaultDataReadReference<int32>(METASOUND_GET_PARAM_NAME(Inputs::GridSizeMult), InParams.OperatorSettings);
+		FEnumMidiClockSubdivisionQuantizationReadRef InOffsetUnits   = InputData.GetOrCreateDefaultDataReadReference<FEnumMidiClockSubdivisionQuantizationType>(METASOUND_GET_PARAM_NAME(Inputs::OffsetUnits), InParams.OperatorSettings);
+		FInt32ReadRef InOffsetMult = InputData.GetOrCreateDefaultDataReadReference<int32>(METASOUND_GET_PARAM_NAME(Inputs::OffsetMult), InParams.OperatorSettings);
+
+		return MakeUnique<FMidiClockSubdivisionTriggerOperator>(InParams,
+			InEnabled,
+			InMidiClock,
+			InGridSizeMult,
+			InGridSizeUnits,
+			InOffsetMult,
+			InOffsetUnits);
+	}
+
+	FMidiClockSubdivisionTriggerOperator::FMidiClockSubdivisionTriggerOperator(const FBuildOperatorParams& InParams,
+															 const FBoolReadRef&      InEnabled,
+															 const FMidiClockReadRef& InMidiClock,
+															 const FInt32ReadRef&	  InGridSizeMult,
+															 const FEnumMidiClockSubdivisionQuantizationReadRef& InGridUnits,
+															 const FInt32ReadRef&	  InOffsetMult,
+															 const FEnumMidiClockSubdivisionQuantizationReadRef& InOffsetUnits)
+		: MidiClockInPin(InMidiClock)
+		, EnableInPin(InEnabled)
+		, GridSizeMultInPin(InGridSizeMult)
+		, GridSizeUnitsInPin(InGridUnits)
+		, GridOffsetMultInPin(InOffsetMult)
+		, GridOffsetUnitsInPin(InOffsetUnits)
+		, TriggerOutPin(FTriggerWriteRef::CreateNew(InParams.OperatorSettings))
+	{
+		Reset(InParams);
+	}
+
+	FMidiClockSubdivisionTriggerOperator::~FMidiClockSubdivisionTriggerOperator()
+	{
+	}
+
+	void FMidiClockSubdivisionTriggerOperator::BindInputs(FInputVertexInterfaceData& InVertexData)
+	{
+		using namespace CommonPinNames;
+
+		InVertexData.BindReadVertex(METASOUND_GET_PARAM_NAME(Inputs::Enable),    EnableInPin);
+		InVertexData.BindReadVertex(METASOUND_GET_PARAM_NAME(Inputs::MidiClock), MidiClockInPin);
+		InVertexData.BindReadVertex(METASOUND_GET_PARAM_NAME(Inputs::GridSizeMult), GridSizeMultInPin);
+		InVertexData.BindReadVertex(METASOUND_GET_PARAM_NAME(Inputs::GridSizeUnits), GridSizeUnitsInPin);
+		InVertexData.BindReadVertex(METASOUND_GET_PARAM_NAME(Inputs::OffsetMult), GridOffsetMultInPin);
+		InVertexData.BindReadVertex(METASOUND_GET_PARAM_NAME(Inputs::OffsetUnits), GridOffsetUnitsInPin);
+	}
+
+	void FMidiClockSubdivisionTriggerOperator::BindOutputs(FOutputVertexInterfaceData& InVertexData)
+	{
+		InVertexData.BindReadVertex(METASOUND_GET_PARAM_NAME(MidiClockSubdivisionTriggerPinNames::TriggerOutput), TriggerOutPin);
+	}
+	
+	void FMidiClockSubdivisionTriggerOperator::Reset(const FResetParams& ResetParams)
+	{
+		TriggerOutPin->Reset();
+		
+		SetMessageFilter(FMidiPlayCursor::EFilterPassFlags::TimeSig);
+		MidiClockInPin->RegisterHiResPlayCursor(this);
+
+		GridSizeUnits = *GridSizeUnitsInPin;
+		GridSizeMultiplier = *GridSizeMultInPin;
+		if (GridSizeMultiplier < 1)
+		{
+			GridSizeMultiplier = 1;
+		}
+		GridOffsetUnits = *GridOffsetUnitsInPin;
+		GridOffsetMultiplier = *GridOffsetMultInPin;
+
+		const FSongMaps& SongMaps = MidiClockInPin->GetSongMaps();
+		GridOffsetTicks = SubdivisionToMidiTicks(GridOffsetUnits, 0, SongMaps) * GridOffsetMultiplier;
+		GridSizeTicks = SubdivisionToMidiTicks(GridSizeUnits, 0, SongMaps) * GridSizeMultiplier;
+
+		AdvanceBlockNeeded = true;
+	}
+
+	void FMidiClockSubdivisionTriggerOperator::Execute()
+	{
+		if (AdvanceBlockNeeded)
+		{
+			TriggerOutPin->AdvanceBlock();
+			AdvanceBlockNeeded = false;
+		}
+
+		// first let's see if our configuration has changed at all...
+		if (*GridSizeUnitsInPin != GridSizeUnits || *GridSizeMultInPin != GridSizeMultiplier)
+		{
+			GridSizeUnits = *GridSizeUnitsInPin;
+			GridSizeMultiplier = *GridSizeMultInPin;
+			if (GridSizeMultiplier < 1)
+			{
+				GridSizeMultiplier = 1;
+			}
+			const FSongMaps& SongMaps = MidiClockInPin->GetSongMaps();
+			GridSizeTicks = SubdivisionToMidiTicks(GridSizeUnits, CurrentTick, SongMaps) * GridSizeMultiplier;
+		}
+
+		if (*GridOffsetUnitsInPin != GridOffsetUnits || *GridOffsetMultInPin != GridOffsetMultiplier)
+		{
+			GridOffsetUnits = *GridOffsetUnitsInPin;
+			GridOffsetMultiplier = *GridOffsetMultInPin;
+			const FSongMaps& SongMaps = MidiClockInPin->GetSongMaps();
+			GridOffsetTicks = SubdivisionToMidiTicks(GridOffsetUnits, 0, SongMaps) * GridOffsetMultiplier;
+		}
+
+		AdvanceBlockNeeded = true;
+	}
+
+	void FMidiClockSubdivisionTriggerOperator::AdvanceThruTick(int32 Tick, bool IsPreRoll)
+	{
+		if (AdvanceBlockNeeded)
+		{
+			TriggerOutPin->AdvanceBlock();
+			AdvanceBlockNeeded = false;
+		}
+
+		int32 TickProceedingThisAdvance = CurrentTick - GridOffsetTicks;
+
+		FMidiPlayCursor::AdvanceThruTick(Tick, IsPreRoll);
+		if (IsPreRoll || !*EnableInPin)
+		{
+			return;
+		}
+
+		int32 LastTickProcessed = CurrentTick - GridOffsetTicks;
+		if (LastTickProcessed < 0)
+		{
+			return;
+		}
+
+		// if this is one of the quantization types that is variable in size we have to do it a little different...
+		if (GridSizeUnits == EMidiClockSubdivisionQuantization::Bar ||
+			GridSizeUnits == EMidiClockSubdivisionQuantization::Beat)
+		{
+			const FSongMaps& SongMaps = MidiClockInPin->GetSongMaps();
+			FMusicTimestamp Start = SongMaps.GetBarMap().TickToMusicTimestamp(TickProceedingThisAdvance);
+			FMusicTimestamp End = SongMaps.GetBarMap().TickToMusicTimestamp(LastTickProcessed);
+			if (GridSizeUnits == EMidiClockSubdivisionQuantization::Bar && Start.Bar != End.Bar)
+			{
+				if (End.Bar % GridSizeMultiplier == 0)
+				{
+					TriggerOutPin->TriggerFrame(MidiClockInPin->GetCurrentBlockFrameIndex());
+				}
+			}
+			else if (GridSizeUnits == EMidiClockSubdivisionQuantization::Beat && FMath::FloorToInt32(Start.Beat) != FMath::FloorToInt32(End.Beat))
+			{
+				if (FMath::FloorToInt32(End.Beat-1.0f) % GridSizeMultiplier == 0)
+				{
+					TriggerOutPin->TriggerFrame(MidiClockInPin->GetCurrentBlockFrameIndex());
+				}
+			}
+		}
+		else
+		{
+			int32 StartingGridSquare = TickProceedingThisAdvance / GridSizeTicks;
+			int32 EndingGridSquare = LastTickProcessed / GridSizeTicks;
+			if (StartingGridSquare != EndingGridSquare)
+			{
+				TriggerOutPin->TriggerFrame(MidiClockInPin->GetCurrentBlockFrameIndex());
+			}
+		}
+	}
+
+	void FMidiClockSubdivisionTriggerOperator::OnTimeSig(int32 TrackIndex, int32 Tick, int32 Numerator, int32 Denominator, bool IsPreroll /*= false*/)
+	{
+		const FSongMaps& SongMaps = MidiClockInPin->GetSongMaps();
+		GridSizeTicks = SubdivisionToMidiTicks(GridSizeUnits, Tick, SongMaps) * GridSizeMultiplier;
+		//NOTE: We DO NOT calculate a new grid offset because this is always specified from the beginning of the song... at the song's starting time signature.  
+	}
+
+}
+
+#undef LOCTEXT_NAMESPACE // "HarmonixMetaSound"
diff --git a/Engine/Plugins/Runtime/Harmonix/Source/HarmonixMetasound/Private/HarmonixMetasound/Nodes/MusicTimestampBuilder.cpp b/Engine/Plugins/Runtime/Harmonix/Source/HarmonixMetasound/Private/HarmonixMetasound/Nodes/MusicTimestampBuilder.cpp
new file mode 100644
--- /dev/null
+++ b/Engine/Plugins/Runtime/Harmonix/Source/HarmonixMetasound/Private/HarmonixMetasound/Nodes/MusicTimestampBuilder.cpp
@@ -0,0 +1,148 @@
+// Copyright Epic Games, Inc. All Rights Reserved.
+
+#include "MetasoundExecutableOperator.h"
+#include "MetasoundFacade.h"
+#include "MetasoundNodeInterface.h"
+#include "MetasoundParamHelper.h"
+#include "MetasoundSampleCounter.h"
+#include "MetasoundStandardNodesCategories.h"
+#include "MetasoundVertex.h"
+
+#include "HarmonixMetasound/Common.h"
+#include "HarmonixMetasound/DataTypes/MusicTransport.h"
+#include "HarmonixMetasound/DataTypes/MusicTimestamp.h"
+
+#define LOCTEXT_NAMESPACE "HarmonixMetaSound"
+
+namespace HarmonixMetasound
+{
+	using namespace Metasound;
+
+	class FMusicTimestampBuilderOperator : public TExecutableOperator<FMusicTimestampBuilderOperator>
+	{
+	public:
+		static const FNodeClassMetadata& GetNodeInfo();
+		static const FVertexInterface& GetVertexInterface();
+		static TUniquePtr<IOperator> CreateOperator(const FBuildOperatorParams& InParams, FBuildResults& OutResults);
+
+		FMusicTimestampBuilderOperator(
+			const FBuildOperatorParams& InParams,
+			const FInt32ReadRef& InBar,
+			const FFloatReadRef& InBeat);
+
+		virtual void BindInputs(FInputVertexInterfaceData& InVertexData) override;
+		virtual void BindOutputs(FOutputVertexInterfaceData& InVertexData) override;
+
+		void Reset(const FResetParams& ResetParams);
+
+		void Execute();
+
+	private:
+		//** INPUTS
+		FInt32ReadRef BarInPin;
+		FFloatReadRef BeatInPin;
+
+		//** OUTPUTS
+		FMusicTimestampWriteRef MusicTimeStampOutPin;
+	};
+
+	class FMusicTimestampBuilderNode : public FNodeFacade
+	{
+	public:
+		FMusicTimestampBuilderNode(const FNodeInitData& InInitData)
+			: FNodeFacade(InInitData.InstanceName, InInitData.InstanceID, TFacadeOperatorClass<FMusicTimestampBuilderOperator>())
+		{}
+		virtual ~FMusicTimestampBuilderNode() = default;
+	};
+
+	METASOUND_REGISTER_NODE(FMusicTimestampBuilderNode)
+
+	const FNodeClassMetadata& FMusicTimestampBuilderOperator::GetNodeInfo()
+	{
+		auto InitNodeInfo = []() -> FNodeClassMetadata
+		{
+			FNodeClassMetadata Info;
+			Info.ClassName        = { HarmonixNodeNamespace, TEXT("MusicTimestampBuilder"), TEXT("") };
+			Info.MajorVersion     = 0;
+			Info.MinorVersion     = 1;
+			Info.DisplayName      = METASOUND_LOCTEXT("MusicTimestampBuilderNode_DisplayName", "Music Timestamp Builder");
+			Info.Description      = METASOUND_LOCTEXT("MusicTimestampBuilderNode_Description", "Combines the input Bar and Beat into a single Music Timestamp structure.");
+			Info.Author           = PluginAuthor;
+			Info.PromptIfMissing  = PluginNodeMissingPrompt;
+			Info.DefaultInterface = GetVertexInterface();
+			Info.CategoryHierarchy = { MetasoundNodeCategories::Harmonix, NodeCategories::Music };
+			return Info;
+		};
+
+		static const FNodeClassMetadata Info = InitNodeInfo();
+
+		return Info;
+	}
+
+	const FVertexInterface& FMusicTimestampBuilderOperator::GetVertexInterface()
+	{
+		using namespace CommonPinNames;
+
+		static const FVertexInterface Interface(
+			FInputVertexInterface(
+				TInputDataVertex<int32>(METASOUND_GET_PARAM_NAME_AND_METADATA(Inputs::Bar), 1),
+				TInputDataVertex<float>(METASOUND_GET_PARAM_NAME_AND_METADATA(Inputs::FloatBeat), 1.0f)
+			),
+			FOutputVertexInterface(
+				TOutputDataVertex<FMusicTimestamp>(METASOUND_GET_PARAM_NAME_AND_METADATA(Outputs::MusicTimestamp))
+			)
+		);
+
+		return Interface;
+	}
+
+	TUniquePtr<IOperator> FMusicTimestampBuilderOperator::CreateOperator(const FBuildOperatorParams& InParams, FBuildResults& OutResults)
+	{
+		using namespace CommonPinNames;
+
+		const FInputVertexInterfaceData& InputData = InParams.InputData;
+		FInt32ReadRef InBar  = InputData.GetOrCreateDefaultDataReadReference<int32>(METASOUND_GET_PARAM_NAME(Inputs::Bar), InParams.OperatorSettings);
+		FFloatReadRef InBeat = InputData.GetOrCreateDefaultDataReadReference<float>(METASOUND_GET_PARAM_NAME(Inputs::FloatBeat), InParams.OperatorSettings);
+
+		return MakeUnique<FMusicTimestampBuilderOperator>(InParams, InBar, InBeat);
+	}
+
+	FMusicTimestampBuilderOperator::FMusicTimestampBuilderOperator(
+		const FBuildOperatorParams& InParams,
+		const FInt32ReadRef& InBar,
+		const FFloatReadRef& InBeat)
+	: BarInPin(InBar)
+	, BeatInPin(InBeat)
+	, MusicTimeStampOutPin(FMusicTimestampWriteRef::CreateNew())
+	{
+		Reset(InParams);
+	}
+
+	void FMusicTimestampBuilderOperator::BindInputs(FInputVertexInterfaceData& InVertexData)
+	{
+		using namespace CommonPinNames;
+
+		InVertexData.BindReadVertex(METASOUND_GET_PARAM_NAME(Inputs::Bar), BarInPin);
+		InVertexData.BindReadVertex(METASOUND_GET_PARAM_NAME(Inputs::FloatBeat), BeatInPin);
+	}
+
+	void FMusicTimestampBuilderOperator::BindOutputs(FOutputVertexInterfaceData& InVertexData)
+	{
+		using namespace CommonPinNames;
+
+		InVertexData.BindReadVertex(METASOUND_GET_PARAM_NAME(Outputs::MusicTimestamp), MusicTimeStampOutPin);
+	}
+
+	void FMusicTimestampBuilderOperator::Reset(const FResetParams& ResetParams)
+	{
+		MusicTimeStampOutPin->Reset();
+	}
+
+	void FMusicTimestampBuilderOperator::Execute()
+	{
+		MusicTimeStampOutPin->Bar = *BarInPin;
+		MusicTimeStampOutPin->Beat = *BeatInPin;
+	}
+}
+
+#undef LOCTEXT_NAMESPACE // "HarmonixMetaSound"
