MusicTimestamp to Trigger & Music Subdivision To Trigger nodes in to public

diff --git a/Engine/Plugins/Runtime/Harmonix/Source/HarmonixMetasound/Private/HarmonixMetasound/Nodes/MidiClockTimestampTriggerNode.cpp b/Engine/Plugins/Runtime/Harmonix/Source/HarmonixMetasound/Private/HarmonixMetasound/Nodes/MidiClockTimestampTriggerNode.cpp
new file mode 100644
--- /dev/null
+++ b/Engine/Plugins/Runtime/Harmonix/Source/HarmonixMetasound/Private/HarmonixMetasound/Nodes/MidiClockTimestampTriggerNode.cpp
@@ -0,0 +1,322 @@
+// Copyright Epic Games, Inc. All Rights Reserved.
+
+#include "Tickable.h"
+#include "Containers/Queue.h"
+
+#include <atomic>
+#include <limits>
+
+#include "MetasoundExecutableOperator.h"
+#include "MetasoundFacade.h"
+#include "MetasoundNodeInterface.h"
+#include "MetasoundParamHelper.h"
+#include "MetasoundSampleCounter.h"
+#include "MetasoundStandardNodesCategories.h"
+#include "MetasoundEnumRegistrationMacro.h"
+#include "MetasoundVertex.h"
+
+#include "HarmonixMetasound/Common.h"
+#include "HarmonixMetasound/DataTypes/MidiClock.h"
+#include "HarmonixMetasound/DataTypes/MusicTimestamp.h"
+
+DEFINE_LOG_CATEGORY_STATIC(LogMidiClockTimestampTrigger, Log, All);
+
+#define LOCTEXT_NAMESPACE "HarmonixMetaSound_MidiClockTimestampTriggerNode"
+
+namespace HarmonixMetasound
+{
+	using namespace Metasound;
+
+	class FMidiClockTimestampTriggerOperator : public TExecutableOperator<FMidiClockTimestampTriggerOperator>, public FMidiPlayCursor
+	{
+	public:
+		static const FNodeClassMetadata& GetNodeInfo();
+		static const FVertexInterface& GetVertexInterface();
+		static TUniquePtr<IOperator> CreateOperator(const FBuildOperatorParams& InParams, FBuildResults& OutResults);
+
+		FMidiClockTimestampTriggerOperator(const FBuildOperatorParams& InParams,
+									const FBoolReadRef&                                 InEnabled,
+		                            const FMidiClockReadRef&                            InMidiClock,
+									const FMusicTimestampReadRef&                       InTimestamp,
+									const FBoolReadRef&                                 InShouldQuantize,
+									const FEnumMidiClockSubdivisionQuantizationReadRef& InQuanitizationUnit,
+									const FBoolReadRef&                                 InTriggerDuringSeek);
+		virtual ~FMidiClockTimestampTriggerOperator() override;
+
+		virtual void BindInputs(FInputVertexInterfaceData& InVertexData) override;
+		virtual void BindOutputs(FOutputVertexInterfaceData& InVertexData) override;
+
+		void Reset(const FResetParams& ResetParams);
+		
+		void Execute();
+
+	private:
+		//** INPUTS
+		FMidiClockReadRef                            MidiClockInPin;
+		FBoolReadRef                                 EnableInPin;
+		FMusicTimestampReadRef                       TimestampInPin;
+		FBoolReadRef                                 ShouldQuantizeInPin;
+		FEnumMidiClockSubdivisionQuantizationReadRef QuantizeUnitInPin;
+		FBoolReadRef		                         TriggerDuringSeekInPin;
+
+		//** OUTPUTS
+		FTriggerWriteRef   TriggerOutPin;
+
+ 		//** DATA (current state)
+		FMusicTimestamp CurrentTimestamp{1, 1.0f};
+		bool            bShouldQuantize = false;
+		EMidiClockSubdivisionQuantization QuantizeUnit = EMidiClockSubdivisionQuantization::ThirtySecondNote;
+
+		int32 TriggerTick = 0;
+
+		struct FTickSpan
+		{
+			int32 FromTick;
+			int32 ThruTick;
+			int32 BlockFrameIndex;
+			bool  bIsSeek;
+			FTickSpan(int32 InFromTick, int32 InThruTick, int32 InBlockFrameIndex, bool InIsSeek)
+				: FromTick(InFromTick)
+				, ThruTick(InThruTick)
+				, BlockFrameIndex(InBlockFrameIndex)
+				, bIsSeek(InIsSeek)
+			{}
+		};
+		TArray<FTickSpan> TickSpans;
+
+		//** BEGIN FMidiPlayCursor
+		virtual void SeekToTick(int32 Tick) override { SeekThruTick(Tick - 1); }
+		virtual void SeekThruTick(int32 Tick) override;
+		virtual void AdvanceThruTick(int32 Tick, bool IsPreRoll) override;
+		// We have to override this to disambiguate the FMidiPlayCursor Reset and the MS operator Reset
+		virtual void Reset(bool ForceNoBroadcast) override { FMidiPlayCursor::Reset(ForceNoBroadcast); }
+		//** END FMidiPlayCursor
+
+		void CalculateTriggerTick();
+	};
+
+	class FMidiClockTimestampTriggerNode : public FNodeFacade
+	{
+	public:
+		FMidiClockTimestampTriggerNode(const FNodeInitData& InInitData)
+			: FNodeFacade(InInitData.InstanceName, InInitData.InstanceID, TFacadeOperatorClass<FMidiClockTimestampTriggerOperator>())
+		{}
+		virtual ~FMidiClockTimestampTriggerNode() = default;
+	};
+
+	METASOUND_REGISTER_NODE(FMidiClockTimestampTriggerNode)
+		
+	const FNodeClassMetadata& FMidiClockTimestampTriggerOperator::GetNodeInfo()
+	{
+		auto InitNodeInfo = []() -> FNodeClassMetadata
+		{
+			FNodeClassMetadata Info;
+			Info.ClassName        = { HarmonixNodeNamespace, TEXT("MidiClockTimestampTrigger"), TEXT("")};
+			Info.MajorVersion     = 1;
+			Info.MinorVersion     = 0;
+			Info.DisplayName      = METASOUND_LOCTEXT("MIDIClockTimestampTriggerNode_DisplayName", "MIDI Clock Timestamp Trigger");
+			Info.Description      = METASOUND_LOCTEXT("MIDIClockTimestampTriggerNode_Description", "Watches a MIDI clock and outputs a trigger at the specified musical timestamp. The floating point beat can optionally be quantized to a musical subdivision.");
+			Info.Author           = PluginAuthor;
+			Info.PromptIfMissing  = PluginNodeMissingPrompt;
+			Info.DefaultInterface = GetVertexInterface();
+			Info.CategoryHierarchy = { MetasoundNodeCategories::Harmonix, NodeCategories::Music };
+			return Info;
+		};
+
+		static const FNodeClassMetadata Info = InitNodeInfo();
+
+		return Info;
+	}
+
+	namespace MidiClockTimestampTriggerPinNames
+	{
+		METASOUND_PARAM(TriggerDuringSeek, "Trigger During Seek", "Whether a trigger should be generated is a seek over the timestamp is detected.")
+		METASOUND_PARAM(TriggerOutput, "Trigger Out", "A trigger when the timestamp is detected.")
+	}
+
+	const FVertexInterface& FMidiClockTimestampTriggerOperator::GetVertexInterface()
+	{
+		using namespace MidiClockTimestampTriggerPinNames;
+		using namespace CommonPinNames;
+
+		static const FVertexInterface Interface(
+			FInputVertexInterface(
+				TInputDataVertex<FMidiClock>(METASOUND_GET_PARAM_NAME_AND_METADATA(Inputs::MidiClock)),
+				TInputDataVertex<bool>(METASOUND_GET_PARAM_NAME_AND_METADATA(Inputs::Enable), true),
+				TInputDataVertex<FMusicTimestamp>(METASOUND_GET_PARAM_NAME_AND_METADATA(Inputs::Timestamp)),
+				TInputDataVertex<bool>(METASOUND_GET_PARAM_NAME_AND_METADATA(Inputs::ShouldQuantizeTimestamp), false),
+				TInputDataVertex<FEnumMidiClockSubdivisionQuantizationType>(METASOUND_GET_PARAM_NAME_AND_METADATA(Inputs::QuantizationUnit), (int32)EMidiClockSubdivisionQuantization::Beat),
+				TInputDataVertex<bool>(METASOUND_GET_PARAM_NAME_AND_METADATA(TriggerDuringSeek), false)
+				),
+			FOutputVertexInterface(
+				TOutputDataVertex<FTrigger>(METASOUND_GET_PARAM_NAME_AND_METADATA(TriggerOutput))
+			)
+		);
+
+		return Interface;
+	}
+
+	TUniquePtr<IOperator> FMidiClockTimestampTriggerOperator::CreateOperator(const FBuildOperatorParams& InParams, FBuildResults& OutResults)
+	{
+		using namespace MidiClockTimestampTriggerPinNames;
+		using namespace CommonPinNames;
+
+		const FMidiClockTimestampTriggerNode& TheNode = static_cast<const FMidiClockTimestampTriggerNode&>(InParams.Node);
+
+		const FInputVertexInterfaceData& InputData = InParams.InputData;
+		FMidiClockReadRef InMidiClock      = InputData.GetOrConstructDataReadReference<FMidiClock>(METASOUND_GET_PARAM_NAME(Inputs::MidiClock), InParams.OperatorSettings);
+		FBoolReadRef  InEnabled            = InputData.GetOrCreateDefaultDataReadReference<bool>(METASOUND_GET_PARAM_NAME(Inputs::Enable), InParams.OperatorSettings);
+		FMusicTimestampReadRef InTimestamp = InputData.GetOrConstructDataReadReference<FMusicTimestamp>(METASOUND_GET_PARAM_NAME(Inputs::Timestamp), 1, 1.0f);
+		FBoolReadRef  InShouldQuantize     = InputData.GetOrCreateDefaultDataReadReference<bool>(METASOUND_GET_PARAM_NAME(Inputs::ShouldQuantizeTimestamp), InParams.OperatorSettings);
+		FEnumMidiClockSubdivisionQuantizationReadRef InQuantizeUnits = InputData.GetOrCreateDefaultDataReadReference<FEnumMidiClockSubdivisionQuantizationType>(METASOUND_GET_PARAM_NAME(Inputs::QuantizationUnit), InParams.OperatorSettings);
+		FBoolReadRef  InTriggerDuringSeek  = InputData.GetOrCreateDefaultDataReadReference<bool>(METASOUND_GET_PARAM_NAME(TriggerDuringSeek), InParams.OperatorSettings);
+
+		return MakeUnique<FMidiClockTimestampTriggerOperator>(InParams,
+			InEnabled,
+			InMidiClock,
+			InTimestamp,
+			InShouldQuantize,
+			InQuantizeUnits,
+			InTriggerDuringSeek);
+	}
+
+	FMidiClockTimestampTriggerOperator::FMidiClockTimestampTriggerOperator(const FBuildOperatorParams&           InParams,
+															 const FBoolReadRef&                                 InEnabled,
+															 const FMidiClockReadRef&                            InMidiClock,
+															 const FMusicTimestampReadRef&                       InTimestamp,
+															 const FBoolReadRef&                                 InShouldQuantize,
+															 const FEnumMidiClockSubdivisionQuantizationReadRef& InQuanitizationUnit,
+															 const FBoolReadRef&                                 InTriggerDuringSeek)
+		: MidiClockInPin(InMidiClock)
+		, EnableInPin(InEnabled)
+		, TimestampInPin(InTimestamp)
+		, ShouldQuantizeInPin(InShouldQuantize)
+		, QuantizeUnitInPin(InQuanitizationUnit)
+		, TriggerDuringSeekInPin(InTriggerDuringSeek)
+		, TriggerOutPin(FTriggerWriteRef::CreateNew(InParams.OperatorSettings))
+	{
+		Reset(InParams);
+	}
+
+	FMidiClockTimestampTriggerOperator::~FMidiClockTimestampTriggerOperator()
+	{
+	}
+
+	void FMidiClockTimestampTriggerOperator::BindInputs(FInputVertexInterfaceData& InVertexData)
+	{
+		using namespace CommonPinNames;
+		using namespace MidiClockTimestampTriggerPinNames;
+
+		InVertexData.BindReadVertex(METASOUND_GET_PARAM_NAME(Inputs::Enable),    EnableInPin);
+		InVertexData.BindReadVertex(METASOUND_GET_PARAM_NAME(Inputs::MidiClock), MidiClockInPin);
+		InVertexData.BindReadVertex(METASOUND_GET_PARAM_NAME(Inputs::Timestamp), TimestampInPin);
+		InVertexData.BindReadVertex(METASOUND_GET_PARAM_NAME(Inputs::ShouldQuantizeTimestamp), ShouldQuantizeInPin);
+		InVertexData.BindReadVertex(METASOUND_GET_PARAM_NAME(Inputs::QuantizationUnit), QuantizeUnitInPin);
+		InVertexData.BindReadVertex(METASOUND_GET_PARAM_NAME(TriggerDuringSeek), TriggerDuringSeekInPin);
+	}
+
+	void FMidiClockTimestampTriggerOperator::BindOutputs(FOutputVertexInterfaceData& InVertexData)
+	{
+		InVertexData.BindReadVertex(METASOUND_GET_PARAM_NAME(MidiClockTimestampTriggerPinNames::TriggerOutput), TriggerOutPin);
+	}
+	
+	void FMidiClockTimestampTriggerOperator::Reset(const FResetParams& ResetParams)
+	{
+		TriggerOutPin->Reset();
+		
+		FMidiPlayCursor::Reset(true);
+		
+		SetMessageFilter(FMidiPlayCursor::EFilterPassFlags::None);
+		MidiClockInPin->RegisterHiResPlayCursor(this);
+
+		CurrentTimestamp = *TimestampInPin;
+		bShouldQuantize = *ShouldQuantizeInPin;
+		QuantizeUnit = *QuantizeUnitInPin;
+		TriggerTick = 0;
+		CalculateTriggerTick();
+	}
+
+	void FMidiClockTimestampTriggerOperator::CalculateTriggerTick()
+	{
+		const FSongMaps& SongMaps = MidiClockInPin->GetSongMaps();
+		if (!bShouldQuantize)
+		{
+			TriggerTick = SongMaps.GetBarMap().MusicTimestampToTick(CurrentTimestamp);
+		}
+		else
+		{
+			const FBarMap& BarMap = SongMaps.GetBarMap();
+			int32 RawTick = BarMap.MusicTimestampToTick(CurrentTimestamp);
+			int32 BarTick = BarMap.MusicTimestampBarToTick(CurrentTimestamp.Bar);
+			int32 TicksPerQuantizationUnit = SubdivisionToMidiTicks(QuantizeUnit, RawTick, SongMaps);
+			if (ensure(TicksPerQuantizationUnit > 0))
+			{
+				float NumUnits = ((float)(RawTick - BarTick)) / (float)TicksPerQuantizationUnit;
+				int32 NumWholeUnits = FMath::RoundToInt(NumUnits);
+				TriggerTick = BarTick + (NumWholeUnits * TicksPerQuantizationUnit);
+			}
+			else
+			{
+				TriggerTick = SongMaps.GetBarMap().MusicTimestampToTick(CurrentTimestamp);
+			}
+		}
+	}
+
+	void FMidiClockTimestampTriggerOperator::Execute()
+	{
+		TriggerOutPin->AdvanceBlock();
+
+		// first let's see if our configuration has changed at all...
+		if (CurrentTimestamp != *TimestampInPin || bShouldQuantize != *ShouldQuantizeInPin || QuantizeUnit != *QuantizeUnitInPin)
+		{
+			CurrentTimestamp = *TimestampInPin;
+			bShouldQuantize = *ShouldQuantizeInPin;
+			QuantizeUnit = *QuantizeUnitInPin;
+			CalculateTriggerTick();
+		}
+
+		if (*EnableInPin)
+		{
+			for (const FTickSpan& Span : TickSpans)
+			{
+				if (Span.FromTick < TriggerTick && Span.ThruTick >= TriggerTick && 
+					(!Span.bIsSeek || *TriggerDuringSeekInPin))
+				{
+					TriggerOutPin->TriggerFrame(Span.BlockFrameIndex);
+					break;
+				}
+			}
+		}
+		TickSpans.Empty(8);
+	}
+
+	void FMidiClockTimestampTriggerOperator::SeekThruTick(int32 Tick)
+	{
+		int32 TickProceedingThisAdvance = CurrentTick;
+		FMidiPlayCursor::SeekThruTick(Tick);
+
+		// don't trigger if seeking backward or no progress being made...
+		if (Tick < TickProceedingThisAdvance || TickProceedingThisAdvance == CurrentTick)
+		{
+			return;
+		}
+
+		TickSpans.Emplace(TickProceedingThisAdvance, CurrentTick, MidiClockInPin->GetCurrentBlockFrameIndex(), true);
+	}
+
+	void FMidiClockTimestampTriggerOperator::AdvanceThruTick(int32 Tick, bool IsPreRoll)
+	{
+		int32 TickProceedingThisAdvance = CurrentTick;
+		FMidiPlayCursor::AdvanceThruTick(Tick, IsPreRoll);
+
+		// don't trigger during preroll or if no progress being made...
+		if (IsPreRoll || TickProceedingThisAdvance == CurrentTick)
+		{
+			return;
+		}
+
+		TickSpans.Emplace(TickProceedingThisAdvance, CurrentTick, MidiClockInPin->GetCurrentBlockFrameIndex(), false);
+	}
+}
+
+#undef LOCTEXT_NAMESPACE // "HarmonixMetaSound"
