- Add/Move Quantize functionality to FSongMaps
- Move some static functions from MusicClock.h/cpp to FSongMaps

diff --git a/Engine/Plugins/Runtime/Harmonix/Source/HarmonixMetasound/Private/HarmonixMetasound/DataTypes/MidiClock.cpp b/Engine/Plugins/Runtime/Harmonix/Source/HarmonixMetasound/Private/HarmonixMetasound/DataTypes/MidiClock.cpp
--- a/Engine/Plugins/Runtime/Harmonix/Source/HarmonixMetasound/Private/HarmonixMetasound/DataTypes/MidiClock.cpp
+++ b/Engine/Plugins/Runtime/Harmonix/Source/HarmonixMetasound/Private/HarmonixMetasound/DataTypes/MidiClock.cpp
@@ -12,6 +12,7 @@ DEFINE_LOG_CATEGORY_STATIC(LogHarmonixMidiClock, Log, All);
 namespace Metasound
 {
 	DEFINE_METASOUND_ENUM_BEGIN(EMidiClockSubdivisionQuantization, FEnumMidiClockSubdivisionQuantizationType, "SubdivisionQuantizationType")
+		DEFINE_METASOUND_ENUM_ENTRY(EMidiClockSubdivisionQuantization::None, "NoneDesc", "None", "NoneTT", "None"),
 		DEFINE_METASOUND_ENUM_ENTRY(EMidiClockSubdivisionQuantization::Bar, "BarDesc", "Bar", "BarTT", "Bar"),
 		DEFINE_METASOUND_ENUM_ENTRY(EMidiClockSubdivisionQuantization::Beat, "BeatDesc", "Beat", "BeatTT", "Beat"),
 		DEFINE_METASOUND_ENUM_ENTRY(EMidiClockSubdivisionQuantization::ThirtySecondNote, "ThirtySecondNoteDesc", "1/32", "ThirtySecondNoteTT", "1/32"),
@@ -32,97 +33,6 @@ namespace Metasound
 	DEFINE_METASOUND_ENUM_END()
 }
 
-float SubdivisionToBeats(EMidiClockSubdivisionQuantization Subdivision, const FTimeSignature& TimeSignature)
-{
-	// Easy cases first
-	if (Subdivision == EMidiClockSubdivisionQuantization::Bar)
-	{
-		return TimeSignature.Numerator;
-	}
-
-	if (Subdivision == EMidiClockSubdivisionQuantization::Beat)
-	{
-		return 1;
-	}
-
-	const float BeatsPerQuarter = TimeSignature.Denominator / 4.0f;
-	
-	switch (Subdivision)
-	{
-	case EMidiClockSubdivisionQuantization::ThirtySecondNote:
-		return BeatsPerQuarter / 8;
-	case EMidiClockSubdivisionQuantization::SixteenthNote:
-		return BeatsPerQuarter / 4;
-	case EMidiClockSubdivisionQuantization::EighthNote:
-		return BeatsPerQuarter / 2;
-	case EMidiClockSubdivisionQuantization::QuarterNote:
-		return BeatsPerQuarter;
-	case EMidiClockSubdivisionQuantization::HalfNote:
-		return BeatsPerQuarter * 2;
-	case EMidiClockSubdivisionQuantization::WholeNote:
-		return BeatsPerQuarter * 4;
-	case EMidiClockSubdivisionQuantization::DottedSixteenthNote:
-		return BeatsPerQuarter / 4 + BeatsPerQuarter / 8; 
-	case EMidiClockSubdivisionQuantization::DottedEighthNote:
-		return BeatsPerQuarter / 2 + BeatsPerQuarter / 4;
-	case EMidiClockSubdivisionQuantization::DottedQuarterNote:
-		return BeatsPerQuarter + BeatsPerQuarter / 2;
-	case EMidiClockSubdivisionQuantization::DottedHalfNote:
-		return BeatsPerQuarter * 3;
-	case EMidiClockSubdivisionQuantization::DottedWholeNote:
-		return BeatsPerQuarter * 6;
-	case EMidiClockSubdivisionQuantization::SixteenthNoteTriplet:
-		return (BeatsPerQuarter / 4) * 2 / 3;
-	case EMidiClockSubdivisionQuantization::EighthNoteTriplet:
-		return (BeatsPerQuarter / 2) * 2 / 3;
-	case EMidiClockSubdivisionQuantization::QuarterNoteTriplet:
-		return BeatsPerQuarter * 2 / 3;
-	case EMidiClockSubdivisionQuantization::HalfNoteTriplet:
-		return BeatsPerQuarter * 4 / 3;
-	default:
-		checkNoEntry();
-		return 0;
-	}
-}
-
-int32 SubdivisionToMidiTicks(const EMidiClockSubdivisionQuantization Division, const int32 CurrentTick, const FSongMaps& SongMap)
-{
-	int32 BarMapPointIndex = SongMap.GetBarMap().GetPointIndexForTick(CurrentTick);
-	if (BarMapPointIndex < 0)
-	{
-		return 0;
-	}
-	const FTimeSignature* TimeSignature =  SongMap.GetTimeSignatureAtTick(CurrentTick);
-	if (!TimeSignature)
-	{
-		return 0;
-	}
-
-	using namespace Harmonix::Midi::Constants;
-	
-	switch (Division)
-	{
-	case EMidiClockSubdivisionQuantization::Bar: 					return SongMap.GetBarMap().GetTicksInBarAfterPoint(BarMapPointIndex);
-	case EMidiClockSubdivisionQuantization::Beat:					return (GTicksPerQuarterNoteInt * 4) / TimeSignature->Denominator;
-	case EMidiClockSubdivisionQuantization::ThirtySecondNote:		return GTicksPerQuarterNoteInt / 8;
-	case EMidiClockSubdivisionQuantization::SixteenthNote:			return GTicksPerQuarterNoteInt / 4;
-	case EMidiClockSubdivisionQuantization::EighthNote:				return GTicksPerQuarterNoteInt / 2;
-	case EMidiClockSubdivisionQuantization::QuarterNote:			return GTicksPerQuarterNoteInt;
-	case EMidiClockSubdivisionQuantization::HalfNote:				return GTicksPerQuarterNoteInt * 2;
-	case EMidiClockSubdivisionQuantization::WholeNote:				return GTicksPerQuarterNoteInt * 4;
-	case EMidiClockSubdivisionQuantization::DottedSixteenthNote:	return (GTicksPerQuarterNoteInt / 4) + (GTicksPerQuarterNoteInt / 8);
-	case EMidiClockSubdivisionQuantization::DottedEighthNote:		return (GTicksPerQuarterNoteInt / 2) + (GTicksPerQuarterNoteInt / 4);
-	case EMidiClockSubdivisionQuantization::DottedQuarterNote:		return (GTicksPerQuarterNoteInt)     + (GTicksPerQuarterNoteInt / 2);
-	case EMidiClockSubdivisionQuantization::DottedHalfNote:			return (GTicksPerQuarterNoteInt * 2) + (GTicksPerQuarterNoteInt);
-	case EMidiClockSubdivisionQuantization::DottedWholeNote:		return (GTicksPerQuarterNoteInt * 4) + (GTicksPerQuarterNoteInt * 2);
-	case EMidiClockSubdivisionQuantization::SixteenthNoteTriplet:   return (GTicksPerQuarterNoteInt / 2) / 3;
-	case EMidiClockSubdivisionQuantization::EighthNoteTriplet:		return GTicksPerQuarterNoteInt / 3;
-	case EMidiClockSubdivisionQuantization::QuarterNoteTriplet:		return (GTicksPerQuarterNoteInt * 2) / 3;
-	case EMidiClockSubdivisionQuantization::HalfNoteTriplet:        return (GTicksPerQuarterNoteInt * 4) / 3;
-	default:	/* Beat */											return (GTicksPerQuarterNoteInt * 4) / TimeSignature->Denominator;
-	}
-}
-
 REGISTER_METASOUND_DATATYPE(HarmonixMetasound::FMidiClock, "MIDIClock")
 
 namespace HarmonixMetasound
diff --git a/Engine/Plugins/Runtime/Harmonix/Source/HarmonixMetasound/Private/HarmonixMetasound/DataTypes/MusicTimeInterval.cpp b/Engine/Plugins/Runtime/Harmonix/Source/HarmonixMetasound/Private/HarmonixMetasound/DataTypes/MusicTimeInterval.cpp
--- a/Engine/Plugins/Runtime/Harmonix/Source/HarmonixMetasound/Private/HarmonixMetasound/DataTypes/MusicTimeInterval.cpp
+++ b/Engine/Plugins/Runtime/Harmonix/Source/HarmonixMetasound/Private/HarmonixMetasound/DataTypes/MusicTimeInterval.cpp
@@ -1,12 +1,13 @@
 // Copyright Epic Games, Inc. All Rights Reserved.
 
 #include "HarmonixMetasound/DataTypes/MusicTimeInterval.h"
+#include "HarmonixMidi/SongMaps.h"
 
 namespace Harmonix
 {
 	float GetIntervalInBeats(const FMusicTimeInterval& Interval, const FTimeSignature& TimeSignature)
 	{
-		return SubdivisionToBeats(Interval.Interval, TimeSignature) * Interval.IntervalMultiplier;
+		return FSongMaps::SubdivisionToBeats(Interval.Interval, TimeSignature) * Interval.IntervalMultiplier;
 	}
 
 	void IncrementTimestampByBeats(FMusicTimestamp& Timestamp, float Beats, const FTimeSignature& TimeSignature)
@@ -42,7 +43,7 @@ namespace Harmonix
 	float GetOffsetInBeats(const FMusicTimeInterval& Interval, const FTimeSignature& TimeSignature)
 	{
 		const float IntervalBeats = GetIntervalInBeats(Interval, TimeSignature);
-		const float OffsetBeats = SubdivisionToBeats(Interval.Offset, TimeSignature) * Interval.OffsetMultiplier;
+		const float OffsetBeats = FSongMaps::SubdivisionToBeats(Interval.Offset, TimeSignature) * Interval.OffsetMultiplier;
 		return FMath::Fmod(OffsetBeats, IntervalBeats);
 	}
 
diff --git a/Engine/Plugins/Runtime/Harmonix/Source/HarmonixMetasound/Private/HarmonixMetasound/Nodes/MidiClockSubdivisionTriggerNode.cpp b/Engine/Plugins/Runtime/Harmonix/Source/HarmonixMetasound/Private/HarmonixMetasound/Nodes/MidiClockSubdivisionTriggerNode.cpp
--- a/Engine/Plugins/Runtime/Harmonix/Source/HarmonixMetasound/Private/HarmonixMetasound/Nodes/MidiClockSubdivisionTriggerNode.cpp
+++ b/Engine/Plugins/Runtime/Harmonix/Source/HarmonixMetasound/Private/HarmonixMetasound/Nodes/MidiClockSubdivisionTriggerNode.cpp
@@ -220,8 +220,8 @@ namespace HarmonixMetasound
 		GridOffsetMultiplier = *GridOffsetMultInPin;
 
 		const FSongMaps& SongMaps = MidiClockInPin->GetSongMaps();
-		GridOffsetTicks = SubdivisionToMidiTicks(GridOffsetUnits, 0, SongMaps) * GridOffsetMultiplier;
-		GridSizeTicks = SubdivisionToMidiTicks(GridSizeUnits, 0, SongMaps) * GridSizeMultiplier;
+		GridOffsetTicks = SongMaps.SubdivisionToMidiTicks(GridOffsetUnits, 0) * GridOffsetMultiplier;
+		GridSizeTicks = SongMaps.SubdivisionToMidiTicks(GridSizeUnits, 0) * GridSizeMultiplier;
 
 		AdvanceBlockNeeded = true;
 	}
@@ -244,7 +244,7 @@ namespace HarmonixMetasound
 				GridSizeMultiplier = 1;
 			}
 			const FSongMaps& SongMaps = MidiClockInPin->GetSongMaps();
-			GridSizeTicks = SubdivisionToMidiTicks(GridSizeUnits, CurrentTick, SongMaps) * GridSizeMultiplier;
+			GridSizeTicks = SongMaps.SubdivisionToMidiTicks(GridSizeUnits, CurrentTick) * GridSizeMultiplier;
 		}
 
 		if (*GridOffsetUnitsInPin != GridOffsetUnits || *GridOffsetMultInPin != GridOffsetMultiplier)
@@ -252,7 +252,7 @@ namespace HarmonixMetasound
 			GridOffsetUnits = *GridOffsetUnitsInPin;
 			GridOffsetMultiplier = *GridOffsetMultInPin;
 			const FSongMaps& SongMaps = MidiClockInPin->GetSongMaps();
-			GridOffsetTicks = SubdivisionToMidiTicks(GridOffsetUnits, 0, SongMaps) * GridOffsetMultiplier;
+			GridOffsetTicks = SongMaps.SubdivisionToMidiTicks(GridOffsetUnits, 0) * GridOffsetMultiplier;
 		}
 
 		AdvanceBlockNeeded = true;
@@ -316,7 +316,7 @@ namespace HarmonixMetasound
 	void FMidiClockSubdivisionTriggerOperator::OnTimeSig(int32 TrackIndex, int32 Tick, int32 Numerator, int32 Denominator, bool IsPreroll /*= false*/)
 	{
 		const FSongMaps& SongMaps = MidiClockInPin->GetSongMaps();
-		GridSizeTicks = SubdivisionToMidiTicks(GridSizeUnits, Tick, SongMaps) * GridSizeMultiplier;
+		GridSizeTicks = SongMaps.SubdivisionToMidiTicks(GridSizeUnits, Tick) * GridSizeMultiplier;
 		//NOTE: We DO NOT calculate a new grid offset because this is always specified from the beginning of the song... at the song's starting time signature.  
 	}
 
diff --git a/Engine/Plugins/Runtime/Harmonix/Source/HarmonixMetasound/Private/HarmonixMetasound/Nodes/MidiClockTimestampTriggerNode.cpp b/Engine/Plugins/Runtime/Harmonix/Source/HarmonixMetasound/Private/HarmonixMetasound/Nodes/MidiClockTimestampTriggerNode.cpp
--- a/Engine/Plugins/Runtime/Harmonix/Source/HarmonixMetasound/Private/HarmonixMetasound/Nodes/MidiClockTimestampTriggerNode.cpp
+++ b/Engine/Plugins/Runtime/Harmonix/Source/HarmonixMetasound/Private/HarmonixMetasound/Nodes/MidiClockTimestampTriggerNode.cpp
@@ -38,7 +38,6 @@ namespace HarmonixMetasound
 									const FBoolReadRef&                                 InEnabled,
 		                            const FMidiClockReadRef&                            InMidiClock,
 									const FMusicTimestampReadRef&                       InTimestamp,
-									const FBoolReadRef&                                 InShouldQuantize,
 									const FEnumMidiClockSubdivisionQuantizationReadRef& InQuanitizationUnit,
 									const FBoolReadRef&                                 InTriggerDuringSeek);
 		virtual ~FMidiClockTimestampTriggerOperator() override;
@@ -55,7 +54,6 @@ namespace HarmonixMetasound
 		FMidiClockReadRef                            MidiClockInPin;
 		FBoolReadRef                                 EnableInPin;
 		FMusicTimestampReadRef                       TimestampInPin;
-		FBoolReadRef                                 ShouldQuantizeInPin;
 		FEnumMidiClockSubdivisionQuantizationReadRef QuantizeUnitInPin;
 		FBoolReadRef		                         TriggerDuringSeekInPin;
 
@@ -64,8 +62,7 @@ namespace HarmonixMetasound
 
  		//** DATA (current state)
 		FMusicTimestamp CurrentTimestamp{1, 1.0f};
-		bool            bShouldQuantize = false;
-		EMidiClockSubdivisionQuantization QuantizeUnit = EMidiClockSubdivisionQuantization::ThirtySecondNote;
+		EMidiClockSubdivisionQuantization Quantize = EMidiClockSubdivisionQuantization::None;
 
 		int32 TriggerTick = 0;
 
@@ -144,8 +141,7 @@ namespace HarmonixMetasound
 				TInputDataVertex<FMidiClock>(METASOUND_GET_PARAM_NAME_AND_METADATA(Inputs::MidiClock)),
 				TInputDataVertex<bool>(METASOUND_GET_PARAM_NAME_AND_METADATA(Inputs::Enable), true),
 				TInputDataVertex<FMusicTimestamp>(METASOUND_GET_PARAM_NAME_AND_METADATA(Inputs::Timestamp)),
-				TInputDataVertex<bool>(METASOUND_GET_PARAM_NAME_AND_METADATA(Inputs::ShouldQuantizeTimestamp), false),
-				TInputDataVertex<FEnumMidiClockSubdivisionQuantizationType>(METASOUND_GET_PARAM_NAME_AND_METADATA(Inputs::QuantizationUnit), (int32)EMidiClockSubdivisionQuantization::Beat),
+				TInputDataVertex<FEnumMidiClockSubdivisionQuantizationType>(METASOUND_GET_PARAM_NAME_AND_METADATA(Inputs::QuantizationUnit), (int32)EMidiClockSubdivisionQuantization::None),
 				TInputDataVertex<bool>(METASOUND_GET_PARAM_NAME_AND_METADATA(TriggerDuringSeek), false)
 				),
 			FOutputVertexInterface(
@@ -167,7 +163,6 @@ namespace HarmonixMetasound
 		FMidiClockReadRef InMidiClock      = InputData.GetOrConstructDataReadReference<FMidiClock>(METASOUND_GET_PARAM_NAME(Inputs::MidiClock), InParams.OperatorSettings);
 		FBoolReadRef  InEnabled            = InputData.GetOrCreateDefaultDataReadReference<bool>(METASOUND_GET_PARAM_NAME(Inputs::Enable), InParams.OperatorSettings);
 		FMusicTimestampReadRef InTimestamp = InputData.GetOrConstructDataReadReference<FMusicTimestamp>(METASOUND_GET_PARAM_NAME(Inputs::Timestamp), 1, 1.0f);
-		FBoolReadRef  InShouldQuantize     = InputData.GetOrCreateDefaultDataReadReference<bool>(METASOUND_GET_PARAM_NAME(Inputs::ShouldQuantizeTimestamp), InParams.OperatorSettings);
 		FEnumMidiClockSubdivisionQuantizationReadRef InQuantizeUnits = InputData.GetOrCreateDefaultDataReadReference<FEnumMidiClockSubdivisionQuantizationType>(METASOUND_GET_PARAM_NAME(Inputs::QuantizationUnit), InParams.OperatorSettings);
 		FBoolReadRef  InTriggerDuringSeek  = InputData.GetOrCreateDefaultDataReadReference<bool>(METASOUND_GET_PARAM_NAME(TriggerDuringSeek), InParams.OperatorSettings);
 
@@ -175,22 +170,20 @@ namespace HarmonixMetasound
 			InEnabled,
 			InMidiClock,
 			InTimestamp,
-			InShouldQuantize,
 			InQuantizeUnits,
 			InTriggerDuringSeek);
 	}
 
-	FMidiClockTimestampTriggerOperator::FMidiClockTimestampTriggerOperator(const FBuildOperatorParams&           InParams,
-															 const FBoolReadRef&                                 InEnabled,
-															 const FMidiClockReadRef&                            InMidiClock,
-															 const FMusicTimestampReadRef&                       InTimestamp,
-															 const FBoolReadRef&                                 InShouldQuantize,
-															 const FEnumMidiClockSubdivisionQuantizationReadRef& InQuanitizationUnit,
-															 const FBoolReadRef&                                 InTriggerDuringSeek)
+	FMidiClockTimestampTriggerOperator::FMidiClockTimestampTriggerOperator(
+			const FBuildOperatorParams&                         InParams,
+			const FBoolReadRef&                                 InEnabled,
+			const FMidiClockReadRef&                            InMidiClock,
+			const FMusicTimestampReadRef&                       InTimestamp,
+			const FEnumMidiClockSubdivisionQuantizationReadRef& InQuanitizationUnit,
+			const FBoolReadRef&                                 InTriggerDuringSeek)
 		: MidiClockInPin(InMidiClock)
 		, EnableInPin(InEnabled)
 		, TimestampInPin(InTimestamp)
-		, ShouldQuantizeInPin(InShouldQuantize)
 		, QuantizeUnitInPin(InQuanitizationUnit)
 		, TriggerDuringSeekInPin(InTriggerDuringSeek)
 		, TriggerOutPin(FTriggerWriteRef::CreateNew(InParams.OperatorSettings))
@@ -207,10 +200,9 @@ namespace HarmonixMetasound
 		using namespace CommonPinNames;
 		using namespace MidiClockTimestampTriggerPinNames;
 
-		InVertexData.BindReadVertex(METASOUND_GET_PARAM_NAME(Inputs::Enable),    EnableInPin);
+		InVertexData.BindReadVertex(METASOUND_GET_PARAM_NAME(Inputs::Enable), EnableInPin);
 		InVertexData.BindReadVertex(METASOUND_GET_PARAM_NAME(Inputs::MidiClock), MidiClockInPin);
 		InVertexData.BindReadVertex(METASOUND_GET_PARAM_NAME(Inputs::Timestamp), TimestampInPin);
-		InVertexData.BindReadVertex(METASOUND_GET_PARAM_NAME(Inputs::ShouldQuantizeTimestamp), ShouldQuantizeInPin);
 		InVertexData.BindReadVertex(METASOUND_GET_PARAM_NAME(Inputs::QuantizationUnit), QuantizeUnitInPin);
 		InVertexData.BindReadVertex(METASOUND_GET_PARAM_NAME(TriggerDuringSeek), TriggerDuringSeekInPin);
 	}
@@ -230,8 +222,7 @@ namespace HarmonixMetasound
 		MidiClockInPin->RegisterHiResPlayCursor(this);
 
 		CurrentTimestamp = *TimestampInPin;
-		bShouldQuantize = *ShouldQuantizeInPin;
-		QuantizeUnit = *QuantizeUnitInPin;
+		
 		TriggerTick = 0;
 		CalculateTriggerTick();
 	}
@@ -239,27 +230,7 @@ namespace HarmonixMetasound
 	void FMidiClockTimestampTriggerOperator::CalculateTriggerTick()
 	{
 		const FSongMaps& SongMaps = MidiClockInPin->GetSongMaps();
-		if (!bShouldQuantize)
-		{
-			TriggerTick = SongMaps.GetBarMap().MusicTimestampToTick(CurrentTimestamp);
-		}
-		else
-		{
-			const FBarMap& BarMap = SongMaps.GetBarMap();
-			int32 RawTick = BarMap.MusicTimestampToTick(CurrentTimestamp);
-			int32 BarTick = BarMap.MusicTimestampBarToTick(CurrentTimestamp.Bar);
-			int32 TicksPerQuantizationUnit = SubdivisionToMidiTicks(QuantizeUnit, RawTick, SongMaps);
-			if (ensure(TicksPerQuantizationUnit > 0))
-			{
-				float NumUnits = ((float)(RawTick - BarTick)) / (float)TicksPerQuantizationUnit;
-				int32 NumWholeUnits = FMath::RoundToInt(NumUnits);
-				TriggerTick = BarTick + (NumWholeUnits * TicksPerQuantizationUnit);
-			}
-			else
-			{
-				TriggerTick = SongMaps.GetBarMap().MusicTimestampToTick(CurrentTimestamp);
-			}
-		}
+		TriggerTick = SongMaps.CalculateMidiTick(CurrentTimestamp, Quantize);
 	}
 
 	void FMidiClockTimestampTriggerOperator::Execute()
@@ -267,11 +238,10 @@ namespace HarmonixMetasound
 		TriggerOutPin->AdvanceBlock();
 
 		// first let's see if our configuration has changed at all...
-		if (CurrentTimestamp != *TimestampInPin || bShouldQuantize != *ShouldQuantizeInPin || QuantizeUnit != *QuantizeUnitInPin)
+		if (CurrentTimestamp != *TimestampInPin || Quantize != *QuantizeUnitInPin)
 		{
 			CurrentTimestamp = *TimestampInPin;
-			bShouldQuantize = *ShouldQuantizeInPin;
-			QuantizeUnit = *QuantizeUnitInPin;
+			Quantize = *QuantizeUnitInPin;
 			CalculateTriggerTick();
 		}
 
diff --git a/Engine/Plugins/Runtime/Harmonix/Source/HarmonixMetasound/Public/HarmonixMetasound/DataTypes/MidiClock.h b/Engine/Plugins/Runtime/Harmonix/Source/HarmonixMetasound/Public/HarmonixMetasound/DataTypes/MidiClock.h
--- a/Engine/Plugins/Runtime/Harmonix/Source/HarmonixMetasound/Public/HarmonixMetasound/DataTypes/MidiClock.h
+++ b/Engine/Plugins/Runtime/Harmonix/Source/HarmonixMetasound/Public/HarmonixMetasound/DataTypes/MidiClock.h
@@ -13,39 +13,17 @@
 
 #include "HarmonixMidi/MidiFile.h"
 #include "HarmonixMidi/SongMaps.h"
-#include "HarmonixMidi/MidiFile.h"
 #include "HarmonixMetasound/DataTypes/MusicTransport.h"
 #include "HarmonixMetasound/DataTypes/MidiClockEvent.h"
 #include "HarmonixMidi/MidiPlayCursorMgr.h"
 #include "HarmonixMidi/MidiPlayCursor.h"
 #include "Sound/QuartzQuantizationUtilities.h"
 
-enum class EMidiClockSubdivisionQuantization : uint8
-{
-	Bar						= static_cast<uint8>(EQuartzCommandQuantization::Bar),
-	Beat					= static_cast<uint8>(EQuartzCommandQuantization::Beat),
-	ThirtySecondNote		= static_cast<uint8>(EQuartzCommandQuantization::ThirtySecondNote),
-	SixteenthNote			= static_cast<uint8>(EQuartzCommandQuantization::SixteenthNote),
-	EighthNote				= static_cast<uint8>(EQuartzCommandQuantization::EighthNote),
-	QuarterNote				= static_cast<uint8>(EQuartzCommandQuantization::QuarterNote),
-	HalfNote				= static_cast<uint8>(EQuartzCommandQuantization::HalfNote),
-	WholeNote				= static_cast<uint8>(EQuartzCommandQuantization::WholeNote),
-	DottedSixteenthNote		= static_cast<uint8>(EQuartzCommandQuantization::DottedSixteenthNote),
-	DottedEighthNote		= static_cast<uint8>(EQuartzCommandQuantization::DottedEighthNote),
-	DottedQuarterNote		= static_cast<uint8>(EQuartzCommandQuantization::DottedQuarterNote),
-	DottedHalfNote			= static_cast<uint8>(EQuartzCommandQuantization::DottedHalfNote),
-	DottedWholeNote			= static_cast<uint8>(EQuartzCommandQuantization::DottedWholeNote),
-	SixteenthNoteTriplet	= static_cast<uint8>(EQuartzCommandQuantization::SixteenthNoteTriplet),
-	EighthNoteTriplet		= static_cast<uint8>(EQuartzCommandQuantization::EighthNoteTriplet),
-	QuarterNoteTriplet		= static_cast<uint8>(EQuartzCommandQuantization::QuarterNoteTriplet),
-	HalfNoteTriplet			= static_cast<uint8>(EQuartzCommandQuantization::HalfNoteTriplet),
-};
-
 namespace Metasound
 {
 	DECLARE_METASOUND_ENUM(
 		EMidiClockSubdivisionQuantization,
-		EMidiClockSubdivisionQuantization::Beat,
+		EMidiClockSubdivisionQuantization::None,
 		HARMONIXMETASOUND_API,
 		FEnumMidiClockSubdivisionQuantizationType,
 		FEnumMidiClockSubdivisionQuantizationTypeInfo,
@@ -54,10 +32,6 @@ namespace Metasound
 	);
 }
 
-HARMONIXMETASOUND_API float SubdivisionToBeats(EMidiClockSubdivisionQuantization Subdivision, const FTimeSignature& TimeSignature);
-
-HARMONIXMETASOUND_API int32 SubdivisionToMidiTicks(EMidiClockSubdivisionQuantization Division, int32 CurrentTick, const FSongMaps& SongMap);
-
 namespace HarmonixMetasound
 {
 	struct HARMONIXMETASOUND_API FMidiTimestampTransportState
diff --git a/Engine/Plugins/Runtime/Harmonix/Source/HarmonixMidi/Private/HarmonixMidi/SongMaps.cpp b/Engine/Plugins/Runtime/Harmonix/Source/HarmonixMidi/Private/HarmonixMidi/SongMaps.cpp
--- a/Engine/Plugins/Runtime/Harmonix/Source/HarmonixMidi/Private/HarmonixMidi/SongMaps.cpp
+++ b/Engine/Plugins/Runtime/Harmonix/Source/HarmonixMidi/Private/HarmonixMidi/SongMaps.cpp
@@ -542,6 +542,124 @@ void FSongMaps::SetLengthTotalBars(int32 Bars)
 	LengthData.LengthBars = Bars;
 }
 
+int32 FSongMaps::CalculateMidiTick(const FMusicTimestamp& Timestamp, const EMidiClockSubdivisionQuantization Quantize) const
+{
+	int32 TriggerTick = 0;
+	if (Quantize == EMidiClockSubdivisionQuantization::None)
+	{
+		TriggerTick = FMath::RoundToInt32(BarMap.MusicTimestampToTick(Timestamp));
+	}
+	else
+	{
+		int32 RawTick = FMath::RoundToInt32(BarMap.MusicTimestampToTick(Timestamp));
+		int32 BarTick = BarMap.MusicTimestampBarToTick(Timestamp.Bar);
+		int32 TicksPerQuantizationUnit = SubdivisionToMidiTicks(Quantize, RawTick);
+		if (ensure(TicksPerQuantizationUnit > 0))
+		{
+			float NumUnits = ((float)(RawTick - BarTick)) / (float)TicksPerQuantizationUnit;
+			int32 NumWholeUnits = FMath::RoundToInt32(NumUnits);
+			TriggerTick = BarTick + (NumWholeUnits * TicksPerQuantizationUnit);
+		}
+		else
+		{
+			TriggerTick = FMath::RoundToInt32(BarMap.MusicTimestampToTick(Timestamp));
+		}
+	}
+	return TriggerTick;
+}
+
+int32 FSongMaps::SubdivisionToMidiTicks(const EMidiClockSubdivisionQuantization Division, const int32 Tick) const
+{
+	int32 BarMapPointIndex = GetBarMap().GetPointIndexForTick(Tick);
+	if (BarMapPointIndex < 0)
+	{
+		return 1;
+	}
+	const FTimeSignature* TimeSignature = GetTimeSignatureAtTick(Tick);
+	if (!TimeSignature)
+	{
+		return 1;
+	}
+
+	using namespace Harmonix::Midi::Constants;
+
+	switch (Division)
+	{
+	case EMidiClockSubdivisionQuantization::None:                   return 1;
+	case EMidiClockSubdivisionQuantization::Bar: 					return GetBarMap().GetTicksInBarAfterPoint(BarMapPointIndex);
+	case EMidiClockSubdivisionQuantization::Beat:					return (GTicksPerQuarterNoteInt * 4) / TimeSignature->Denominator;
+	case EMidiClockSubdivisionQuantization::ThirtySecondNote:		return GTicksPerQuarterNoteInt / 8;
+	case EMidiClockSubdivisionQuantization::SixteenthNote:			return GTicksPerQuarterNoteInt / 4;
+	case EMidiClockSubdivisionQuantization::EighthNote:				return GTicksPerQuarterNoteInt / 2;
+	case EMidiClockSubdivisionQuantization::QuarterNote:			return GTicksPerQuarterNoteInt;
+	case EMidiClockSubdivisionQuantization::HalfNote:				return GTicksPerQuarterNoteInt * 2;
+	case EMidiClockSubdivisionQuantization::WholeNote:				return GTicksPerQuarterNoteInt * 4;
+	case EMidiClockSubdivisionQuantization::DottedSixteenthNote:	return (GTicksPerQuarterNoteInt / 4) + (GTicksPerQuarterNoteInt / 8);
+	case EMidiClockSubdivisionQuantization::DottedEighthNote:		return (GTicksPerQuarterNoteInt / 2) + (GTicksPerQuarterNoteInt / 4);
+	case EMidiClockSubdivisionQuantization::DottedQuarterNote:		return (GTicksPerQuarterNoteInt)+(GTicksPerQuarterNoteInt / 2);
+	case EMidiClockSubdivisionQuantization::DottedHalfNote:			return (GTicksPerQuarterNoteInt * 2) + (GTicksPerQuarterNoteInt);
+	case EMidiClockSubdivisionQuantization::DottedWholeNote:		return (GTicksPerQuarterNoteInt * 4) + (GTicksPerQuarterNoteInt * 2);
+	case EMidiClockSubdivisionQuantization::SixteenthNoteTriplet:   return (GTicksPerQuarterNoteInt / 2) / 3;
+	case EMidiClockSubdivisionQuantization::EighthNoteTriplet:		return GTicksPerQuarterNoteInt / 3;
+	case EMidiClockSubdivisionQuantization::QuarterNoteTriplet:		return (GTicksPerQuarterNoteInt * 2) / 3;
+	case EMidiClockSubdivisionQuantization::HalfNoteTriplet:        return (GTicksPerQuarterNoteInt * 4) / 3;
+	default:	                                             		checkNoEntry();	return 1;
+	}
+}
+
+float FSongMaps::SubdivisionToBeats(EMidiClockSubdivisionQuantization Subdivision, const FTimeSignature& TimeSignature)
+{
+	// Easy cases first
+	if (Subdivision == EMidiClockSubdivisionQuantization::Bar)
+	{
+		return TimeSignature.Numerator;
+	}
+
+	if (Subdivision == EMidiClockSubdivisionQuantization::Beat)
+	{
+		return 1;
+	}
+
+	const float BeatsPerQuarter = TimeSignature.Denominator / 4.0f;
+
+	switch (Subdivision)
+	{
+	case EMidiClockSubdivisionQuantization::ThirtySecondNote:
+		return BeatsPerQuarter / 8;
+	case EMidiClockSubdivisionQuantization::SixteenthNote:
+		return BeatsPerQuarter / 4;
+	case EMidiClockSubdivisionQuantization::EighthNote:
+		return BeatsPerQuarter / 2;
+	case EMidiClockSubdivisionQuantization::QuarterNote:
+		return BeatsPerQuarter;
+	case EMidiClockSubdivisionQuantization::HalfNote:
+		return BeatsPerQuarter * 2;
+	case EMidiClockSubdivisionQuantization::WholeNote:
+		return BeatsPerQuarter * 4;
+	case EMidiClockSubdivisionQuantization::DottedSixteenthNote:
+		return BeatsPerQuarter / 4 + BeatsPerQuarter / 8;
+	case EMidiClockSubdivisionQuantization::DottedEighthNote:
+		return BeatsPerQuarter / 2 + BeatsPerQuarter / 4;
+	case EMidiClockSubdivisionQuantization::DottedQuarterNote:
+		return BeatsPerQuarter + BeatsPerQuarter / 2;
+	case EMidiClockSubdivisionQuantization::DottedHalfNote:
+		return BeatsPerQuarter * 3;
+	case EMidiClockSubdivisionQuantization::DottedWholeNote:
+		return BeatsPerQuarter * 6;
+	case EMidiClockSubdivisionQuantization::SixteenthNoteTriplet:
+		return (BeatsPerQuarter / 4) * 2 / 3;
+	case EMidiClockSubdivisionQuantization::EighthNoteTriplet:
+		return (BeatsPerQuarter / 2) * 2 / 3;
+	case EMidiClockSubdivisionQuantization::QuarterNoteTriplet:
+		return BeatsPerQuarter * 2 / 3;
+	case EMidiClockSubdivisionQuantization::HalfNoteTriplet:
+		return BeatsPerQuarter * 4 / 3;
+	default:
+		checkNoEntry();
+		return 0;
+	}
+}
+
 ///////////////////////////////////////////////////////////////////////////////////
 // Sections
 
diff --git a/Engine/Plugins/Runtime/Harmonix/Source/HarmonixMidi/Public/HarmonixMidi/SongMaps.h b/Engine/Plugins/Runtime/Harmonix/Source/HarmonixMidi/Public/HarmonixMidi/SongMaps.h
--- a/Engine/Plugins/Runtime/Harmonix/Source/HarmonixMidi/Public/HarmonixMidi/SongMaps.h
+++ b/Engine/Plugins/Runtime/Harmonix/Source/HarmonixMidi/Public/HarmonixMidi/SongMaps.h
@@ -5,6 +5,7 @@
 #include "HarmonixMidi/BeatMap.h"
 #include "HarmonixMidi/ChordMap.h"
 #include "HarmonixMidi/SectionMap.h"
+#include "Sound/QuartzQuantizationUtilities.h"
 
 #include "SongMaps.generated.h"
 
@@ -13,6 +14,29 @@ class FStdMidiFileReader;
 class FSongMapReceiver;
 class UMidiFile;
 
+UENUM(BlueprintType)
+enum class EMidiClockSubdivisionQuantization : uint8
+{
+	Bar = static_cast<uint8>(EQuartzCommandQuantization::Bar),
+	Beat = static_cast<uint8>(EQuartzCommandQuantization::Beat),
+	ThirtySecondNote = static_cast<uint8>(EQuartzCommandQuantization::ThirtySecondNote),
+	SixteenthNote = static_cast<uint8>(EQuartzCommandQuantization::SixteenthNote),
+	EighthNote = static_cast<uint8>(EQuartzCommandQuantization::EighthNote),
+	QuarterNote = static_cast<uint8>(EQuartzCommandQuantization::QuarterNote),
+	HalfNote = static_cast<uint8>(EQuartzCommandQuantization::HalfNote),
+	WholeNote = static_cast<uint8>(EQuartzCommandQuantization::WholeNote),
+	DottedSixteenthNote = static_cast<uint8>(EQuartzCommandQuantization::DottedSixteenthNote),
+	DottedEighthNote = static_cast<uint8>(EQuartzCommandQuantization::DottedEighthNote),
+	DottedQuarterNote = static_cast<uint8>(EQuartzCommandQuantization::DottedQuarterNote),
+	DottedHalfNote = static_cast<uint8>(EQuartzCommandQuantization::DottedHalfNote),
+	DottedWholeNote = static_cast<uint8>(EQuartzCommandQuantization::DottedWholeNote),
+	SixteenthNoteTriplet = static_cast<uint8>(EQuartzCommandQuantization::SixteenthNoteTriplet),
+	EighthNoteTriplet = static_cast<uint8>(EQuartzCommandQuantization::EighthNoteTriplet),
+	QuarterNoteTriplet = static_cast<uint8>(EQuartzCommandQuantization::QuarterNoteTriplet),
+	HalfNoteTriplet = static_cast<uint8>(EQuartzCommandQuantization::HalfNoteTriplet),
+	None = static_cast<uint8>(EQuartzCommandQuantization::None)
+};
+
 USTRUCT()
 struct FSongLengthData
 {
@@ -114,6 +138,10 @@ public:
 	bool                  BarMapIsEmpty() { return BarMap.GetNumTimeSignaturePoints() == 0; }
 	void                  EmptyBarMap() { BarMap.Empty(); }
 	void                  SetLengthTotalBars(int32 Bars);
+	int32                 CalculateMidiTick(const FMusicTimestamp& Timestamp, const EMidiClockSubdivisionQuantization Quantize) const;
+	int32                 SubdivisionToMidiTicks(const EMidiClockSubdivisionQuantization Division, const int32 Tick) const;
+	static float          SubdivisionToBeats(EMidiClockSubdivisionQuantization Subdivision, const FTimeSignature& TimeSignature);
+
 
 	// sections
 	float               GetSectionStartMsAtMs(float Ms) const;
