[Harmonix] Reference display name and description text instead of attempt to copy. Attempting to copy was causing crashes after changing a shared definition and re-buildng.

diff --git a/Engine/Plugins/Runtime/Harmonix/Source/HarmonixMetasound/Private/HarmonixMetasound/Nodes/MidiCCTriggerNode.cpp b/Engine/Plugins/Runtime/Harmonix/Source/HarmonixMetasound/Private/HarmonixMetasound/Nodes/MidiCCTriggerNode.cpp
--- a/Engine/Plugins/Runtime/Harmonix/Source/HarmonixMetasound/Private/HarmonixMetasound/Nodes/MidiCCTriggerNode.cpp
+++ b/Engine/Plugins/Runtime/Harmonix/Source/HarmonixMetasound/Private/HarmonixMetasound/Nodes/MidiCCTriggerNode.cpp
@@ -45,9 +45,9 @@ namespace HarmonixMetasound::Nodes::MidiCCTriggerNode
 		DEFINE_INPUT_METASOUND_PARAM(InputMidiControllerID, "Control Number", "MIDI Control Number");
 
 		// Removed in V1
-		DECLARE_METASOUND_PARAM_EXTERN(MidiTrackNumber);
+		DECLARE_METASOUND_PARAM_ALIAS(MidiTrackNumber);
 		DEFINE_METASOUND_PARAM_ALIAS(MidiTrackNumber, CommonPinNames::Inputs::MidiTrackNumber);
-		DECLARE_METASOUND_PARAM_EXTERN(MidiChannelNumber);
+		DECLARE_METASOUND_PARAM_ALIAS(MidiChannelNumber);
 		DEFINE_METASOUND_PARAM_ALIAS(MidiChannelNumber, CommonPinNames::Inputs::MidiChannelNumber);
 	}
 
diff --git a/Engine/Plugins/Runtime/Harmonix/Source/HarmonixMetasound/Private/HarmonixMetasound/Nodes/PeakNode.cpp b/Engine/Plugins/Runtime/Harmonix/Source/HarmonixMetasound/Private/HarmonixMetasound/Nodes/PeakNode.cpp
--- a/Engine/Plugins/Runtime/Harmonix/Source/HarmonixMetasound/Private/HarmonixMetasound/Nodes/PeakNode.cpp
+++ b/Engine/Plugins/Runtime/Harmonix/Source/HarmonixMetasound/Private/HarmonixMetasound/Nodes/PeakNode.cpp
@@ -29,8 +29,8 @@ namespace HarmonixMetasound::Nodes::Peak
 
 	namespace Inputs
 	{
+		DEFINE_INPUT_METASOUND_PARAM(Enable, "Enable", "Enable the peak node.")
 		DEFINE_METASOUND_PARAM_ALIAS(AudioMono, CommonPinNames::Inputs::AudioMono);
-		METASOUND_PARAM(ParamEnable, "Enable", "Enable the peak node.")
 	}
 
 	namespace Outputs
@@ -50,7 +50,7 @@ namespace HarmonixMetasound::Nodes::Peak
 				return {
 					FInputVertexInterface
 					{
-						TInputDataVertex<bool>(METASOUND_GET_PARAM_NAME_AND_METADATA(Inputs::ParamEnable), true),
+						TInputDataVertex<bool>(METASOUND_GET_PARAM_NAME_AND_METADATA(Inputs::Enable), true),
 						TInputDataVertex<FAudioBuffer>(METASOUND_GET_PARAM_NAME_AND_METADATA(Inputs::AudioMono))
 					},
 					Metasound::FOutputVertexInterface
@@ -105,7 +105,7 @@ namespace HarmonixMetasound::Nodes::Peak
 					Inputs::AudioMonoName,
 					InParams.OperatorSettings),
 				InParams.InputData.GetOrCreateDefaultDataReadReference<bool>(
-					Inputs::ParamEnableName,
+					Inputs::EnableName,
 					InParams.OperatorSettings)
 			};
 
@@ -127,7 +127,7 @@ namespace HarmonixMetasound::Nodes::Peak
 		virtual void BindInputs(Metasound::FInputVertexInterfaceData& InVertexData) override
 		{
 			InVertexData.BindReadVertex(Inputs::AudioMonoName, Inputs.Audio);
-			InVertexData.BindReadVertex(Inputs::ParamEnableName, Inputs.Enable);
+			InVertexData.BindReadVertex(Inputs::EnableName, Inputs.Enable);
 
 			UpdatePeak();
 		}
diff --git a/Engine/Plugins/Runtime/Harmonix/Source/HarmonixMetasound/Public/HarmonixMetasound/Common.h b/Engine/Plugins/Runtime/Harmonix/Source/HarmonixMetasound/Public/HarmonixMetasound/Common.h
--- a/Engine/Plugins/Runtime/Harmonix/Source/HarmonixMetasound/Public/HarmonixMetasound/Common.h
+++ b/Engine/Plugins/Runtime/Harmonix/Source/HarmonixMetasound/Public/HarmonixMetasound/Common.h
@@ -26,7 +26,12 @@ namespace HarmonixMetasound
 #define DECLARE_METASOUND_PARAM_EXTERN(NAME)      \
 	HARMONIXMETASOUND_API extern const TCHAR* NAME##Name;       \
 	HARMONIXMETASOUND_API extern const FText NAME##Tooltip;     \
-	HARMONIXMETASOUND_API extern const FText NAME##DisplayName; 
+	HARMONIXMETASOUND_API extern const FText NAME##DisplayName;
+
+#define DECLARE_METASOUND_PARAM_ALIAS(NAME)      \
+	HARMONIXMETASOUND_API extern const TCHAR* NAME##Name;       \
+	HARMONIXMETASOUND_API extern const FText& NAME##Tooltip;     \
+	HARMONIXMETASOUND_API extern const FText& NAME##DisplayName;
 
 #if WITH_EDITOR
 #define LOC_DEFINE_REGION
@@ -53,9 +58,9 @@ namespace HarmonixMetasound
 // Alias an existing param definition. Helpful for re-using definitions while providing a public node API.
 #define DEFINE_METASOUND_PARAM_ALIAS(ALIAS_NAME, TARGET_NAME) \
 	const TCHAR* ALIAS_NAME##Name = TARGET_NAME##Name; \
-	const FText ALIAS_NAME##Tooltip = TARGET_NAME##Tooltip; \
-	const FText ALIAS_NAME##DisplayName = TARGET_NAME##DisplayName; 
-	
+	const FText& ALIAS_NAME##Tooltip = TARGET_NAME##Tooltip; \
+	const FText& ALIAS_NAME##DisplayName = TARGET_NAME##DisplayName; 
+
 	namespace CommonPinNames
 	{
 		namespace Inputs
diff --git a/Engine/Plugins/Runtime/Harmonix/Source/HarmonixMetasound/Public/HarmonixMetasound/Nodes/DelayNode.h b/Engine/Plugins/Runtime/Harmonix/Source/HarmonixMetasound/Public/HarmonixMetasound/Nodes/DelayNode.h
--- a/Engine/Plugins/Runtime/Harmonix/Source/HarmonixMetasound/Public/HarmonixMetasound/Nodes/DelayNode.h
+++ b/Engine/Plugins/Runtime/Harmonix/Source/HarmonixMetasound/Public/HarmonixMetasound/Nodes/DelayNode.h
@@ -14,9 +14,9 @@ namespace HarmonixMetasound::DelayNode
 
 	namespace Inputs
 	{
-		DECLARE_METASOUND_PARAM_EXTERN(AudioLeft);
-		DECLARE_METASOUND_PARAM_EXTERN(AudioRight);
-		DECLARE_METASOUND_PARAM_EXTERN(MidiClock)
+		DECLARE_METASOUND_PARAM_ALIAS(AudioLeft);
+		DECLARE_METASOUND_PARAM_ALIAS(AudioRight);
+		DECLARE_METASOUND_PARAM_ALIAS(MidiClock)
 		DECLARE_METASOUND_PARAM_EXTERN(DelayTimeType);
 		DECLARE_METASOUND_PARAM_EXTERN(DelayTime);
 		DECLARE_METASOUND_PARAM_EXTERN(Feedback);
@@ -29,7 +29,7 @@ namespace HarmonixMetasound::DelayNode
 		DECLARE_METASOUND_PARAM_EXTERN(FilterQ);
 		DECLARE_METASOUND_PARAM_EXTERN(LFOEnabled);
 		DECLARE_METASOUND_PARAM_EXTERN(LFOTimeType);
-		DECLARE_METASOUND_PARAM_EXTERN(LFOFrequency);
+		DECLARE_METASOUND_PARAM_ALIAS(LFOFrequency);
 		DECLARE_METASOUND_PARAM_EXTERN(LFODepth);
 		DECLARE_METASOUND_PARAM_EXTERN(StereoType);
 		DECLARE_METASOUND_PARAM_EXTERN(StereoSpreadLeft);
@@ -38,7 +38,7 @@ namespace HarmonixMetasound::DelayNode
 
 	namespace Outputs
 	{
-		DECLARE_METASOUND_PARAM_EXTERN(AudioLeft);
-		DECLARE_METASOUND_PARAM_EXTERN(AudioRight);
+		DECLARE_METASOUND_PARAM_ALIAS(AudioLeft);
+		DECLARE_METASOUND_PARAM_ALIAS(AudioRight);
 	}
 }
diff --git a/Engine/Plugins/Runtime/Harmonix/Source/HarmonixMetasound/Public/HarmonixMetasound/Nodes/DjFilterNode.h b/Engine/Plugins/Runtime/Harmonix/Source/HarmonixMetasound/Public/HarmonixMetasound/Nodes/DjFilterNode.h
--- a/Engine/Plugins/Runtime/Harmonix/Source/HarmonixMetasound/Public/HarmonixMetasound/Nodes/DjFilterNode.h
+++ b/Engine/Plugins/Runtime/Harmonix/Source/HarmonixMetasound/Public/HarmonixMetasound/Nodes/DjFilterNode.h
@@ -12,7 +12,7 @@ namespace HarmonixMetasound::Nodes::DjFilter
 	
 	namespace Inputs
 	{
-		DECLARE_METASOUND_PARAM_EXTERN(AudioMono);
+		DECLARE_METASOUND_PARAM_ALIAS(AudioMono);
 		DECLARE_METASOUND_PARAM_EXTERN(Amount);
 		DECLARE_METASOUND_PARAM_EXTERN(Resonance);
 		DECLARE_METASOUND_PARAM_EXTERN(LowPassMinFrequency);
@@ -24,6 +24,6 @@ namespace HarmonixMetasound::Nodes::DjFilter
 
 	namespace Outputs
 	{
-		DECLARE_METASOUND_PARAM_EXTERN(AudioMono);
+		DECLARE_METASOUND_PARAM_ALIAS(AudioMono);
 	}
 }
diff --git a/Engine/Plugins/Runtime/Harmonix/Source/HarmonixMetasound/Public/HarmonixMetasound/Nodes/MidiCCTriggerNode.h b/Engine/Plugins/Runtime/Harmonix/Source/HarmonixMetasound/Public/HarmonixMetasound/Nodes/MidiCCTriggerNode.h
--- a/Engine/Plugins/Runtime/Harmonix/Source/HarmonixMetasound/Public/HarmonixMetasound/Nodes/MidiCCTriggerNode.h
+++ b/Engine/Plugins/Runtime/Harmonix/Source/HarmonixMetasound/Public/HarmonixMetasound/Nodes/MidiCCTriggerNode.h
@@ -12,8 +12,8 @@ namespace HarmonixMetasound::Nodes::MidiCCTriggerNode
 	
 	namespace Inputs
 	{
-		DECLARE_METASOUND_PARAM_EXTERN(Enable);
-		DECLARE_METASOUND_PARAM_EXTERN(MidiStream);
+		DECLARE_METASOUND_PARAM_ALIAS(Enable);
+		DECLARE_METASOUND_PARAM_ALIAS(MidiStream);
 		DECLARE_METASOUND_PARAM_EXTERN(InputMidiControllerID);
 	}
 
diff --git a/Engine/Plugins/Runtime/Harmonix/Source/HarmonixMetasound/Public/HarmonixMetasound/Nodes/MidiChannelFilterNode.h b/Engine/Plugins/Runtime/Harmonix/Source/HarmonixMetasound/Public/HarmonixMetasound/Nodes/MidiChannelFilterNode.h
--- a/Engine/Plugins/Runtime/Harmonix/Source/HarmonixMetasound/Public/HarmonixMetasound/Nodes/MidiChannelFilterNode.h
+++ b/Engine/Plugins/Runtime/Harmonix/Source/HarmonixMetasound/Public/HarmonixMetasound/Nodes/MidiChannelFilterNode.h
@@ -13,13 +13,13 @@ namespace HarmonixMetasound::Nodes::MidiChannelFilter
 
 	namespace Inputs
 	{
-		DECLARE_METASOUND_PARAM_EXTERN(Enable);
-		DECLARE_METASOUND_PARAM_EXTERN(MidiStream);
-		DECLARE_METASOUND_PARAM_EXTERN(Channel);
+		DECLARE_METASOUND_PARAM_ALIAS(Enable);
+		DECLARE_METASOUND_PARAM_ALIAS(MidiStream);
+		DECLARE_METASOUND_PARAM_ALIAS(Channel);
 	}
 
 	namespace Outputs
 	{
-		DECLARE_METASOUND_PARAM_EXTERN(MidiStream);
+		DECLARE_METASOUND_PARAM_ALIAS(MidiStream);
 	}
 }
diff --git a/Engine/Plugins/Runtime/Harmonix/Source/HarmonixMetasound/Public/HarmonixMetasound/Nodes/MidiNoteFilterNode.h b/Engine/Plugins/Runtime/Harmonix/Source/HarmonixMetasound/Public/HarmonixMetasound/Nodes/MidiNoteFilterNode.h
--- a/Engine/Plugins/Runtime/Harmonix/Source/HarmonixMetasound/Public/HarmonixMetasound/Nodes/MidiNoteFilterNode.h
+++ b/Engine/Plugins/Runtime/Harmonix/Source/HarmonixMetasound/Public/HarmonixMetasound/Nodes/MidiNoteFilterNode.h
@@ -13,17 +13,17 @@ namespace HarmonixMetasound::Nodes::MidiNoteFilter
 
 	namespace Inputs
 	{
-		DECLARE_METASOUND_PARAM_EXTERN(Enable);
-		DECLARE_METASOUND_PARAM_EXTERN(MidiStream);
-		DECLARE_METASOUND_PARAM_EXTERN(MinNoteNumber);
-		DECLARE_METASOUND_PARAM_EXTERN(MaxNoteNumber);
-		DECLARE_METASOUND_PARAM_EXTERN(MinVelocity);
-		DECLARE_METASOUND_PARAM_EXTERN(MaxVelocity);
+		DECLARE_METASOUND_PARAM_ALIAS(Enable);
+		DECLARE_METASOUND_PARAM_ALIAS(MidiStream);
+		DECLARE_METASOUND_PARAM_ALIAS(MinNoteNumber);
+		DECLARE_METASOUND_PARAM_ALIAS(MaxNoteNumber);
+		DECLARE_METASOUND_PARAM_ALIAS(MinVelocity);
+		DECLARE_METASOUND_PARAM_ALIAS(MaxVelocity);
 		DECLARE_METASOUND_PARAM_EXTERN(IncludeOtherEvents);
 	}
 
 	namespace Outputs
 	{
-		DECLARE_METASOUND_PARAM_EXTERN(MidiStream);
+		DECLARE_METASOUND_PARAM_ALIAS(MidiStream);
 	}
 }
diff --git a/Engine/Plugins/Runtime/Harmonix/Source/HarmonixMetasound/Public/HarmonixMetasound/Nodes/MidiNoteTriggerNode.h b/Engine/Plugins/Runtime/Harmonix/Source/HarmonixMetasound/Public/HarmonixMetasound/Nodes/MidiNoteTriggerNode.h
--- a/Engine/Plugins/Runtime/Harmonix/Source/HarmonixMetasound/Public/HarmonixMetasound/Nodes/MidiNoteTriggerNode.h
+++ b/Engine/Plugins/Runtime/Harmonix/Source/HarmonixMetasound/Public/HarmonixMetasound/Nodes/MidiNoteTriggerNode.h
@@ -13,15 +13,15 @@ namespace HarmonixMetasound::Nodes::MidiNoteTriggerNode
 
 	namespace Inputs
 	{
-		DECLARE_METASOUND_PARAM_EXTERN(Enable);
-		DECLARE_METASOUND_PARAM_EXTERN(MidiStream);
+		DECLARE_METASOUND_PARAM_ALIAS(Enable);
+		DECLARE_METASOUND_PARAM_ALIAS(MidiStream);
 	}
 
 	namespace Outputs
 	{
-		DECLARE_METASOUND_PARAM_EXTERN(NoteOnTrigger);
-		DECLARE_METASOUND_PARAM_EXTERN(NoteOffTrigger);
-		DECLARE_METASOUND_PARAM_EXTERN(MidiNoteNumber);
-		DECLARE_METASOUND_PARAM_EXTERN(MidiVelocity);
+		DECLARE_METASOUND_PARAM_ALIAS(NoteOnTrigger);
+		DECLARE_METASOUND_PARAM_ALIAS(NoteOffTrigger);
+		DECLARE_METASOUND_PARAM_ALIAS(MidiNoteNumber);
+		DECLARE_METASOUND_PARAM_ALIAS(MidiVelocity);
 	}
 }
\ No newline at end of file
diff --git a/Engine/Plugins/Runtime/Harmonix/Source/HarmonixMetasound/Public/HarmonixMetasound/Nodes/MidiPulseGeneratorNode.h b/Engine/Plugins/Runtime/Harmonix/Source/HarmonixMetasound/Public/HarmonixMetasound/Nodes/MidiPulseGeneratorNode.h
--- a/Engine/Plugins/Runtime/Harmonix/Source/HarmonixMetasound/Public/HarmonixMetasound/Nodes/MidiPulseGeneratorNode.h
+++ b/Engine/Plugins/Runtime/Harmonix/Source/HarmonixMetasound/Public/HarmonixMetasound/Nodes/MidiPulseGeneratorNode.h
@@ -13,19 +13,19 @@ namespace HarmonixMetasound::Nodes::MidiPulseGeneratorNode
 
 	namespace Inputs
 	{
-		DECLARE_METASOUND_PARAM_EXTERN(MidiClock);
+		DECLARE_METASOUND_PARAM_ALIAS(MidiClock);
 		DECLARE_METASOUND_PARAM_EXTERN(Interval);
 		DECLARE_METASOUND_PARAM_EXTERN(IntervalMultiplier);
 		DECLARE_METASOUND_PARAM_EXTERN(Offset);
 		DECLARE_METASOUND_PARAM_EXTERN(OffsetMultiplier);
-		DECLARE_METASOUND_PARAM_EXTERN(MidiTrack);
-		DECLARE_METASOUND_PARAM_EXTERN(MidiChannel);
+		DECLARE_METASOUND_PARAM_ALIAS(MidiTrack);
+		DECLARE_METASOUND_PARAM_ALIAS(MidiChannel);
 		DECLARE_METASOUND_PARAM_EXTERN(MidiNoteNumber);
 		DECLARE_METASOUND_PARAM_EXTERN(MidiVelocity);
 	}
 
 	namespace Outputs
 	{
-		DECLARE_METASOUND_PARAM_EXTERN(MidiStream);
+		DECLARE_METASOUND_PARAM_ALIAS(MidiStream);
 	}
 }
\ No newline at end of file
diff --git a/Engine/Plugins/Runtime/Harmonix/Source/HarmonixMetasound/Public/HarmonixMetasound/Nodes/MidiStreamTrackFilterNode.h b/Engine/Plugins/Runtime/Harmonix/Source/HarmonixMetasound/Public/HarmonixMetasound/Nodes/MidiStreamTrackFilterNode.h
--- a/Engine/Plugins/Runtime/Harmonix/Source/HarmonixMetasound/Public/HarmonixMetasound/Nodes/MidiStreamTrackFilterNode.h
+++ b/Engine/Plugins/Runtime/Harmonix/Source/HarmonixMetasound/Public/HarmonixMetasound/Nodes/MidiStreamTrackFilterNode.h
@@ -13,8 +13,8 @@ namespace HarmonixMetasound::Nodes::MidiTrackFilter
 
 	namespace Inputs
 	{
-		DECLARE_METASOUND_PARAM_EXTERN(Enable);
-		DECLARE_METASOUND_PARAM_EXTERN(MidiStream);
+		DECLARE_METASOUND_PARAM_ALIAS(Enable);
+		DECLARE_METASOUND_PARAM_ALIAS(MidiStream);
 		DECLARE_METASOUND_PARAM_EXTERN(MinTrackIndex);
 		DECLARE_METASOUND_PARAM_EXTERN(MaxTrackIndex);
 		DECLARE_METASOUND_PARAM_EXTERN(IncludeConductorTrack);
@@ -22,6 +22,6 @@ namespace HarmonixMetasound::Nodes::MidiTrackFilter
 
 	namespace Outputs
 	{
-		DECLARE_METASOUND_PARAM_EXTERN(MidiStream);
+		DECLARE_METASOUND_PARAM_ALIAS(MidiStream);
 	}
 }
diff --git a/Engine/Plugins/Runtime/Harmonix/Source/HarmonixMetasound/Public/HarmonixMetasound/Nodes/MidiStreamTransposerNode.h b/Engine/Plugins/Runtime/Harmonix/Source/HarmonixMetasound/Public/HarmonixMetasound/Nodes/MidiStreamTransposerNode.h
--- a/Engine/Plugins/Runtime/Harmonix/Source/HarmonixMetasound/Public/HarmonixMetasound/Nodes/MidiStreamTransposerNode.h
+++ b/Engine/Plugins/Runtime/Harmonix/Source/HarmonixMetasound/Public/HarmonixMetasound/Nodes/MidiStreamTransposerNode.h
@@ -13,13 +13,13 @@ namespace HarmonixMetasound::Nodes::MidiNoteTranspose
 
 	namespace Inputs
 	{
-		DECLARE_METASOUND_PARAM_EXTERN(Enable);
-		DECLARE_METASOUND_PARAM_EXTERN(MidiStream);
-		DECLARE_METASOUND_PARAM_EXTERN(Transposition);
+		DECLARE_METASOUND_PARAM_ALIAS(Enable);
+		DECLARE_METASOUND_PARAM_ALIAS(MidiStream);
+		DECLARE_METASOUND_PARAM_ALIAS(Transposition);
 	}
 
 	namespace Outputs
 	{
-		DECLARE_METASOUND_PARAM_EXTERN(MidiStream);
+		DECLARE_METASOUND_PARAM_ALIAS(MidiStream);
 	}
 }
diff --git a/Engine/Plugins/Runtime/Harmonix/Source/HarmonixMetasound/Public/HarmonixMetasound/Nodes/MorphingLfoNode.h b/Engine/Plugins/Runtime/Harmonix/Source/HarmonixMetasound/Public/HarmonixMetasound/Nodes/MorphingLfoNode.h
--- a/Engine/Plugins/Runtime/Harmonix/Source/HarmonixMetasound/Public/HarmonixMetasound/Nodes/MorphingLfoNode.h
+++ b/Engine/Plugins/Runtime/Harmonix/Source/HarmonixMetasound/Public/HarmonixMetasound/Nodes/MorphingLfoNode.h
@@ -22,11 +22,11 @@ namespace HarmonixMetasound::Nodes::MorphingLFO
 	
 	namespace Inputs
 	{
-		DECLARE_METASOUND_PARAM_EXTERN(MidiClock);
-		DECLARE_METASOUND_PARAM_EXTERN(LFOSyncType);
-		DECLARE_METASOUND_PARAM_EXTERN(LFOFrequency);
-		DECLARE_METASOUND_PARAM_EXTERN(LFOInvert);
-		DECLARE_METASOUND_PARAM_EXTERN(LFOShape)
+		DECLARE_METASOUND_PARAM_ALIAS(MidiClock);
+		DECLARE_METASOUND_PARAM_ALIAS(LFOSyncType);
+		DECLARE_METASOUND_PARAM_ALIAS(LFOFrequency);
+		DECLARE_METASOUND_PARAM_ALIAS(LFOInvert);
+		DECLARE_METASOUND_PARAM_ALIAS(LFOShape)
 	}
 
 	namespace Outputs
diff --git a/Engine/Plugins/Runtime/Harmonix/Source/HarmonixMetasound/Public/HarmonixMetasound/Nodes/MultibandAnalyzerNode.h b/Engine/Plugins/Runtime/Harmonix/Source/HarmonixMetasound/Public/HarmonixMetasound/Nodes/MultibandAnalyzerNode.h
--- a/Engine/Plugins/Runtime/Harmonix/Source/HarmonixMetasound/Public/HarmonixMetasound/Nodes/MultibandAnalyzerNode.h
+++ b/Engine/Plugins/Runtime/Harmonix/Source/HarmonixMetasound/Public/HarmonixMetasound/Nodes/MultibandAnalyzerNode.h
@@ -23,8 +23,8 @@ namespace HarmonixMetasound::Nodes::MultibandAnalyzer
 	
 	namespace Inputs
 	{
-		DECLARE_METASOUND_PARAM_EXTERN(Enable);
-		DECLARE_METASOUND_PARAM_EXTERN(AudioMono);
+		DECLARE_METASOUND_PARAM_ALIAS(Enable);
+		DECLARE_METASOUND_PARAM_ALIAS(AudioMono);
 		DECLARE_METASOUND_PARAM_EXTERN(CrossoverFrequencies);
 		DECLARE_METASOUND_PARAM_EXTERN(ApplySmoothing);
 		DECLARE_METASOUND_PARAM_EXTERN(AttackTime);
diff --git a/Engine/Plugins/Runtime/Harmonix/Source/HarmonixMetasound/Public/HarmonixMetasound/Nodes/PeakNode.h b/Engine/Plugins/Runtime/Harmonix/Source/HarmonixMetasound/Public/HarmonixMetasound/Nodes/PeakNode.h
--- a/Engine/Plugins/Runtime/Harmonix/Source/HarmonixMetasound/Public/HarmonixMetasound/Nodes/PeakNode.h
+++ b/Engine/Plugins/Runtime/Harmonix/Source/HarmonixMetasound/Public/HarmonixMetasound/Nodes/PeakNode.h
@@ -12,7 +12,8 @@ namespace HarmonixMetasound::Nodes::Peak
 	
 	namespace Inputs
 	{
-		DECLARE_METASOUND_PARAM_EXTERN(AudioMono);
+		DECLARE_METASOUND_PARAM_EXTERN(Enable);
+		DECLARE_METASOUND_PARAM_ALIAS(AudioMono);
 	}
 
 	namespace Outputs
