Make sure Bar & Beat pin tooltips clarify that they're 1 based and not 0 based.

Make sure nodes that take a Bar & Beat as input default to 1 & 1.0f respectively

#tests PIE

diff --git a/Engine/Plugins/Runtime/Harmonix/Source/HarmonixMetasound/Private/HarmonixMetasound/Common.cpp b/Engine/Plugins/Runtime/Harmonix/Source/HarmonixMetasound/Private/HarmonixMetasound/Common.cpp
--- a/Engine/Plugins/Runtime/Harmonix/Source/HarmonixMetasound/Private/HarmonixMetasound/Common.cpp
+++ b/Engine/Plugins/Runtime/Harmonix/Source/HarmonixMetasound/Private/HarmonixMetasound/Common.cpp
@@ -53,8 +53,8 @@ namespace HarmonixMetasound
 			DEFINE_INPUT_METASOUND_PARAM(GridSizeMult,                  "Grid Size Multiplier",               "Grid size multiplier.")
 			DEFINE_INPUT_METASOUND_PARAM(OffsetUnits,                   "Base Offset Size ",                  "Base size of the qrid offset.")
 			DEFINE_INPUT_METASOUND_PARAM(OffsetMult,                    "Offset Multiplier",                  "Offset multiplier.")
-			DEFINE_INPUT_METASOUND_PARAM(Bar,		   					"Bar",                                "Musical bar (or measure) index (0 based!)")
-			DEFINE_INPUT_METASOUND_PARAM(FloatBeat,			            "Beat",                               "Musical beat index. 0 based and in units of the time signature denomninator! The value is rounded to the nearest 32nd note.")
+			DEFINE_INPUT_METASOUND_PARAM(Bar,		   					"Bar",                                "Musical bar (or measure) (1 based!)")
+			DEFINE_INPUT_METASOUND_PARAM(FloatBeat,			            "Beat",                               "Musical beat. 1 based and in units of the time signature denomninator! The value is rounded to the nearest 32nd note.")
 			DEFINE_INPUT_METASOUND_PARAM(SynthPatch,                    "Patch",                              "Synthesizer patch asset.")
 			DEFINE_INPUT_METASOUND_PARAM(Transposition,                 "Transposition",                      "Transposition in semitones.")
 			DEFINE_INPUT_METASOUND_PARAM(ClockSpeedToPitch,             "MIDI Clock Speed Affects Pitch",     "When true the clock speed of the incoming MIDI stream will affect the pitch of the synth.")
