I made a mistake when making the MidiClockEvents...

Change StartTick and EndTick to Tick1 and Tick2... because...

Tick1 is FromTick or LoopStartTick, depending on the event
Tick2 is ToTick, ThruTick, or LoopEndTick, depending on the event

And I fixed my code so they actually store the right values...

Things using these events should make sure they're using the right ones.

Maybe I should add methods to get at each tick, something similar to how MidiMsg works.

#tests PIE

diff --git a/Engine/Plugins/Runtime/Harmonix/Source/HarmonixMetasound/Private/HarmonixMetasound/DataTypes/MidiClock.cpp b/Engine/Plugins/Runtime/Harmonix/Source/HarmonixMetasound/Private/HarmonixMetasound/DataTypes/MidiClock.cpp
--- a/Engine/Plugins/Runtime/Harmonix/Source/HarmonixMetasound/Private/HarmonixMetasound/DataTypes/MidiClock.cpp
+++ b/Engine/Plugins/Runtime/Harmonix/Source/HarmonixMetasound/Private/HarmonixMetasound/DataTypes/MidiClock.cpp
@@ -512,8 +512,9 @@ namespace HarmonixMetasound
 
 	void FMidiClockEventCursor::Reset(bool ForceNoBroadcast /*= false*/)
 	{
+		const int32 FromTick = CurrentTick;
 		FMidiPlayCursor::Reset(ForceNoBroadcast);
-		AddEvent(FMidiClockEvent::MakeResetEvent(MidiClock->GetCurrentBlockFrameIndex(), MidiClock->GetCurrentHiResTick(), ForceNoBroadcast));
+		AddEvent(FMidiClockEvent::MakeResetEvent(MidiClock->GetCurrentBlockFrameIndex(), FromTick, CurrentTick, ForceNoBroadcast));
 	}
 
 	void FMidiClockEventCursor::OnLoop(int32 LoopStartTick, int32 LoopEndTick)
@@ -522,22 +523,25 @@ namespace HarmonixMetasound
 		AddEvent(FMidiClockEvent::MakeLoopEvent(MidiClock->GetCurrentBlockFrameIndex(), LoopStartTick, LoopEndTick));
 	}
 
-	void FMidiClockEventCursor::SeekToTick(int32 Tick)
+	void FMidiClockEventCursor::SeekToTick(int32 Tick) 
 	{
+		const int32 FromTick = CurrentTick;
 		FMidiPlayCursor::SeekToTick(Tick);
-		AddEvent(FMidiClockEvent::MakeSeekToEvent(MidiClock->GetCurrentBlockFrameIndex(), Tick));
+		AddEvent(FMidiClockEvent::MakeSeekToEvent(MidiClock->GetCurrentBlockFrameIndex(), FromTick, Tick));
 	}
 
 	void FMidiClockEventCursor::SeekThruTick(int32 Tick)
 	{
+		const int32 FromTick = CurrentTick;
 		FMidiPlayCursor::SeekThruTick(Tick);
-		AddEvent(FMidiClockEvent::MakeSeekThruEvent(MidiClock->GetCurrentBlockFrameIndex(), Tick));
+		AddEvent(FMidiClockEvent::MakeSeekThruEvent(MidiClock->GetCurrentBlockFrameIndex(), FromTick, Tick));
 	}
 
 	void FMidiClockEventCursor::AdvanceThruTick(int32 Tick, bool IsPreRoll)
 	{
+		const int32 FromTick = CurrentTick;
 		FMidiPlayCursor::AdvanceThruTick(Tick, IsPreRoll);
-		AddEvent(FMidiClockEvent::MakeAdvanceThruEvent(MidiClock->GetCurrentBlockFrameIndex(), Tick, IsPreRoll));
+		AddEvent(FMidiClockEvent::MakeAdvanceThruEvent(MidiClock->GetCurrentBlockFrameIndex(), FromTick, Tick, IsPreRoll));
 	}
 
 	void FMidiClockEventCursor::AddEvent(const FMidiClockEvent& InEvent)
diff --git a/Engine/Plugins/Runtime/Harmonix/Source/HarmonixMetasound/Private/HarmonixMetasound/DataTypes/MidiClockEvent.cpp b/Engine/Plugins/Runtime/Harmonix/Source/HarmonixMetasound/Private/HarmonixMetasound/DataTypes/MidiClockEvent.cpp
--- a/Engine/Plugins/Runtime/Harmonix/Source/HarmonixMetasound/Private/HarmonixMetasound/DataTypes/MidiClockEvent.cpp
+++ b/Engine/Plugins/Runtime/Harmonix/Source/HarmonixMetasound/Private/HarmonixMetasound/DataTypes/MidiClockEvent.cpp
@@ -6,9 +6,9 @@
 
 namespace HarmonixMetasound
 {
-	FMidiClockEvent FMidiClockEvent::MakeResetEvent(int32 InBlockFrameIndex, int32 Tick, bool ForceNoBroadcast)
+	FMidiClockEvent FMidiClockEvent::MakeResetEvent(int32 InBlockFrameIndex, int32 FromTick, int32 ToTick, bool ForceNoBroadcast)
 	{
-		return FMidiClockEvent(EType::Reset, InBlockFrameIndex, Tick, Tick, false, ForceNoBroadcast);
+		return FMidiClockEvent(EType::Reset, InBlockFrameIndex, FromTick, ToTick, false, ForceNoBroadcast);
 	}
 
 	FMidiClockEvent FMidiClockEvent::MakeLoopEvent(int32 InBlockFrameIndex, int32 InStartTick, int32 InEndTick)
@@ -16,28 +16,28 @@ namespace HarmonixMetasound
 		return FMidiClockEvent(EType::Loop, InBlockFrameIndex, InStartTick, InEndTick, false, true);
 	}
 
-	FMidiClockEvent FMidiClockEvent::MakeSeekToEvent(int32 InBlockFrameIndex, int32 Tick)
+	FMidiClockEvent FMidiClockEvent::MakeSeekToEvent(int32 InBlockFrameIndex, int32 FromTick, int32 ToTick)
 	{
-		return FMidiClockEvent(EType::SeekTo, InBlockFrameIndex, Tick, Tick, false, true);
+		return FMidiClockEvent(EType::SeekTo, InBlockFrameIndex, FromTick, ToTick, false, true);
 	}
 
-	FMidiClockEvent FMidiClockEvent::MakeSeekThruEvent(int32 InBlockFrameIndex, int32 Tick)
+	FMidiClockEvent FMidiClockEvent::MakeSeekThruEvent(int32 InBlockFrameIndex, int32 FromTick, int32 ThruTick)
 	{
-		return FMidiClockEvent(EType::SeekThru, InBlockFrameIndex, Tick, Tick, false, true);
+		return FMidiClockEvent(EType::SeekThru, InBlockFrameIndex, FromTick, ThruTick, false, true);
 	}
 
-	FMidiClockEvent FMidiClockEvent::MakeAdvanceThruEvent(int32 InBlockFrameIndex, int32 Tick, bool IsPreRoll)
+	FMidiClockEvent FMidiClockEvent::MakeAdvanceThruEvent(int32 InBlockFrameIndex, int32 FromTick, int32 ThruTick, bool IsPreRoll)
 	{
-		return FMidiClockEvent(EType::AdvanceThru, InBlockFrameIndex, Tick, Tick, IsPreRoll, true);
+		return FMidiClockEvent(EType::AdvanceThru, InBlockFrameIndex, FromTick, ThruTick, IsPreRoll, true);
 	}
 
-	FMidiClockEvent::FMidiClockEvent(EType InType, int32 InBlockFrameIndex, int32 InStartTick, int32 InEndTick, bool InIsPreRoll, bool InForceNoBroadcast)
+	FMidiClockEvent::FMidiClockEvent(EType InType, int32 InBlockFrameIndex, int32 InTick1, int32 InTick2, bool InIsPreRoll, bool InForceNoBroadcast)
 		: Type(InType)
 		, IsPreRoll(InIsPreRoll)
 		, ForceNoBroadcast(InForceNoBroadcast)
 		, BlockFrameIndex(InBlockFrameIndex)
-		, StartTick(InStartTick)
-		, EndTick(InEndTick)
+		, Tick1(InTick1)
+		, Tick2(InTick2)
 	{}
 
 }
\ No newline at end of file
diff --git a/Engine/Plugins/Runtime/Harmonix/Source/HarmonixMetasound/Private/HarmonixMetasound/Nodes/MidiClockOffsetNode.cpp b/Engine/Plugins/Runtime/Harmonix/Source/HarmonixMetasound/Private/HarmonixMetasound/Nodes/MidiClockOffsetNode.cpp
--- a/Engine/Plugins/Runtime/Harmonix/Source/HarmonixMetasound/Private/HarmonixMetasound/Nodes/MidiClockOffsetNode.cpp
+++ b/Engine/Plugins/Runtime/Harmonix/Source/HarmonixMetasound/Private/HarmonixMetasound/Nodes/MidiClockOffsetNode.cpp
@@ -177,7 +177,7 @@ namespace HarmonixMetasound::Nodes::MidiClockOffset
 				case FMidiClockEvent::EType::SeekTo:
 				case FMidiClockEvent::EType::Reset:
 					{
-						const int32 Tick = GetTickWithOffset(Event.StartTick, OffsetBars, OffsetBeats, OffsetMs);
+						const int32 Tick = GetTickWithOffset(Event.Tick2, OffsetBars, OffsetBeats, OffsetMs);
 						FMusicSeekTarget SeekTarget;
 						SeekTarget.Type = ESeekPointType::Millisecond;
 						SeekTarget.Ms = MidiClockOut->GetSongMaps().TickToMs(Tick);
@@ -187,7 +187,7 @@ namespace HarmonixMetasound::Nodes::MidiClockOffset
 					}
 				case FMidiClockEvent::EType::SeekThru:
 					{
-						const int32 Tick = GetTickWithOffset(Event.StartTick, OffsetBars, OffsetBeats, OffsetMs);
+						const int32 Tick = GetTickWithOffset(Event.Tick2, OffsetBars, OffsetBeats, OffsetMs);
 						FMusicSeekTarget SeekTarget;
 						SeekTarget.Type = ESeekPointType::Millisecond;
 						SeekTarget.Ms = MidiClockOut->GetSongMaps().TickToMs(Tick + 1);
@@ -197,7 +197,7 @@ namespace HarmonixMetasound::Nodes::MidiClockOffset
 					}
 				case FMidiClockEvent::EType::AdvanceThru:
 					{
-						const int32 Tick = GetTickWithOffset(Event.StartTick, OffsetBars, OffsetBeats, OffsetMs);
+						const int32 Tick = GetTickWithOffset(Event.Tick2, OffsetBars, OffsetBeats, OffsetMs);
 						
 						// if our offset changed while we were advancing, seek to the new offset first
 						if (!FMath::IsNearlyEqual(PrevOffsetMs, OffsetMs) || !FMath::IsNearlyEqual(PrevOffsetBeats, OffsetBeats) || PrevOffsetBars != OffsetBars)
@@ -215,7 +215,7 @@ namespace HarmonixMetasound::Nodes::MidiClockOffset
 
 						const float AdvanceToMs = MidiClockOut->GetSongMaps().TickToMs(Tick);
 						const float ClockInSpeed = MidiClockIn->GetSpeedAtBlockSampleFrame(EventIndex);
-						const float AdvanceRatio = MidiClockIn->GetSongMaps().GetTempoAtTick(Event.StartTick) / MidiClockOut->GetSongMaps().GetTempoAtTick(Tick);
+						const float AdvanceRatio = MidiClockIn->GetSongMaps().GetTempoAtTick(Event.Tick2) / MidiClockOut->GetSongMaps().GetTempoAtTick(Tick);
 						MidiClockOut->InformOfCurrentAdvanceRate(ClockInSpeed * AdvanceRatio);
 						MidiClockOut->AdvanceHiResToMs(Event.BlockFrameIndex, AdvanceToMs, true);
 						break;
diff --git a/Engine/Plugins/Runtime/Harmonix/Source/HarmonixMetasound/Private/HarmonixMetasound/Nodes/MidiPlayerNode.cpp b/Engine/Plugins/Runtime/Harmonix/Source/HarmonixMetasound/Private/HarmonixMetasound/Nodes/MidiPlayerNode.cpp
--- a/Engine/Plugins/Runtime/Harmonix/Source/HarmonixMetasound/Private/HarmonixMetasound/Nodes/MidiPlayerNode.cpp
+++ b/Engine/Plugins/Runtime/Harmonix/Source/HarmonixMetasound/Private/HarmonixMetasound/Nodes/MidiPlayerNode.cpp
@@ -416,8 +416,8 @@ namespace HarmonixMetasound
 					{
 					case FMidiClockEvent::EType::Reset:
 					{
-						UpdateLoopOffsetTickFromTick(Event.StartTick);
-						int32 Tick = Event.StartTick - LoopOffsetTick;
+						UpdateLoopOffsetTickFromTick(Event.Tick2);
+						int32 Tick = Event.Tick2 - LoopOffsetTick;
 
 						float Ms = MidiClockOut->GetSongMaps().TickToMs(Tick);
 						FMusicSeekTarget SeekTarget;
@@ -433,8 +433,8 @@ namespace HarmonixMetasound
 					}
 					case FMidiClockEvent::EType::SeekTo:
 					{
-						UpdateLoopOffsetTickFromTick(Event.StartTick);
-						int32 Tick = Event.StartTick - LoopOffsetTick;
+						UpdateLoopOffsetTickFromTick(Event.Tick2);
+						int32 Tick = Event.Tick2 - LoopOffsetTick;
 
 						float Ms = MidiClockOut->GetSongMaps().TickToMs(Tick);
 						FMusicSeekTarget SeekTarget;
@@ -445,8 +445,8 @@ namespace HarmonixMetasound
 					}
 					case FMidiClockEvent::EType::SeekThru:
 					{
-						UpdateLoopOffsetTickFromTick(Event.StartTick);
-						int32 Tick = Event.StartTick - LoopOffsetTick;
+						UpdateLoopOffsetTickFromTick(Event.Tick2);
+						int32 Tick = Event.Tick2 - LoopOffsetTick;
 
 						float Ms = MidiClockOut->GetSongMaps().TickToMs(Tick + 1);
 						FMusicSeekTarget SeekTarget;
@@ -459,11 +459,11 @@ namespace HarmonixMetasound
 					{
 						
 						// The MidiClock handles looping on its own, so we can conveniently advance it
-						int32 Tick = Event.StartTick - LoopOffsetTick;
+						int32 Tick = Event.Tick2 - LoopOffsetTick;
 						float Ms = MidiClockOut->GetSongMaps().TickToMs(Tick);
 
 						float ClockInSpeed = MidiClockIn->GetSpeedAtBlockSampleFrame(StartFrameIndex);
-						float AdvanceRatio = MidiClockIn->GetSongMaps().GetTempoAtTick(Event.StartTick) / MidiClockOut->GetSongMaps().GetTempoAtTick(Tick);
+						float AdvanceRatio = MidiClockIn->GetSongMaps().GetTempoAtTick(Event.Tick2) / MidiClockOut->GetSongMaps().GetTempoAtTick(Tick);
 						// midi clock needs to know how fast its advancing based on their authority
 						MidiClockOut->InformOfCurrentAdvanceRate(ClockInSpeed * *SpeedMultInPin * AdvanceRatio);
 						MidiClockOut->AdvanceHiResToMs(Event.BlockFrameIndex, Ms, true);
@@ -477,7 +477,7 @@ namespace HarmonixMetasound
 							int32 NewTick = MidiClockOut->GetCurrentHiResTick();
 							if (NewTick < Tick)
 							{
-								UpdateLoopOffsetTickFromTick(Event.StartTick);
+								UpdateLoopOffsetTickFromTick(Event.Tick2);
 							}
 						}
 						break;
diff --git a/Engine/Plugins/Runtime/Harmonix/Source/HarmonixMetasound/Public/HarmonixMetasound/DataTypes/MidiClockEvent.h b/Engine/Plugins/Runtime/Harmonix/Source/HarmonixMetasound/Public/HarmonixMetasound/DataTypes/MidiClockEvent.h
--- a/Engine/Plugins/Runtime/Harmonix/Source/HarmonixMetasound/Public/HarmonixMetasound/DataTypes/MidiClockEvent.h
+++ b/Engine/Plugins/Runtime/Harmonix/Source/HarmonixMetasound/Public/HarmonixMetasound/DataTypes/MidiClockEvent.h
@@ -21,19 +21,25 @@ namespace HarmonixMetasound
 		const bool  IsPreRoll;
 		const bool  ForceNoBroadcast;
 		const int32 BlockFrameIndex;
-		const int32 StartTick;
-		const int32 EndTick;
 
+		// FromTick when event is SeekTo, SeekThru, AdvanceThru or Reset;
+		// LoopStartTick when OnLoop
+		const int32 Tick1;
 
-		static FMidiClockEvent MakeResetEvent(int32 InBlockFrameIndex, int32 Tick, bool ForceNoBroadcast = false);
-		static FMidiClockEvent MakeLoopEvent(int32 InBlockFrameIndex, int32 InStartTick, int32 InEndTick);
-		static FMidiClockEvent MakeSeekToEvent(int32 InBlockFrameIndex, int32 Tick);
-		static FMidiClockEvent MakeSeekThruEvent(int32 InBlockFrameIndex, int32 Tick);
-		static FMidiClockEvent MakeAdvanceThruEvent(int32 InBlockFrameIndex, int32 Tick, bool IsPreRoll);
+		// ToTick when SeekTo, Reset;
+		// ThruTick when SeekThru, AdvanceThru;
+		// LoopEndTick when OnLoop
+		const int32 Tick2;
+
+		static FMidiClockEvent MakeResetEvent(int32 InBlockFrameIndex, int32 FromTick, int32 ToTick, bool ForceNoBroadcast = false);
+		static FMidiClockEvent MakeLoopEvent(int32 InBlockFrameIndex, int32 LoopStartTick, int32 LoopEndTick);
+		static FMidiClockEvent MakeSeekToEvent(int32 InBlockFrameIndex, int32 FromTick, int32 ToTick);
+		static FMidiClockEvent MakeSeekThruEvent(int32 InBlockFrameIndex, int32 FromTick, int32 ThruTick);
+		static FMidiClockEvent MakeAdvanceThruEvent(int32 InBlockFrameIndex, int32 FromTick, int32 ThruTick, bool IsPreRoll);
 
 	private:
 
-		FMidiClockEvent(EType InType, int32 InBlockFrameIndex, int32 InStartTick, int32 InEndTick, bool InIsPreRoll, bool InBroadcastEvents);
+		FMidiClockEvent(EType InType, int32 InBlockFrameIndex, int32 InTick1, int32 InTick2, bool InIsPreRoll, bool InBroadcastEvents);
 	};
 
 };
\ No newline at end of file
