Reassign "GeneratorID" to MidiEvents when merging

#tests PIE

diff --git a/Engine/Plugins/Runtime/Harmonix/Source/HarmonixMetasound/Private/HarmonixMetasound/Nodes/MidiStreamMergeNode.cpp b/Engine/Plugins/Runtime/Harmonix/Source/HarmonixMetasound/Private/HarmonixMetasound/Nodes/MidiStreamMergeNode.cpp
--- a/Engine/Plugins/Runtime/Harmonix/Source/HarmonixMetasound/Private/HarmonixMetasound/Nodes/MidiStreamMergeNode.cpp
+++ b/Engine/Plugins/Runtime/Harmonix/Source/HarmonixMetasound/Private/HarmonixMetasound/Nodes/MidiStreamMergeNode.cpp
@@ -44,6 +44,9 @@ namespace HarmonixMetasound
 		//** OUTPUTS
 		FMidiStreamWriteRef MidiStreamOutPin;
 
+		//** DATA
+		TMap<uint32, FMidiVoiceGeneratorBase> GeneratorMap;
+		
 		bool bNeedsClockRefresh { true };
 	};
 
@@ -173,7 +176,17 @@ namespace HarmonixMetasound
 		}
 		
 		MidiStreamOutPin->PrepareBlock();
-		MidiStreamOutPin->Copy({ MidiStreamAInPin, MidiStreamBInPin });
+		TArray<FMidiStreamReadRef> SameClockStreams;
+		MidiStreamOutPin->FilterArrayToStreamsWithTheSameClock({MidiStreamAInPin, MidiStreamBInPin}, SameClockStreams);
+		MidiStreamOutPin->CopyTransportEvents(SameClockStreams);
+		MidiStreamOutPin->CopyMidiEvents(MidiStreamAInPin);
+		for (FMidiStreamEvent Event : MidiStreamBInPin->GetEventsInBlock())
+		{
+			uint32 GenId = Event.GetVoiceId().GetGeneratorId();
+			Event.ReassignOwner(&GeneratorMap.FindOrAdd(GenId));
+			MidiStreamOutPin->InsertMidiEvent(Event);
+		}
+		
 	}
 }
 
diff --git a/Engine/Plugins/Runtime/Harmonix/Source/HarmonixMetasound/Public/HarmonixMetasound/DataTypes/MidiStream.h b/Engine/Plugins/Runtime/Harmonix/Source/HarmonixMetasound/Public/HarmonixMetasound/DataTypes/MidiStream.h
--- a/Engine/Plugins/Runtime/Harmonix/Source/HarmonixMetasound/Public/HarmonixMetasound/DataTypes/MidiStream.h
+++ b/Engine/Plugins/Runtime/Harmonix/Source/HarmonixMetasound/Public/HarmonixMetasound/DataTypes/MidiStream.h
@@ -41,6 +41,18 @@ namespace HarmonixMetasound
 			VoiceId = FMidiVoiceId(OwnerId, MidiMessage);
 		}
 
+		// reassign owner with a new generator
+		void ReassignOwner(const FMidiVoiceGeneratorBase* Owner)
+		{
+			VoiceId.ReassignGenerator(Owner);
+		}
+
+		// reassign owner with a new generator id
+		void ReassignOwner(uint32 OwnerId)
+		{
+			VoiceId.ReassignGenerator(OwnerId);
+		}
+
 		FMidiVoiceId GetVoiceId() const { return VoiceId; }
 		void SetVoiceId(FMidiVoiceId InVoiceId) { VoiceId = InVoiceId; }
 
diff --git a/Engine/Plugins/Runtime/Harmonix/Source/HarmonixMidi/Public/HarmonixMidi/MidiVoiceId.h b/Engine/Plugins/Runtime/Harmonix/Source/HarmonixMidi/Public/HarmonixMidi/MidiVoiceId.h
--- a/Engine/Plugins/Runtime/Harmonix/Source/HarmonixMidi/Public/HarmonixMidi/MidiVoiceId.h
+++ b/Engine/Plugins/Runtime/Harmonix/Source/HarmonixMidi/Public/HarmonixMidi/MidiVoiceId.h
@@ -50,6 +50,19 @@ public:
 		Other.Id = 0;
 	}
 
+	void ReassignGenerator(const FMidiVoiceGeneratorBase* Generator)
+	{
+		ReassignGenerator(Generator ? Generator->GetIdBits() : 0);
+	}
+
+	void ReassignGenerator(uint32 GeneratorId)
+	{
+		uint8 Ch;
+		uint8 Note;
+		GetChannelAndNote(Ch, Note);
+		Id = (GeneratorId | ((uint32)(Ch & 0xF) << 8) | (uint32)Note);
+	}
+
 	void GetGeneratorChannelAndNote(uint32& GeneratorOut, uint8& ChOut, uint8& NoteOut) const
 	{
 		GeneratorOut = Id & 0xFFFFFC00;
@@ -62,6 +75,11 @@ public:
 		NoteOut = Id & 0xFF;
 	}
 
+	uint32 GetGeneratorId() const
+	{
+		return Id & 0xFFFFFC00;
+	}
+
 	FMidiVoiceId& operator=(const FMidiVoiceId& Other)
 	{
 		Id = Other.Id;
