[Harmonix] Fix loctext collision

diff --git a/Engine/Plugins/Runtime/Harmonix/Source/HarmonixMetasound/Private/HarmonixMetasound/Nodes/PeakNode.cpp b/Engine/Plugins/Runtime/Harmonix/Source/HarmonixMetasound/Private/HarmonixMetasound/Nodes/PeakNode.cpp
--- a/Engine/Plugins/Runtime/Harmonix/Source/HarmonixMetasound/Private/HarmonixMetasound/Nodes/PeakNode.cpp
+++ b/Engine/Plugins/Runtime/Harmonix/Source/HarmonixMetasound/Private/HarmonixMetasound/Nodes/PeakNode.cpp
@@ -29,7 +29,7 @@ namespace HarmonixMetasound::Nodes::Peak
 
 	namespace Inputs
 	{
-		DEFINE_INPUT_METASOUND_PARAM(Enable, "Enable", "Enable the peak node.")
+		DEFINE_METASOUND_PARAM_ALIAS(Enable, CommonPinNames::Inputs::Enable)
 		DEFINE_METASOUND_PARAM_ALIAS(AudioMono, CommonPinNames::Inputs::AudioMono);
 	}
 
diff --git a/Engine/Plugins/Runtime/Harmonix/Source/HarmonixMetasound/Public/HarmonixMetasound/Nodes/PeakNode.h b/Engine/Plugins/Runtime/Harmonix/Source/HarmonixMetasound/Public/HarmonixMetasound/Nodes/PeakNode.h
--- a/Engine/Plugins/Runtime/Harmonix/Source/HarmonixMetasound/Public/HarmonixMetasound/Nodes/PeakNode.h
+++ b/Engine/Plugins/Runtime/Harmonix/Source/HarmonixMetasound/Public/HarmonixMetasound/Nodes/PeakNode.h
@@ -12,7 +12,7 @@ namespace HarmonixMetasound::Nodes::Peak
 	
 	namespace Inputs
 	{
-		DECLARE_METASOUND_PARAM_EXTERN(Enable);
+		DECLARE_METASOUND_PARAM_ALIAS(Enable);
 		DECLARE_METASOUND_PARAM_ALIAS(AudioMono);
 	}
 
