HarmonixMetasound custom pin colors

- move Icons out of project content and into plugin "Editor" folder (so that way an empty "Icons" folder doesn't show up in the editor content)
- try to improve some code and logic around finding styles
- use new RegisterCustomPinType API

#tests PIE

diff --git a/Engine/Plugins/Runtime/Harmonix/Source/HarmonixMetasoundEditor/HarmonixMetasoundEditor.Build.cs b/Engine/Plugins/Runtime/Harmonix/Source/HarmonixMetasoundEditor/HarmonixMetasoundEditor.Build.cs
--- a/Engine/Plugins/Runtime/Harmonix/Source/HarmonixMetasoundEditor/HarmonixMetasoundEditor.Build.cs
+++ b/Engine/Plugins/Runtime/Harmonix/Source/HarmonixMetasoundEditor/HarmonixMetasoundEditor.Build.cs
@@ -31,7 +31,8 @@ public class HarmonixMetasoundEditor : ModuleRules
 				"HarmonixMetasound",
 				"MetasoundEditor",
 				"MetasoundFrontend",
-				"DetailCustomizations"
+				"DetailCustomizations",
+				"Projects"
 			}
 		);
 	}
diff --git a/Engine/Plugins/Runtime/Harmonix/Source/HarmonixMetasoundEditor/Private/HarmonixMetasoundEditorModule.cpp b/Engine/Plugins/Runtime/Harmonix/Source/HarmonixMetasoundEditor/Private/HarmonixMetasoundEditorModule.cpp
--- a/Engine/Plugins/Runtime/Harmonix/Source/HarmonixMetasoundEditor/Private/HarmonixMetasoundEditorModule.cpp
+++ b/Engine/Plugins/Runtime/Harmonix/Source/HarmonixMetasoundEditor/Private/HarmonixMetasoundEditorModule.cpp
@@ -29,10 +29,20 @@ void FHarmonixMetasoundEditorModule::StartupModule()
 	MetasoundEditorModule.RegisterPinType("Enum:Distortion:FilterPasses", "Int32");
 	MetasoundEditorModule.RegisterPinType("Enum:StdMIDIControllerID", "Int32");
 
-	MetasoundEditorModule.RegisterPinType("MIDIStream", {}, {}, Style.GetMidiStreamConnectedIcon(), Style.GetMidiStreamDisconnectedIcon());
-	MetasoundEditorModule.RegisterPinType("MIDIClock", {}, {}, Style.GetMidiClockConnectedIcon(), Style.GetMidiClockDisconnectedIcon());
-	MetasoundEditorModule.RegisterPinType("MusicTransport", {}, {}, Style.GetTransportConnectedIcon(), Style.GetTransportDisconnectedIcon());
-
+	auto RegisterCustomPinType = [&Style, &MetasoundEditorModule](FName PinType)
+	{
+		Metasound::Editor::FGraphPinParams MidiStreamParams;
+		MidiStreamParams.PinCategory = PinType;
+		MidiStreamParams.PinColor = &Style.GetPinColor(PinType);
+		MidiStreamParams.PinConnectedIcon = Style.GetConnectedIcon(PinType);
+		MidiStreamParams.PinDisconnectedIcon = Style.GetDisconnectedIcon(PinType);
+		MetasoundEditorModule.RegisterCustomPinType(PinType, MidiStreamParams);
+	};
+
+	RegisterCustomPinType("MIDIStream");
+	RegisterCustomPinType("MIDIClock");
+	RegisterCustomPinType("MusicTransport");
+	
 	FPropertyEditorModule& PropertyModule = FModuleManager::LoadModuleChecked<FPropertyEditorModule>("PropertyEditor");
 	PropertyModule.RegisterCustomClassLayout(
 		UMidiStepSequence::StaticClass()->GetFName(),
diff --git a/Engine/Plugins/Runtime/Harmonix/Source/HarmonixMetasoundEditor/Private/HarmonixMetasoundSlateStyle.cpp b/Engine/Plugins/Runtime/Harmonix/Source/HarmonixMetasoundEditor/Private/HarmonixMetasoundSlateStyle.cpp
--- a/Engine/Plugins/Runtime/Harmonix/Source/HarmonixMetasoundEditor/Private/HarmonixMetasoundSlateStyle.cpp
+++ b/Engine/Plugins/Runtime/Harmonix/Source/HarmonixMetasoundEditor/Private/HarmonixMetasoundSlateStyle.cpp
@@ -2,65 +2,71 @@
 
 #include "HarmonixMetasoundSlateStyle.h"
 #include "Styling/SlateStyleRegistry.h"
+#include "Interfaces/IPluginManager.h"
 
 namespace HarmonixMetasoundEditor
 {
+	static const FString ModuleName = TEXT("HarmonixMetasound");
+	static const FString PinColorName = TEXT("PinColor");
+	static const FString ConnectedIconName = TEXT("ConnectedIcon");
+	static const FString DisconnectedIconName = TEXT("DisconnectedIcon");
+	
 	FSlateStyle::FSlateStyle()
 		: FSlateStyleSet("HarmonixMetasoundSlateStyle")
 	{
+		const FLinearColor MidiStreamColor = FColor(117, 106, 182); // pastel purple
+		const FLinearColor MidiClockColor = FColor(172, 135, 197); // pastel light purple
+		const FLinearColor TransportColor = FColor(255, 229, 229); // pastel cream
 		const FVector2D Icon22x22(22.0f, 22.0f);
 		const FVector2D Icon18x10(18.0f, 10.0f);
 		const FVector2D Icon18x18(18.0f, 18.0f);
-		SetContentRoot(FPaths::EnginePluginsDir() / TEXT("Runtime/Harmonix/Content/Editor"));
+		TSharedPtr<IPlugin> Plugin = IPluginManager::Get().FindPlugin(TEXT("Harmonix"));
+		check(Plugin);
+		SetContentRoot(Plugin->GetBaseDir() / TEXT("Editor/Slate"));
 
 #define IMAGE_BRUSH( RelativePath, ... ) FSlateImageBrush( RootToContentDir( RelativePath, TEXT(".png") ), __VA_ARGS__ )
-		Set("HarmonixMetasound.MidiConnectedIcon", new IMAGE_BRUSH(TEXT("Icons/MidiConnectedPin"), Icon22x22));
-		Set("HarmonixMetasound.MidiDisconnectedIcon", new IMAGE_BRUSH(TEXT("Icons/MidiDisconnectedPin"), Icon22x22));
-		Set("HarmonixMetasound.ClockConnectedIcon", new IMAGE_BRUSH(TEXT("Icons/ClockConnectedPin"), Icon18x18));
-		Set("HarmonixMetasound.ClockDisconnectedIcon", new IMAGE_BRUSH(TEXT("Icons/ClockDisconnectedPin"), Icon18x18));
-		Set("HarmonixMetasound.TransportConnectedIcon", new IMAGE_BRUSH(TEXT("Icons/TransportConnectedPin"), Icon18x10));
-		Set("HarmonixMetasound.TransportDisconnectedIcon", new IMAGE_BRUSH(TEXT("Icons/TransportDisconnectedPin"), Icon18x10));
+		SetCustomPinStyle("MIDIStream", MidiStreamColor, new IMAGE_BRUSH(TEXT("Icons/MidiConnectedPin"), Icon22x22), new IMAGE_BRUSH(TEXT("Icons/MidiDisconnectedPin"), Icon22x22));
+		SetCustomPinStyle("MIDIClock", MidiClockColor, new IMAGE_BRUSH(TEXT("Icons/ClockConnectedPin"), Icon22x22), new IMAGE_BRUSH(TEXT("Icons/ClockDisconnectedPin"), Icon22x22));
+		SetCustomPinStyle("MusicTransport", TransportColor, new IMAGE_BRUSH(TEXT("Icons/TransportConnectedPin"), Icon18x10), new IMAGE_BRUSH(TEXT("Icons/TransportDisconnectedPin"), Icon18x10));
 #undef IMAGE_BRUSH
 
 		FSlateStyleRegistry::RegisterSlateStyle(*this);
 	}
-	
+
 	const FSlateStyle& FSlateStyle::Get()
 	{
 		static FSlateStyle SlateStyle;
 		return SlateStyle;
 	}
 
-	const FSlateBrush* FSlateStyle::GetMidiStreamConnectedIcon() const
+	void FSlateStyle::SetCustomPinStyle(const FName PinType, const FLinearColor& PinColor, FSlateBrush* ConnectedIcon, FSlateBrush* DisconnectedIcon)
 	{
-		return GetBrush("HarmonixMetasound.MidiConnectedIcon");
+		FString ColorProperty = FString::Printf(TEXT("%s.%s.%s"), *ModuleName, *PinType.ToString(), *PinColorName);
+		FString ConnectedIconProperty = FString::Printf(TEXT("%s.%s.%s"), *ModuleName, *PinType.ToString(), *ConnectedIconName);
+		FString DisconnectedIconProperty = FString::Printf(TEXT("%s.%s.%s"), *ModuleName, *PinType.ToString(), *DisconnectedIconName);
+		Set(FName(ColorProperty), PinColor);
+		Set(FName(ConnectedIconProperty), ConnectedIcon);
+		Set(FName(DisconnectedIconProperty), DisconnectedIcon);
 	}
 
-	const FSlateBrush* FSlateStyle::GetMidiStreamDisconnectedIcon() const
+	const FLinearColor& FSlateStyle::GetPinColor(FName PinType) const
 	{
-		return GetBrush("HarmonixMetasound.MidiDisconnectedIcon");
+		FString PropertyName = FString::Printf(TEXT("%s.%s.%s"), *ModuleName, *PinType.ToString(), *PinColorName);
+		return GetColor(FName(PropertyName));
 	}
 
-	const FSlateBrush* FSlateStyle::GetMidiClockConnectedIcon() const
+	const FSlateBrush* FSlateStyle::GetConnectedIcon(FName PinType) const
 	{
-		return GetBrush("HarmonixMetasound.ClockConnectedIcon");
+		FString PropertyName = FString::Printf(TEXT("%s.%s.%s"), *ModuleName, *PinType.ToString(), *ConnectedIconName);
+        return GetBrush(FName(PropertyName));
 	}
-
-	const FSlateBrush* FSlateStyle::GetMidiClockDisconnectedIcon() const
-	{
-		return GetBrush("HarmonixMetasound.ClockDisconnectedIcon");
-	}
-
-	const FSlateBrush* FSlateStyle::GetTransportConnectedIcon() const
-	{
-		return GetBrush("HarmonixMetasound.TransportConnectedIcon");
-	}
-
-	const FSlateBrush* FSlateStyle::GetTransportDisconnectedIcon() const
+	
+	const FSlateBrush* FSlateStyle::GetDisconnectedIcon(FName PinType) const
 	{
-		return GetBrush("HarmonixMetasound.TransportDisconnectedIcon");
+		FString PropertyName = FString::Printf(TEXT("%s.%s.%s"), *ModuleName, *PinType.ToString(), *DisconnectedIconName);
+		return GetBrush(FName(PropertyName));
 	}
-
+	
 	FSlateStyle::~FSlateStyle()
 	{
 		FSlateStyleRegistry::UnRegisterSlateStyle(*this);
diff --git a/Engine/Plugins/Runtime/Harmonix/Source/HarmonixMetasoundEditor/Private/HarmonixMetasoundSlateStyle.h b/Engine/Plugins/Runtime/Harmonix/Source/HarmonixMetasoundEditor/Private/HarmonixMetasoundSlateStyle.h
--- a/Engine/Plugins/Runtime/Harmonix/Source/HarmonixMetasoundEditor/Private/HarmonixMetasoundSlateStyle.h
+++ b/Engine/Plugins/Runtime/Harmonix/Source/HarmonixMetasoundEditor/Private/HarmonixMetasoundSlateStyle.h
@@ -11,12 +11,12 @@ namespace HarmonixMetasoundEditor
 	public:
 		FSlateStyle();
 		static const FSlateStyle& Get();
-		const FSlateBrush* GetMidiStreamConnectedIcon() const;
-		const FSlateBrush* GetMidiStreamDisconnectedIcon() const;
-		const FSlateBrush* GetMidiClockConnectedIcon() const;
-		const FSlateBrush* GetMidiClockDisconnectedIcon() const;
-		const FSlateBrush* GetTransportConnectedIcon() const;
-		const FSlateBrush* GetTransportDisconnectedIcon() const;
+
+		void SetCustomPinStyle(const FName PinType, const FLinearColor& PinColor, FSlateBrush* ConnectedIcon, FSlateBrush* DisconnectedIcon);
+
+		const FLinearColor& GetPinColor(FName PinType) const;
+		const FSlateBrush* GetConnectedIcon(FName PinType) const;
+		const FSlateBrush* GetDisconnectedIcon(FName PinType) const;
 
 		~FSlateStyle();
 	};
