MidiPlayCursorMgr::AdvanceTrackerByDeltaMs should use the loop data from the passed in tracker

#tests HarmonixPluginTests

diff --git a/Engine/Plugins/Runtime/Harmonix/Source/HarmonixMidi/Private/HarmonixMidi/MidiPlayCursorMgr.cpp b/Engine/Plugins/Runtime/Harmonix/Source/HarmonixMidi/Private/HarmonixMidi/MidiPlayCursorMgr.cpp
--- a/Engine/Plugins/Runtime/Harmonix/Source/HarmonixMidi/Private/HarmonixMidi/MidiPlayCursorMgr.cpp
+++ b/Engine/Plugins/Runtime/Harmonix/Source/HarmonixMidi/Private/HarmonixMidi/MidiPlayCursorMgr.cpp
@@ -851,10 +851,10 @@ bool FMidiPlayCursorMgr::AdvanceTrackerByDeltaMs(float Ms, FMidiPlayCursorTracke
 	int32 NewTick = (int32)(TempoMap.MsToTick(NewMs) + 0.5f);
 	bool  Looped = false;
 
-	if (NewTick >= GetLowResTracker().LoopEndTick && Tracker.CurrentTick < GetLowResTracker().LoopEndTick && GetLowResTracker().Loop) // && !mDirectMappedFollower) If this is a direct mapped follower, than this function only gets called for the lo-res cursor, which needs to do normal loop handling!
+	if (NewTick >= Tracker.LoopEndTick && Tracker.CurrentTick < Tracker.LoopEndTick && Tracker.Loop)
 	{
 		// loop the tick around...
-		NewTick = (NewTick - GetLowResTracker().LoopEndTick) + GetLowResTracker().LoopStartTick;
+		NewTick = (NewTick - Tracker.LoopEndTick) + Tracker.LoopStartTick;
 		NewMs = TempoMap.TickToMs(NewTick);
 		Tracker.CurrentMs = NewMs;
 		Tracker.CurrentTick = NewTick;
