Fix MetronomeNode and MidiPlayerNode Tests

MetronomeNode and MidiPlayerNode tests need to process enough blocks to experience a loop when testing for looping

#tests HarmonixPluginTests

diff --git a/Engine/Plugins/Runtime/Harmonix/Source/HarmonixMetasoundTests/Private/Nodes/MetronomeNodeTests.cpp b/Engine/Plugins/Runtime/Harmonix/Source/HarmonixMetasoundTests/Private/Nodes/MetronomeNodeTests.cpp
--- a/Engine/Plugins/Runtime/Harmonix/Source/HarmonixMetasoundTests/Private/Nodes/MetronomeNodeTests.cpp
+++ b/Engine/Plugins/Runtime/Harmonix/Source/HarmonixMetasoundTests/Private/Nodes/MetronomeNodeTests.cpp
@@ -240,7 +240,6 @@ namespace HarmonixMetasoundTests::MetronomeNode
 			bool AllTicksEqual = true;
 			for (int32 BlockIndex = 0; BlockIndex < Params.NumBlocks; ++BlockIndex)
 			{
-
 				TAudioBuffer<float> Buffer{ Generator->GetNumChannels(), Params.NumSamplesPerBlock, EAudioBufferCleanupMode::Delete };
 				Generator->OnGenerateAudio(Buffer.GetRawChannelData(0), Buffer.GetNumTotalValidSamples());
 
@@ -253,7 +252,10 @@ namespace HarmonixMetasoundTests::MetronomeNode
 
 				float ClockTempo = (*OutputMidiClock)->GetTempoAtEndOfBlock();
 
-				AllTicksEqual |= InTest.TestEqual(Testf(FString::Printf(TEXT("Midi Clock Tempo at block: %d"), BlockIndex)), ClockTempo, Params.Tempo, 0.001f);
+				if (!InTest.TestEqual(Testf(FString::Printf(TEXT("Midi Clock Tempo at block: %d"), BlockIndex)), ClockTempo, Params.Tempo, 0.001f))
+				{
+					return false;
+				}
 
 				if (Params.Loop)
 				{
@@ -263,8 +265,10 @@ namespace HarmonixMetasoundTests::MetronomeNode
 				int32 ActualTick = (*OutputMidiClock)->GetCurrentMidiTick();
 				// allow for single tick tolerance?
 				ExpectedTick = FMath::Abs(ExpectedTick - ActualTick) <= 1 ? ActualTick : ExpectedTick;
-				if (!InTest.TestEqual(Testf(FString::Printf(TEXT("Tick at block: %d"), BlockIndex)), ActualTick, ExpectedTick))
+				if (AllTicksEqual && (ActualTick != ExpectedTick))
 				{
+					FString What = Testf(FString::Printf(TEXT("All ticks not equal. First failure at block: %d"), BlockIndex));
+					InTest.AddError(FString::Printf(TEXT("%s. Expected tick to be %d, but it was %d."), *What, ExpectedTick, ActualTick));
 					AllTicksEqual = false;
 				}
 
@@ -353,22 +357,24 @@ namespace HarmonixMetasoundTests::MetronomeNode
 		EAutomationTestFlags::EditorContext | EAutomationTestFlags::EngineFilter)
 		bool FMetronomeCreateNodeTestLoopingTempoRange::RunTest(const FString&)
 	{
-		bool AllTestsPass = true;
 		int32 Min = 4;
 		int32 Max = 240;
 		for (int32 Tempo = Min; Tempo <= Max; ++Tempo)
 		{
 			FBasicMetronomeTest::FParameters Params;
 			Params.Loop = true;
-			Params.NumBlocks = 10;
+			Params.NumBlocks = 500;
 			Params.Tempo = (float)Tempo;
 			Params.TimeSigNumerator = 4;
 			Params.TimeSigDenominator = 4;
 			FString TestString = FString::Printf(TEXT("Test %d BPM"), Tempo);
-			AllTestsPass |= FBasicMetronomeTest::RunTest(*this, Params, TestString);
+			if (!FBasicMetronomeTest::RunTest(*this, Params, TestString))
+			{
+				return false;
+			}
 		}
 
-		return AllTestsPass;
+		return !HasAnyErrors();
 	}
 
 	// Test range a good range of time signatures: 1/1 - 12/12.
@@ -378,7 +384,6 @@ namespace HarmonixMetasoundTests::MetronomeNode
 		EAutomationTestFlags::EditorContext | EAutomationTestFlags::EngineFilter)
 		bool FMetronomeCreateNodeTestLoopingAllTimeSigs::RunTest(const FString&)
 		{
-			bool AllTestsPass = true;
 			int32 Min = 1;
 			int32 Max = 12;
 			for (int32 Numerator = Min; Numerator <= Max; ++Numerator)
@@ -388,15 +393,18 @@ namespace HarmonixMetasoundTests::MetronomeNode
 					FBasicMetronomeTest::FParameters Params;
 					Params.Loop = true;
 					// only testing loop lengths, so don't need to advance clock
-					Params.NumBlocks = 2;
+					Params.NumBlocks = 500;
 					Params.TimeSigNumerator = Numerator;
 					Params.TimeSigDenominator = Denominator;
 					FString TestString = FString::Printf(TEXT("Test %d/%d Time"), Numerator, Denominator);
-					AllTestsPass |= FBasicMetronomeTest::RunTest(*this, Params, TestString);
+					if (!FBasicMetronomeTest::RunTest(*this, Params, TestString))
+					{
+						return false;
+					}
 				}
 			}
 
-			return AllTestsPass;
+			return !HasAnyErrors();
 		}
 
 	IMPLEMENT_SIMPLE_AUTOMATION_TEST(
@@ -406,7 +414,7 @@ namespace HarmonixMetasoundTests::MetronomeNode
 		bool FMetronomeCreateNodeTestTempoChangeInBlock::RunTest(const FString&)
 		{
 			FBasicMetronomeTest::FParameters Params;
-			Params.NumBlocks = 10;
+			Params.NumBlocks = 500;
 			return FBasicMetronomeTest::RunTest(*this, Params, TempoChangeTestString);
 		}
 }
diff --git a/Engine/Plugins/Runtime/Harmonix/Source/HarmonixMetasoundTests/Private/Nodes/MidiPlayerNodeTests.cpp b/Engine/Plugins/Runtime/Harmonix/Source/HarmonixMetasoundTests/Private/Nodes/MidiPlayerNodeTests.cpp
--- a/Engine/Plugins/Runtime/Harmonix/Source/HarmonixMetasoundTests/Private/Nodes/MidiPlayerNodeTests.cpp
+++ b/Engine/Plugins/Runtime/Harmonix/Source/HarmonixMetasoundTests/Private/Nodes/MidiPlayerNodeTests.cpp
@@ -188,7 +188,8 @@ namespace HarmonixMetasoundTests::MidiPlayerNode
 
 			int32 MidiPlayerLoopLengthTicks = MidiFileLength.LengthTicks;
 
-			bool AllTicksEqual = true;
+			bool MetronomeAllTicksEqual = true;
+			bool MidiPlayerAllTicksEqual = true;
 			for (int32 BlockIndex = 0; BlockIndex < Params.NumBlocks; ++BlockIndex)
 			{
 				TAudioBuffer<float> Buffer{ Generator->GetNumChannels(), Params.NumSamplesPerBlock, EAudioBufferCleanupMode::Delete };
@@ -223,12 +224,23 @@ namespace HarmonixMetasoundTests::MidiPlayerNode
 
 					int32 MetronomeActualTick = (*MetronomeClockOut)->GetCurrentMidiTick();
 					MetronomeExpectedTick = FMath::Abs(MidiPlayerExpectedTick - MetronomeActualTick) <= 1 ? MetronomeActualTick : MetronomeExpectedTick;
-					AllTicksEqual |= InTest.TestEqual(FString::Printf(TEXT("Metronome Out Clock, Tick at block: %d"), BlockIndex), MetronomeActualTick, MetronomeExpectedTick);
+					if (MetronomeAllTicksEqual && (MetronomeActualTick != MetronomeExpectedTick))
+					{
+						MetronomeAllTicksEqual = false;
+						FString What = FString::Printf(TEXT("Metronome Looping: Metronome Out Clock, Tick at block: %d"), BlockIndex);
+						InTest.AddError(FString::Printf(TEXT("%s: Expected Tick to be: %d, but was %d"), *What, MetronomeExpectedTick, MetronomeActualTick));
+					}
 				}
 
 				int32 MidiPlayerActualTick = (*MidiPlayerClockOut)->GetCurrentMidiTick();
 				MidiPlayerExpectedTick = FMath::Abs(MidiPlayerExpectedTick - MidiPlayerActualTick) <= 1 ? MidiPlayerActualTick : MidiPlayerExpectedTick;
-				AllTicksEqual |= InTest.TestEqual(FString::Printf(TEXT("Midi Player Out Clock, Tick at block: %d"), BlockIndex), MidiPlayerActualTick, MidiPlayerExpectedTick);
+
+				if (MidiPlayerAllTicksEqual && (MidiPlayerActualTick != MidiPlayerExpectedTick))
+				{
+					MidiPlayerAllTicksEqual = false;
+					FString What = FString::Printf(TEXT("Metronome Looping: Midi Player Out Clock, Tick at block: %d"), BlockIndex);
+					InTest.AddError(FString::Printf(TEXT("%s: Expected Tick to be: %d, but was %d"), *What, MidiPlayerExpectedTick, MidiPlayerActualTick));
+				}
 
 				Generator->ApplyToInputValue<FMusicTransportEventStream>(Inputs::TransportName, [](FMusicTransportEventStream& Transport)
 					{
