Always register MidiClockEventCursor with the MidiClock in the constructor of the midi clock.

Replace the TempoChangesCursor with the MidiClockEventCursor. The MidiClockEventCursor also handles tempo changes, so now we just need one cursor.

#tests PIE

diff --git a/Engine/Plugins/Runtime/Harmonix/Source/HarmonixMetasound/Private/HarmonixMetasound/DataTypes/MidiClock.cpp b/Engine/Plugins/Runtime/Harmonix/Source/HarmonixMetasound/Private/HarmonixMetasound/DataTypes/MidiClock.cpp
--- a/Engine/Plugins/Runtime/Harmonix/Source/HarmonixMetasound/Private/HarmonixMetasound/DataTypes/MidiClock.cpp
+++ b/Engine/Plugins/Runtime/Harmonix/Source/HarmonixMetasound/Private/HarmonixMetasound/DataTypes/MidiClock.cpp
@@ -40,7 +40,7 @@ namespace HarmonixMetasound
 	using namespace Metasound;
 	
 	FMidiClock::FMidiClock(const FOperatorSettings& InSettings)
-		: TempoChangesCursor(this)
+		: MidiClockEventCursor(this)
 		, BlockSize(InSettings.GetNumFramesPerBlock())
 		, CurrentBlockFrameIndex(0)
 		, SampleRate(InSettings.GetSampleRate())
@@ -50,14 +50,14 @@ namespace HarmonixMetasound
 	{
 		SpeedChangesInBlock.Add({0, 0.0f, 1.0f});
 		TempoChangesInBlock.Add({0, 0.0f, 120.0f});
-
-		DrivingMidiPlayCursorMgr->RegisterHiResPlayCursor(&TempoChangesCursor);
+		
+		DrivingMidiPlayCursorMgr->RegisterHiResPlayCursor(&MidiClockEventCursor);
 
 		RegisterForGameThreadUpdates();
 	}
 
 	FMidiClock::FMidiClock(const FMidiClock& Other)
-		: TempoChangesCursor(this)
+		: MidiClockEventCursor(this)
 		, BlockSize(Other.BlockSize)
 		, CurrentBlockFrameIndex(Other.CurrentBlockFrameIndex)
 		, SampleRate(Other.SampleRate)
@@ -73,7 +73,7 @@ namespace HarmonixMetasound
 		, SmoothingEnabled(Other.SmoothingEnabled)
 		, DrivingMidiPlayCursorMgr(Other.DrivingMidiPlayCursorMgr)
 	{
-		DrivingMidiPlayCursorMgr->RegisterHiResPlayCursor(&TempoChangesCursor);
+		DrivingMidiPlayCursorMgr->RegisterHiResPlayCursor(&MidiClockEventCursor);
 
 		RegisterForGameThreadUpdates();
 	}
@@ -86,7 +86,7 @@ namespace HarmonixMetasound
 
 			if (DrivingMidiPlayCursorMgr != Other.DrivingMidiPlayCursorMgr)
 			{
-				DrivingMidiPlayCursorMgr->UnregisterPlayCursor(&TempoChangesCursor);
+				DrivingMidiPlayCursorMgr->UnregisterPlayCursor(&MidiClockEventCursor);
 			}
 
 			BlockSize = Other.BlockSize;
@@ -106,7 +106,7 @@ namespace HarmonixMetasound
 			if (DrivingMidiPlayCursorMgr != Other.DrivingMidiPlayCursorMgr)
 			{
 				DrivingMidiPlayCursorMgr = Other.DrivingMidiPlayCursorMgr;
-				DrivingMidiPlayCursorMgr->RegisterHiResPlayCursor(&TempoChangesCursor);
+				DrivingMidiPlayCursorMgr->RegisterHiResPlayCursor(&MidiClockEventCursor);
 			}
 
 			RegisterForGameThreadUpdates();
@@ -119,31 +119,7 @@ namespace HarmonixMetasound
 	{
 		UnregisterForGameThreadUpdates();
 	
-		DrivingMidiPlayCursorMgr->UnregisterPlayCursor(&TempoChangesCursor, false);
-	}
-
-	FMidiClock::FTempoChangesCursor::FTempoChangesCursor(FMidiClock* MidiClock) : MyMidiClock(MidiClock)
-	{
-		check(MyMidiClock);
-		SetMessageFilter(FMidiPlayCursor::EFilterPassFlags::Tempo);
-	}
-
-	void FMidiClock::FTempoChangesCursor::OnTempo(int32 TrackIndex, int32 Tick, int32 Tempo, bool IsPreroll /*= false*/)
-	{
-		check(MyMidiClock);
-		int32 BlockFrameIndex = MyMidiClock->CurrentBlockFrameIndex;
-
-		check(BlockFrameIndex >= MyMidiClock->TempoChangesInBlock.Last().BlockSampleFrameIndex);
-		MyMidiClock->HasTempoChangeInBlock = true;
-		float Bpm = Harmonix::Midi::Constants::MidiTempoToBPM(Tempo);
-		if (MyMidiClock->TempoChangesInBlock.Last().BlockSampleFrameIndex == BlockFrameIndex)
-		{
-			MyMidiClock->TempoChangesInBlock.Last().Tempo = Bpm;
-		}
-		else
-		{ 
-			MyMidiClock->TempoChangesInBlock.Add({BlockFrameIndex, 0.0f, Bpm});
-		}
+		DrivingMidiPlayCursorMgr->UnregisterPlayCursor(&MidiClockEventCursor, false);
 	}
 
 	void FMidiClock::RegisterForGameThreadUpdates()
@@ -479,55 +455,66 @@ namespace HarmonixMetasound
 		return OutMidiData;
 	}
 
-
-	FMidiClockEventCursor::FMidiClockEventCursor(FMidiClockWriteRef InClock)
-		: MidiClock(InClock)
+	FMidiClock::FMidiClockEventCursor::FMidiClockEventCursor(FMidiClock* MidiClock)
+		: MyMidiClock(MidiClock)
 	{
-		MidiClock->RegisterHiResPlayCursor(this);
+		check(MidiClock);
 	}
 
-	FMidiClockEventCursor::~FMidiClockEventCursor()
-	{
-		MidiClock->UnregisterPlayCursor(this);
-	}
-
-	void FMidiClockEventCursor::Reset(bool ForceNoBroadcast /*= false*/)
+	void FMidiClock::FMidiClockEventCursor::Reset(bool ForceNoBroadcast /*= false*/)
 	{
 		const int32 FromTick = CurrentTick;
 		FMidiPlayCursor::Reset(ForceNoBroadcast);
-		AddEvent(FMidiClockEvent::MakeResetEvent(MidiClock->GetCurrentBlockFrameIndex(), FromTick, CurrentTick, ForceNoBroadcast));
+		AddEvent(FMidiClockEvent::MakeResetEvent(MyMidiClock->GetCurrentBlockFrameIndex(), FromTick, CurrentTick, ForceNoBroadcast));
 	}
 
-	void FMidiClockEventCursor::OnLoop(int32 LoopStartTick, int32 LoopEndTick)
+	void FMidiClock::FMidiClockEventCursor::OnLoop(int32 LoopStartTick, int32 LoopEndTick)
 	{
 		FMidiPlayCursor::OnLoop(LoopStartTick, LoopEndTick);
-		AddEvent(FMidiClockEvent::MakeLoopEvent(MidiClock->GetCurrentBlockFrameIndex(), LoopStartTick, LoopEndTick));
+		AddEvent(FMidiClockEvent::MakeLoopEvent(MyMidiClock->GetCurrentBlockFrameIndex(), LoopStartTick, LoopEndTick));
 	}
 
-	void FMidiClockEventCursor::SeekToTick(int32 Tick) 
+	void FMidiClock::FMidiClockEventCursor::SeekToTick(int32 Tick) 
 	{
 		const int32 FromTick = CurrentTick;
 		FMidiPlayCursor::SeekToTick(Tick);
-		AddEvent(FMidiClockEvent::MakeSeekToEvent(MidiClock->GetCurrentBlockFrameIndex(), FromTick, Tick));
+		AddEvent(FMidiClockEvent::MakeSeekToEvent(MyMidiClock->GetCurrentBlockFrameIndex(), FromTick, Tick));
 	}
 
-	void FMidiClockEventCursor::SeekThruTick(int32 Tick)
+	void FMidiClock::FMidiClockEventCursor::SeekThruTick(int32 Tick)
 	{
 		const int32 FromTick = CurrentTick;
 		FMidiPlayCursor::SeekThruTick(Tick);
-		AddEvent(FMidiClockEvent::MakeSeekThruEvent(MidiClock->GetCurrentBlockFrameIndex(), FromTick, Tick));
+		AddEvent(FMidiClockEvent::MakeSeekThruEvent(MyMidiClock->GetCurrentBlockFrameIndex(), FromTick, Tick));
 	}
 
-	void FMidiClockEventCursor::AdvanceThruTick(int32 Tick, bool IsPreRoll)
+	void FMidiClock::FMidiClockEventCursor::AdvanceThruTick(int32 Tick, bool IsPreRoll)
 	{
 		const int32 FromTick = CurrentTick;
 		FMidiPlayCursor::AdvanceThruTick(Tick, IsPreRoll);
-		AddEvent(FMidiClockEvent::MakeAdvanceThruEvent(MidiClock->GetCurrentBlockFrameIndex(), FromTick, Tick, IsPreRoll));
+		AddEvent(FMidiClockEvent::MakeAdvanceThruEvent(MyMidiClock->GetCurrentBlockFrameIndex(), FromTick, Tick, IsPreRoll));
+	}
+
+	void FMidiClock::FMidiClockEventCursor::OnTempo(int32 TrackIndex, int32 Tick, int32 Tempo, bool IsPreroll)
+	{
+		int32 BlockFrameIndex = MyMidiClock->CurrentBlockFrameIndex;
+
+		check(BlockFrameIndex >= MyMidiClock->TempoChangesInBlock.Last().BlockSampleFrameIndex);
+		MyMidiClock->HasTempoChangeInBlock = true;
+		float Bpm = Harmonix::Midi::Constants::MidiTempoToBPM(Tempo);
+		if (MyMidiClock->TempoChangesInBlock.Last().BlockSampleFrameIndex == BlockFrameIndex)
+		{
+			MyMidiClock->TempoChangesInBlock.Last().Tempo = Bpm;
+		}
+		else
+		{ 
+			MyMidiClock->TempoChangesInBlock.Add({BlockFrameIndex, 0.0f, Bpm});
+		}
 	}
 
-	void FMidiClockEventCursor::AddEvent(const FMidiClockEvent& InEvent)
+	void FMidiClock::FMidiClockEventCursor::AddEvent(const FMidiClockEvent& InEvent)
 	{
-		TArray<FMidiClockEvent>& Events = MidiClock->MidiClockEventsInBlock;
+		TArray<FMidiClockEvent>& Events = MyMidiClock->MidiClockEventsInBlock;
 		if (!Events.IsEmpty())
 		{
 			check(Events.Last().BlockFrameIndex <= InEvent.BlockFrameIndex);
diff --git a/Engine/Plugins/Runtime/Harmonix/Source/HarmonixMetasound/Private/HarmonixMetasound/Nodes/MetronomeNode.cpp b/Engine/Plugins/Runtime/Harmonix/Source/HarmonixMetasound/Private/HarmonixMetasound/Nodes/MetronomeNode.cpp
--- a/Engine/Plugins/Runtime/Harmonix/Source/HarmonixMetasound/Private/HarmonixMetasound/Nodes/MetronomeNode.cpp
+++ b/Engine/Plugins/Runtime/Harmonix/Source/HarmonixMetasound/Private/HarmonixMetasound/Nodes/MetronomeNode.cpp
@@ -75,7 +75,6 @@ namespace HarmonixMetasound
 
 		//** DATA
 		FMidiClock MetronomeClock;
-		FMidiClockEventCursor MidiClockEventCursor;
 		TSharedPtr<FMidiFileData> MidiData;
 		FSampleCount BlockSize;
 		float        SampleRate;
@@ -192,7 +191,6 @@ namespace HarmonixMetasound
 		, SeekPreRollBarsInPin(InPreRollBars)
 		, MidiClockOutPin(FMidiClockWriteRef::CreateNew(InParams.OperatorSettings))
 		, MetronomeClock(InParams.OperatorSettings)
-		, MidiClockEventCursor(MidiClockOutPin)
 		, BlockSize(InParams.OperatorSettings.GetNumFramesPerBlock())
 		, SampleRate(InParams.OperatorSettings.GetSampleRate())
 		, CurrentTempo(*TempoInPin)
diff --git a/Engine/Plugins/Runtime/Harmonix/Source/HarmonixMetasound/Private/HarmonixMetasound/Nodes/MidiClockOffsetNode.cpp b/Engine/Plugins/Runtime/Harmonix/Source/HarmonixMetasound/Private/HarmonixMetasound/Nodes/MidiClockOffsetNode.cpp
--- a/Engine/Plugins/Runtime/Harmonix/Source/HarmonixMetasound/Private/HarmonixMetasound/Nodes/MidiClockOffsetNode.cpp
+++ b/Engine/Plugins/Runtime/Harmonix/Source/HarmonixMetasound/Private/HarmonixMetasound/Nodes/MidiClockOffsetNode.cpp
@@ -107,7 +107,6 @@ namespace HarmonixMetasound::Nodes::MidiClockOffset
 			, OffsetBeatsInPin(InOffsetBeats)
 			, OffsetMsInPin(InOffsetMs)
 			, MidiClockOut(FMidiClockWriteRef::CreateNew(InSettings))
-			, MidiClockEventCursor(MidiClockOut)
 			, BlockSize(InSettings.GetNumFramesPerBlock())
 		{
 			TSharedPtr<FMidiFileData> ConductorMidiData = FMidiClock::MakeClockConductorMidiData(120.0f, 4, 4);
@@ -241,7 +240,6 @@ namespace HarmonixMetasound::Nodes::MidiClockOffset
 		FMidiClockWriteRef MidiClockOut;
 
 		//** DATA
-		FMidiClockEventCursor MidiClockEventCursor;
 		FSampleCount BlockSize      = 0;
 		int32 CurrentBlockSpanStart = 0;
 
diff --git a/Engine/Plugins/Runtime/Harmonix/Source/HarmonixMetasound/Private/HarmonixMetasound/Nodes/MidiPlayerNode.cpp b/Engine/Plugins/Runtime/Harmonix/Source/HarmonixMetasound/Private/HarmonixMetasound/Nodes/MidiPlayerNode.cpp
--- a/Engine/Plugins/Runtime/Harmonix/Source/HarmonixMetasound/Private/HarmonixMetasound/Nodes/MidiPlayerNode.cpp
+++ b/Engine/Plugins/Runtime/Harmonix/Source/HarmonixMetasound/Private/HarmonixMetasound/Nodes/MidiPlayerNode.cpp
@@ -81,7 +81,6 @@ namespace HarmonixMetasound
 		FMidiClockWriteRef  MidiClockOut;
 
 		//** DATA
-		FMidiClockEventCursor MidiClockEventCursor;
 		FMidiFileProxyPtr CurrentMidiFile;
 		FSampleCount BlockSize      = 0;
 		int32 CurrentBlockSpanStart = 0;
@@ -258,7 +257,6 @@ namespace HarmonixMetasound
 		, bKillVoicesOnMidiChange(bInKillVoicesOnMidiChange)
 		, MidiOutPin(FMidiStreamWriteRef::CreateNew(InSettings))
 		, MidiClockOut(FMidiClockWriteRef::CreateNew(InSettings))
-		, MidiClockEventCursor(MidiClockOut)
 		, BlockSize(InSettings.GetNumFramesPerBlock())
 	{
 		MidiClockOut->RegisterHiResPlayCursor(this);
diff --git a/Engine/Plugins/Runtime/Harmonix/Source/HarmonixMetasound/Public/HarmonixMetasound/DataTypes/MidiClock.h b/Engine/Plugins/Runtime/Harmonix/Source/HarmonixMetasound/Public/HarmonixMetasound/DataTypes/MidiClock.h
--- a/Engine/Plugins/Runtime/Harmonix/Source/HarmonixMetasound/Public/HarmonixMetasound/DataTypes/MidiClock.h
+++ b/Engine/Plugins/Runtime/Harmonix/Source/HarmonixMetasound/Public/HarmonixMetasound/DataTypes/MidiClock.h
@@ -192,22 +192,27 @@ namespace HarmonixMetasound
 
 	private:
 		friend class FMidiClockEventCursor;
-		friend class FTempoChangesCursor;
 
-		class FTempoChangesCursor : public FMidiPlayCursor
+		class FMidiClockEventCursor : public FMidiPlayCursor
 		{
 		public:
-			FTempoChangesCursor(FMidiClock* MidiClock);
-
-			//** BEGIN FMidiPlayCursor
-			void OnTempo(int32 TrackIndex, int32 Tick, int32 Tempo, bool IsPreroll = false) override;
-			//** END FMidiPlayCursor
-
+			FMidiClockEventCursor(FMidiClock* MidiClock);
+		
+			//~ BEGIN FMidiPlayCursor Overrides
+			virtual void Reset(bool ForceNoBroadcast = false) override;
+			virtual void OnLoop(int32 LoopStartTick, int32 LoopEndTick) override;
+			virtual void SeekToTick(int32 Tick) override;
+			virtual void SeekThruTick(int32 Tick) override;
+			virtual void AdvanceThruTick(int32 Tick, bool IsPreRoll) override;
+			virtual void OnTempo(int32 TrackIndex, int32 Tick, int32 Tempo, bool IsPreroll = false) override;
+			//~ END FMidiPlayCursor Overrides
+			
+			void AddEvent(const FMidiClockEvent& InEvent);
 		private:
 			FMidiClock* MyMidiClock = nullptr;
 		};
-
-		FTempoChangesCursor TempoChangesCursor;
+		
+		FMidiClockEventCursor MidiClockEventCursor;
 
 	private:
 		void RegisterForGameThreadUpdates();
@@ -237,28 +242,6 @@ namespace HarmonixMetasound
 
 	// Declare aliases IN the namespace...
 	DECLARE_METASOUND_DATA_REFERENCE_ALIAS_TYPES(FMidiClock, FMidiClockTypeInfo, FMidiClockReadRef, FMidiClockWriteRef)
-
-
-	class HARMONIXMETASOUND_API FMidiClockEventCursor : public FMidiPlayCursor
-	{
-	public:
-
-		FMidiClockEventCursor(FMidiClockWriteRef InClock);
-		virtual ~FMidiClockEventCursor();
-
-		//~ BEGIN FMidiPlayCursor Overrides
-		virtual void Reset(bool ForceNoBroadcast = false) override;
-		virtual void OnLoop(int32 LoopStartTick, int32 LoopEndTick) override;
-		virtual void SeekToTick(int32 Tick) override;
-		virtual void SeekThruTick(int32 Tick) override;
-		virtual void AdvanceThruTick(int32 Tick, bool IsPreRoll) override;
-		//~ END FMidiPlayCursor Overrides
-	private:
-
-		void AddEvent(const FMidiClockEvent& InEvent);
-
-		FMidiClockWriteRef MidiClock;
-	};
 }
 
 // Declare reference types OUT of the namespace...
