Use MidiClockEvents for MidiClockTimestampTriggerNode

#tests PIE

diff --git a/Engine/Plugins/Runtime/Harmonix/Source/HarmonixMetasound/Private/HarmonixMetasound/Nodes/MidiClockTimestampTriggerNode.cpp b/Engine/Plugins/Runtime/Harmonix/Source/HarmonixMetasound/Private/HarmonixMetasound/Nodes/MidiClockTimestampTriggerNode.cpp
--- a/Engine/Plugins/Runtime/Harmonix/Source/HarmonixMetasound/Private/HarmonixMetasound/Nodes/MidiClockTimestampTriggerNode.cpp
+++ b/Engine/Plugins/Runtime/Harmonix/Source/HarmonixMetasound/Private/HarmonixMetasound/Nodes/MidiClockTimestampTriggerNode.cpp
@@ -27,7 +27,7 @@ namespace HarmonixMetasound
 {
 	using namespace Metasound;
 
-	class FMidiClockTimestampTriggerOperator : public TExecutableOperator<FMidiClockTimestampTriggerOperator>, public FMidiPlayCursor
+	class FMidiClockTimestampTriggerOperator : public TExecutableOperator<FMidiClockTimestampTriggerOperator>
 	{
 	public:
 		static const FNodeClassMetadata& GetNodeInfo();
@@ -66,29 +66,6 @@ namespace HarmonixMetasound
 
 		int32 TriggerTick = 0;
 
-		struct FTickSpan
-		{
-			int32 FromTick;
-			int32 ThruTick;
-			int32 BlockFrameIndex;
-			bool  bIsSeek;
-			FTickSpan(int32 InFromTick, int32 InThruTick, int32 InBlockFrameIndex, bool InIsSeek)
-				: FromTick(InFromTick)
-				, ThruTick(InThruTick)
-				, BlockFrameIndex(InBlockFrameIndex)
-				, bIsSeek(InIsSeek)
-			{}
-		};
-		TArray<FTickSpan> TickSpans;
-
-		//** BEGIN FMidiPlayCursor
-		virtual void SeekToTick(int32 Tick) override { SeekThruTick(Tick - 1); }
-		virtual void SeekThruTick(int32 Tick) override;
-		virtual void AdvanceThruTick(int32 Tick, bool IsPreRoll) override;
-		// We have to override this to disambiguate the FMidiPlayCursor Reset and the MS operator Reset
-		virtual void Reset(bool ForceNoBroadcast) override { FMidiPlayCursor::Reset(ForceNoBroadcast); }
-		//** END FMidiPlayCursor
-
 		void CalculateTriggerTick();
 	};
 
@@ -215,11 +192,6 @@ namespace HarmonixMetasound
 	void FMidiClockTimestampTriggerOperator::Reset(const FResetParams& ResetParams)
 	{
 		TriggerOutPin->Reset();
-		
-		FMidiPlayCursor::Reset(true);
-		
-		SetMessageFilter(FMidiPlayCursor::EFilterPassFlags::None);
-		MidiClockInPin->RegisterHiResPlayCursor(this);
 
 		CurrentTimestamp = *TimestampInPin;
 		
@@ -247,45 +219,31 @@ namespace HarmonixMetasound
 
 		if (*EnableInPin)
 		{
-			for (const FTickSpan& Span : TickSpans)
+			for (const FMidiClockEvent& ClockEvent : MidiClockInPin->GetMidiClockEventsInBlock())
 			{
-				if (Span.FromTick < TriggerTick && Span.ThruTick >= TriggerTick && 
-					(!Span.bIsSeek || *TriggerDuringSeekInPin))
+				switch (ClockEvent.Type)
 				{
-					TriggerOutPin->TriggerFrame(Span.BlockFrameIndex);
-					break;
+				case FMidiClockEvent::EType::SeekThru:
+				case FMidiClockEvent::EType::AdvanceThru:
+					{
+						if (ClockEvent.IsPreRoll)
+						{
+							continue;
+						}
+						
+						const bool EventIsSeek = ClockEvent.Type == FMidiClockEvent::EType::SeekThru;
+
+						if (ClockEvent.Tick1 < TriggerTick && ClockEvent.Tick2 >= TriggerTick)
+						{
+							if (!EventIsSeek || *TriggerDuringSeekInPin)
+							{
+								TriggerOutPin->TriggerFrame(ClockEvent.BlockFrameIndex);
+							}
+						}
+					}
 				}
 			}
 		}
-		TickSpans.Empty(8);
-	}
-
-	void FMidiClockTimestampTriggerOperator::SeekThruTick(int32 Tick)
-	{
-		int32 TickProceedingThisAdvance = CurrentTick;
-		FMidiPlayCursor::SeekThruTick(Tick);
-
-		// don't trigger if seeking backward or no progress being made...
-		if (Tick < TickProceedingThisAdvance || TickProceedingThisAdvance == CurrentTick)
-		{
-			return;
-		}
-
-		TickSpans.Emplace(TickProceedingThisAdvance, CurrentTick, MidiClockInPin->GetCurrentBlockFrameIndex(), true);
-	}
-
-	void FMidiClockTimestampTriggerOperator::AdvanceThruTick(int32 Tick, bool IsPreRoll)
-	{
-		int32 TickProceedingThisAdvance = CurrentTick;
-		FMidiPlayCursor::AdvanceThruTick(Tick, IsPreRoll);
-
-		// don't trigger during preroll or if no progress being made...
-		if (IsPreRoll || TickProceedingThisAdvance == CurrentTick)
-		{
-			return;
-		}
-
-		TickSpans.Emplace(TickProceedingThisAdvance, CurrentTick, MidiClockInPin->GetCurrentBlockFrameIndex(), false);
 	}
 }
 
