Use MidiClockEvents in MidiClockSubdivisionTriggerNode
#tests PIE

diff --git a/Engine/Plugins/Runtime/Harmonix/Source/HarmonixMetasound/Private/HarmonixMetasound/Nodes/MidiClockSubdivisionTriggerNode.cpp b/Engine/Plugins/Runtime/Harmonix/Source/HarmonixMetasound/Private/HarmonixMetasound/Nodes/MidiClockSubdivisionTriggerNode.cpp
--- a/Engine/Plugins/Runtime/Harmonix/Source/HarmonixMetasound/Private/HarmonixMetasound/Nodes/MidiClockSubdivisionTriggerNode.cpp
+++ b/Engine/Plugins/Runtime/Harmonix/Source/HarmonixMetasound/Private/HarmonixMetasound/Nodes/MidiClockSubdivisionTriggerNode.cpp
@@ -27,7 +27,7 @@ namespace HarmonixMetasound
 	using namespace Metasound;
 
 
-	class FMidiClockSubdivisionTriggerOperator : public TExecutableOperator<FMidiClockSubdivisionTriggerOperator>, public FMidiPlayCursor
+	class FMidiClockSubdivisionTriggerOperator : public TExecutableOperator<FMidiClockSubdivisionTriggerOperator>
 	{
 	public:
 		static const FNodeClassMetadata& GetNodeInfo();
@@ -70,14 +70,6 @@ namespace HarmonixMetasound
 
 		int32 GridOffsetTicks = 0;
 		int32 GridSizeTicks	  = 0;
-		bool  AdvanceBlockNeeded = true;
-
-		//** BEGIN FMidiPlayCursor
-		virtual void OnTimeSig(int32 TrackIndex, int32 Tick, int32 Numerator, int32 Denominator, bool IsPreroll = false) override;
-		virtual void AdvanceThruTick(int32 Tick, bool IsPreRoll) override;
-		// We have to override this to disambiguate the FMidiPlayCursor reset and the MS operator one
-		virtual void Reset(bool ForceNoBroadcast) override { FMidiPlayCursor::Reset(ForceNoBroadcast); }
-		//** END FMidiPlayCursor
 	};
 
 	class FMidiClockSubdivisionTriggerNode : public FNodeFacade
@@ -207,9 +199,6 @@ namespace HarmonixMetasound
 	{
 		TriggerOutPin->Reset();
 		
-		SetMessageFilter(FMidiPlayCursor::EFilterPassFlags::TimeSig);
-		MidiClockInPin->RegisterHiResPlayCursor(this);
-
 		GridSizeUnits = *GridSizeUnitsInPin;
 		GridSizeMultiplier = *GridSizeMultInPin;
 		if (GridSizeMultiplier < 1)
@@ -222,18 +211,12 @@ namespace HarmonixMetasound
 		const FSongMaps& SongMaps = MidiClockInPin->GetSongMaps();
 		GridOffsetTicks = SongMaps.SubdivisionToMidiTicks(GridOffsetUnits, 0) * GridOffsetMultiplier;
 		GridSizeTicks = SongMaps.SubdivisionToMidiTicks(GridSizeUnits, 0) * GridSizeMultiplier;
-
-		AdvanceBlockNeeded = true;
 	}
 
 	void FMidiClockSubdivisionTriggerOperator::Execute()
 	{
-		if (AdvanceBlockNeeded)
-		{
-			TriggerOutPin->AdvanceBlock();
-			AdvanceBlockNeeded = false;
-		}
-
+		TriggerOutPin->AdvanceBlock();
+			
 		// first let's see if our configuration has changed at all...
 		if (*GridSizeUnitsInPin != GridSizeUnits || *GridSizeMultInPin != GridSizeMultiplier)
 		{
@@ -243,8 +226,6 @@ namespace HarmonixMetasound
 			{
 				GridSizeMultiplier = 1;
 			}
-			const FSongMaps& SongMaps = MidiClockInPin->GetSongMaps();
-			GridSizeTicks = SongMaps.SubdivisionToMidiTicks(GridSizeUnits, CurrentTick) * GridSizeMultiplier;
 		}
 
 		if (*GridOffsetUnitsInPin != GridOffsetUnits || *GridOffsetMultInPin != GridOffsetMultiplier)
@@ -255,71 +236,60 @@ namespace HarmonixMetasound
 			GridOffsetTicks = SongMaps.SubdivisionToMidiTicks(GridOffsetUnits, 0) * GridOffsetMultiplier;
 		}
 
-		AdvanceBlockNeeded = true;
-	}
-
-	void FMidiClockSubdivisionTriggerOperator::AdvanceThruTick(int32 Tick, bool IsPreRoll)
-	{
-		if (AdvanceBlockNeeded)
+		if (*EnableInPin)
 		{
-			TriggerOutPin->AdvanceBlock();
-			AdvanceBlockNeeded = false;
-		}
-
-		int32 TickProceedingThisAdvance = CurrentTick - GridOffsetTicks;
-
-		FMidiPlayCursor::AdvanceThruTick(Tick, IsPreRoll);
-		if (IsPreRoll || !*EnableInPin)
-		{
-			return;
-		}
-
-		int32 LastTickProcessed = CurrentTick - GridOffsetTicks;
-		if (LastTickProcessed < 0)
-		{
-			return;
-		}
-
-		// if this is one of the quantization types that is variable in size we have to do it a little different...
-		if (GridSizeUnits == EMidiClockSubdivisionQuantization::Bar ||
-			GridSizeUnits == EMidiClockSubdivisionQuantization::Beat)
-		{
-			const FSongMaps& SongMaps = MidiClockInPin->GetSongMaps();
-			FMusicTimestamp Start = SongMaps.GetBarMap().TickToMusicTimestamp(TickProceedingThisAdvance);
-			FMusicTimestamp End = SongMaps.GetBarMap().TickToMusicTimestamp(LastTickProcessed);
-			if (GridSizeUnits == EMidiClockSubdivisionQuantization::Bar && Start.Bar != End.Bar)
-			{
-				if (End.Bar % GridSizeMultiplier == 0)
-				{
-					TriggerOutPin->TriggerFrame(MidiClockInPin->GetCurrentBlockFrameIndex());
-				}
-			}
-			else if (GridSizeUnits == EMidiClockSubdivisionQuantization::Beat && FMath::FloorToInt32(Start.Beat) != FMath::FloorToInt32(End.Beat))
+			for (const FMidiClockEvent& ClockEvent : MidiClockInPin->GetMidiClockEventsInBlock())
 			{
-				if (FMath::FloorToInt32(End.Beat-1.0f) % GridSizeMultiplier == 0)
+				if (ClockEvent.Type == FMidiClockEvent::EType::AdvanceThru)
 				{
-					TriggerOutPin->TriggerFrame(MidiClockInPin->GetCurrentBlockFrameIndex());
+					if (ClockEvent.IsPreRoll)
+					{
+						continue;
+					}
+					
+					int32 TickPreceedingThisAdvance = ClockEvent.Tick1 - GridOffsetTicks;
+					int32 LastTickProcessed = ClockEvent.Tick2 - GridOffsetTicks;
+					const FSongMaps& SongMaps = MidiClockInPin->GetSongMaps();
+					
+					// calculate new GridSizeTicks and GridOffsetTicks at the tick we're evaluating based on any time signature changes
+					GridSizeTicks = SongMaps.SubdivisionToMidiTicks(GridSizeUnits, ClockEvent.Tick2) * GridSizeMultiplier;
+					
+					// if this is one of the quantization types that is variable in size we have to do it a little different...
+					if (GridSizeUnits == EMidiClockSubdivisionQuantization::Bar ||
+						GridSizeUnits == EMidiClockSubdivisionQuantization::Beat)
+					{
+						
+						FMusicTimestamp Start = SongMaps.GetBarMap().TickToMusicTimestamp(TickPreceedingThisAdvance);
+						FMusicTimestamp End = SongMaps.GetBarMap().TickToMusicTimestamp(LastTickProcessed);
+						if (GridSizeUnits == EMidiClockSubdivisionQuantization::Bar && Start.Bar != End.Bar)
+						{
+							if (End.Bar % GridSizeMultiplier == 0)
+							{
+								TriggerOutPin->TriggerFrame(MidiClockInPin->GetCurrentBlockFrameIndex());
+							}
+						}
+						else if (GridSizeUnits == EMidiClockSubdivisionQuantization::Beat && FMath::FloorToInt32(Start.Beat) != FMath::FloorToInt32(End.Beat))
+						{
+							if (FMath::FloorToInt32(End.Beat-1.0f) % GridSizeMultiplier == 0)
+							{
+								TriggerOutPin->TriggerFrame(MidiClockInPin->GetCurrentBlockFrameIndex());
+							}
+						}
+					}
+					else
+					{
+						
+						int32 StartingGridSquare = TickPreceedingThisAdvance / GridSizeTicks;
+						int32 EndingGridSquare = LastTickProcessed / GridSizeTicks;
+						if (StartingGridSquare != EndingGridSquare)
+						{
+							TriggerOutPin->TriggerFrame(MidiClockInPin->GetCurrentBlockFrameIndex());
+						}
+					}
 				}
 			}
 		}
-		else
-		{
-			int32 StartingGridSquare = TickProceedingThisAdvance / GridSizeTicks;
-			int32 EndingGridSquare = LastTickProcessed / GridSizeTicks;
-			if (StartingGridSquare != EndingGridSquare)
-			{
-				TriggerOutPin->TriggerFrame(MidiClockInPin->GetCurrentBlockFrameIndex());
-			}
-		}
-	}
-
-	void FMidiClockSubdivisionTriggerOperator::OnTimeSig(int32 TrackIndex, int32 Tick, int32 Numerator, int32 Denominator, bool IsPreroll /*= false*/)
-	{
-		const FSongMaps& SongMaps = MidiClockInPin->GetSongMaps();
-		GridSizeTicks = SongMaps.SubdivisionToMidiTicks(GridSizeUnits, Tick) * GridSizeMultiplier;
-		//NOTE: We DO NOT calculate a new grid offset because this is always specified from the beginning of the song... at the song's starting time signature.  
 	}
-
 }
 
 #undef LOCTEXT_NAMESPACE // "HarmonixMetaSound"
