Fix issue with offline rendering of metasound on dedicated server and the new MusicClockComponent functional.

diff --git a/Engine/Plugins/Runtime/Harmonix/Source/HarmonixMetasound/Private/HarmonixMetasound/Components/MetasoundOfflinePlayerComponent.cpp b/Engine/Plugins/Runtime/Harmonix/Source/HarmonixMetasound/Private/HarmonixMetasound/Components/MetasoundOfflinePlayerComponent.cpp
--- a/Engine/Plugins/Runtime/Harmonix/Source/HarmonixMetasound/Private/HarmonixMetasound/Components/MetasoundOfflinePlayerComponent.cpp
+++ b/Engine/Plugins/Runtime/Harmonix/Source/HarmonixMetasound/Private/HarmonixMetasound/Components/MetasoundOfflinePlayerComponent.cpp
@@ -161,7 +161,7 @@ void UMetasoundOfflinePlayerComponent::TickComponent(float DeltaTime, ELevelTick
 
 	while (SamplesNeeded >= GeneratorBlockSize)
 	{
-		Generator->OnGenerateAudio(ScratchBuffer.GetData(), ScratchBuffer.Num());
+		Generator->GetNextBuffer(ScratchBuffer.GetData(), ScratchBuffer.Num(), true);
 		SamplesNeeded -= GeneratorBlockSize;
 		RenderedSamples += GeneratorBlockSize;
 	}
