Fix Seek Target Builder for "Bar & Beat To Tick"

#tests PIE

diff --git a/Engine/Plugins/Runtime/Harmonix/Source/HarmonixMetasound/Private/HarmonixMetasound/Nodes/MusicSeekTargetBuilder.cpp b/Engine/Plugins/Runtime/Harmonix/Source/HarmonixMetasound/Private/HarmonixMetasound/Nodes/MusicSeekTargetBuilder.cpp
--- a/Engine/Plugins/Runtime/Harmonix/Source/HarmonixMetasound/Private/HarmonixMetasound/Nodes/MusicSeekTargetBuilder.cpp
+++ b/Engine/Plugins/Runtime/Harmonix/Source/HarmonixMetasound/Private/HarmonixMetasound/Nodes/MusicSeekTargetBuilder.cpp
@@ -119,7 +119,7 @@ namespace HarmonixMetasound::Nodes::BarBeatToSeekTarget
 
 		void Execute()
 		{
-			SeekTargetOutPin->Type = ESeekPointType::Millisecond;
+			SeekTargetOutPin->Type = ESeekPointType::BarBeat;
 			SeekTargetOutPin->BarBeat.Bar = *Bar;
 			SeekTargetOutPin->BarBeat.Beat = *FloatBeat;
 		}
