[Harmonix] Remove unused fields in FMidiStreamEvent

#tests automated tests

diff --git a/Engine/Plugins/Runtime/Harmonix/Source/HarmonixMetasound/Private/HarmonixMetasound/Nodes/MidiPlayerNode.cpp b/Engine/Plugins/Runtime/Harmonix/Source/HarmonixMetasound/Private/HarmonixMetasound/Nodes/MidiPlayerNode.cpp
--- a/Engine/Plugins/Runtime/Harmonix/Source/HarmonixMetasound/Private/HarmonixMetasound/Nodes/MidiPlayerNode.cpp
+++ b/Engine/Plugins/Runtime/Harmonix/Source/HarmonixMetasound/Private/HarmonixMetasound/Nodes/MidiPlayerNode.cpp
@@ -527,10 +527,8 @@ namespace HarmonixMetasound
 				{
 					FMidiStreamEvent MidiEvent(this, FMidiMsg::CreateAllNotesOff());
 					MidiEvent.BlockSampleFrameIndex = MidiClockOut->GetCurrentBlockFrameIndex();
-					MidiEvent.BlockSampleFrameOffset = 0.0f;
 					MidiEvent.AuthoredMidiTick = MidiClockOut->GetCurrentHiResTick();
 					MidiEvent.CurrentMidiTick = MidiClockOut->GetCurrentHiResTick();
-					MidiEvent.MsOffset = 0.0f;
 					MidiEvent.TrackIndex = 0;
 					MidiOutPin->AddMidiEvent(MidiEvent);
 				}
@@ -764,10 +762,8 @@ namespace HarmonixMetasound
 		{
 			FMidiStreamEvent MidiEvent(this, bKillVoicesOnSeek ? FMidiMsg::CreateAllNotesKill() : FMidiMsg::CreateAllNotesOff());
 			MidiEvent.BlockSampleFrameIndex = MidiClockOut->GetCurrentBlockFrameIndex();
-			MidiEvent.BlockSampleFrameOffset = 0.0f;
 			MidiEvent.AuthoredMidiTick = MidiClockOut->GetCurrentHiResTick();
 			MidiEvent.CurrentMidiTick = MidiClockOut->GetCurrentHiResTick();
-			MidiEvent.MsOffset = 0.0f;
 			MidiEvent.TrackIndex = 0;
 			MidiOutPin->AddMidiEvent(MidiEvent);
 		}
@@ -782,10 +778,8 @@ namespace HarmonixMetasound
 		{
 			FMidiStreamEvent MidiEvent(this, bKillVoicesOnSeek ? FMidiMsg::CreateAllNotesKill() : FMidiMsg::CreateAllNotesOff());
 			MidiEvent.BlockSampleFrameIndex = MidiClockOut->GetCurrentBlockFrameIndex();
-			MidiEvent.BlockSampleFrameOffset = 0.0f;
 			MidiEvent.AuthoredMidiTick = MidiClockOut->GetCurrentHiResTick();
 			MidiEvent.CurrentMidiTick = MidiClockOut->GetCurrentHiResTick();
-			MidiEvent.MsOffset = 0.0f;
 			MidiEvent.TrackIndex = 0;
 			MidiOutPin->AddMidiEvent(MidiEvent);
 		}
@@ -802,10 +796,8 @@ namespace HarmonixMetasound
 		{
 			FMidiStreamEvent MidiEvent(this, FMidiMsg(Status, Data1, Data2));
 			MidiEvent.BlockSampleFrameIndex = MidiClockOut->GetCurrentBlockFrameIndex();
-			MidiEvent.BlockSampleFrameOffset = 0.0f;
 			MidiEvent.AuthoredMidiTick = Tick;
 			MidiEvent.CurrentMidiTick = Tick;
-			MidiEvent.MsOffset = 0.0f;
 			MidiEvent.TrackIndex = TrackIndex;
 			MidiOutPin->AddMidiEvent(MidiEvent);
 		}
@@ -817,10 +809,8 @@ namespace HarmonixMetasound
 		{
 			FMidiStreamEvent MidiEvent(this, FMidiMsg((int32)tempo));
 			MidiEvent.BlockSampleFrameIndex  = MidiClockOut->GetCurrentBlockFrameIndex();
-			MidiEvent.BlockSampleFrameOffset = 0.0f;
 			MidiEvent.AuthoredMidiTick = Tick;
 			MidiEvent.CurrentMidiTick  = Tick;
-			MidiEvent.MsOffset         = 0.0f;
 			MidiEvent.TrackIndex       = TrackIndex;
 			MidiOutPin->AddMidiEvent(MidiEvent);
 		}
@@ -832,10 +822,8 @@ namespace HarmonixMetasound
 		{
 			FMidiStreamEvent MidiEvent(this, FMidiMsg::CreateText(TextIndex, Type));
 			MidiEvent.BlockSampleFrameIndex = MidiClockOut->GetCurrentBlockFrameIndex();
-			MidiEvent.BlockSampleFrameOffset = 0.0f;
 			MidiEvent.AuthoredMidiTick = Tick;
 			MidiEvent.CurrentMidiTick = Tick;
-			MidiEvent.MsOffset = 0.0f;
 			MidiEvent.TrackIndex = TrackIndex;
 			MidiOutPin->AddMidiEvent(MidiEvent);
 		}
@@ -848,10 +836,8 @@ namespace HarmonixMetasound
 		{
 			FMidiStreamEvent MidiEvent(this, FMidiMsg(InStatus, Data1, Data2));
 			MidiEvent.BlockSampleFrameIndex = MidiClockOut->GetCurrentBlockFrameIndex();
-			MidiEvent.BlockSampleFrameOffset = 0.0f;
 			MidiEvent.AuthoredMidiTick = EventTick;
 			MidiEvent.CurrentMidiTick = InCurrentTick;
-			MidiEvent.MsOffset = InPrerollMs;
 			MidiEvent.TrackIndex = TrackIndex;
 			MidiOutPin->AddMidiEvent(MidiEvent);
 		}
@@ -864,10 +850,8 @@ namespace HarmonixMetasound
 
 		FMidiStreamEvent MidiEvent(this, bKillVoicesOnMidiChange ? FMidiMsg::CreateAllNotesKill() : FMidiMsg::CreateAllNotesOff());
 		MidiEvent.BlockSampleFrameIndex = 0;
-		MidiEvent.BlockSampleFrameOffset = 0.0f;
 		MidiEvent.AuthoredMidiTick = 0;
 		MidiEvent.CurrentMidiTick = 0;
-		MidiEvent.MsOffset = 0.0f;
 		MidiEvent.TrackIndex = 0;
 		MidiOutPin->InsertMidiEvent(MidiEvent);
 
diff --git a/Engine/Plugins/Runtime/Harmonix/Source/HarmonixMetasound/Private/HarmonixMetasound/Nodes/MidiStreamLoggerNode.cpp b/Engine/Plugins/Runtime/Harmonix/Source/HarmonixMetasound/Private/HarmonixMetasound/Nodes/MidiStreamLoggerNode.cpp
--- a/Engine/Plugins/Runtime/Harmonix/Source/HarmonixMetasound/Private/HarmonixMetasound/Nodes/MidiStreamLoggerNode.cpp
+++ b/Engine/Plugins/Runtime/Harmonix/Source/HarmonixMetasound/Private/HarmonixMetasound/Nodes/MidiStreamLoggerNode.cpp
@@ -199,13 +199,11 @@ namespace HarmonixMetasound
 
 	void FMidiStreamLoggerOperator::DumpMidiEvent(const FMidiStreamEvent& MidiEvent)
 	{
-		UE_LOG(LogMIDIStreamLogger, Log, TEXT("[%d (%f)] Track: %d, Auth Tick %d, Render Tick %d, Ms Offset %f, Message: %s"),
+		UE_LOG(LogMIDIStreamLogger, Log, TEXT("[%d] Track: %d, Auth Tick %d, Render Tick %d, Message: %s"),
 			MidiEvent.BlockSampleFrameIndex,
-			MidiEvent.BlockSampleFrameOffset,
 			MidiEvent.TrackIndex,
 			MidiEvent.AuthoredMidiTick,
 			MidiEvent.CurrentMidiTick,
-			MidiEvent.MsOffset,
 			*FMidiMsg::ToString(MidiEvent.MidiMessage));
 	}
 }
diff --git a/Engine/Plugins/Runtime/Harmonix/Source/HarmonixMetasound/Private/HarmonixMetasound/Nodes/MidiStreamSelectNode.cpp b/Engine/Plugins/Runtime/Harmonix/Source/HarmonixMetasound/Private/HarmonixMetasound/Nodes/MidiStreamSelectNode.cpp
--- a/Engine/Plugins/Runtime/Harmonix/Source/HarmonixMetasound/Private/HarmonixMetasound/Nodes/MidiStreamSelectNode.cpp
+++ b/Engine/Plugins/Runtime/Harmonix/Source/HarmonixMetasound/Private/HarmonixMetasound/Nodes/MidiStreamSelectNode.cpp
@@ -248,10 +248,8 @@ namespace HarmonixMetasound
 					VoiceId.GetChannelAndNote(MidiCh, MidiNote);
 					FMidiStreamEvent MidiEvent(0u, FMidiMsg::CreateNoteOff(MidiCh, MidiNote)); // Do we need a note? VoiceId should handle the note off
 					MidiEvent.BlockSampleFrameIndex = 0;
-					MidiEvent.BlockSampleFrameOffset = 0.0f;
 					MidiEvent.AuthoredMidiTick = 0;
 					MidiEvent.CurrentMidiTick = 0;
-					MidiEvent.MsOffset = 0.0f;
 					MidiEvent.TrackIndex = 1;
 					MidiEvent.SetVoiceId(VoiceId);
 					MidiStreamOutPin->InsertNoteOffEventOrCancelPendingNoteOn(MidiEvent);
diff --git a/Engine/Plugins/Runtime/Harmonix/Source/HarmonixMetasound/Private/HarmonixMetasound/Nodes/StepSequencePlayerNode.cpp b/Engine/Plugins/Runtime/Harmonix/Source/HarmonixMetasound/Private/HarmonixMetasound/Nodes/StepSequencePlayerNode.cpp
--- a/Engine/Plugins/Runtime/Harmonix/Source/HarmonixMetasound/Private/HarmonixMetasound/Nodes/StepSequencePlayerNode.cpp
+++ b/Engine/Plugins/Runtime/Harmonix/Source/HarmonixMetasound/Private/HarmonixMetasound/Nodes/StepSequencePlayerNode.cpp
@@ -780,10 +780,8 @@ namespace HarmonixMetasound
 						// note off!
 						FMidiStreamEvent MidiEvent(this, FMidiMsg::CreateNoteOff(MidiCh, SequenceTable->Notes[i].NoteNumber + AdditionalOctaveNotes));
 						MidiEvent.BlockSampleFrameIndex = BlockFrameIndex;
-						MidiEvent.BlockSampleFrameOffset = 0.0f;
 						MidiEvent.AuthoredMidiTick = ProcessedThruTick;
 						MidiEvent.CurrentMidiTick = ProcessedThruTick;
-						MidiEvent.MsOffset = 0.0f;
 						MidiEvent.TrackIndex = 1;
 						MidiEvent.SetVoiceId(CurrentCellNotes[i]);
 						MidiOutPin->AddNoteOffEventOrCancelPendingNoteOn(MidiEvent);
@@ -806,10 +804,8 @@ namespace HarmonixMetasound
 							float NoteOnVelocity = FMath::Clamp(static_cast<float>(SequenceTable->Notes[i].Velocity) * CurrentVelocityMultiplierValue, 0.0f, 127.0f);
 							MidiEvent.MidiMessage.SetNoteOnVelocity(static_cast<uint8>(NoteOnVelocity));
 							MidiEvent.BlockSampleFrameIndex = BlockFrameIndex;
-							MidiEvent.BlockSampleFrameOffset = 0.0f;
 							MidiEvent.AuthoredMidiTick = ProcessedThruTick;
 							MidiEvent.CurrentMidiTick = ProcessedThruTick;
-							MidiEvent.MsOffset = 0.0f;
 							MidiEvent.TrackIndex = 1;
 							MidiOutPin->AddMidiEvent(MidiEvent);
 							UE_LOG(LogStepSequencePlayer, Verbose, TEXT("0x%x Note-On %d at %d"), (uint32)(size_t)this, SequenceTable->Notes[i].NoteNumber + AdditionalOctaveNotes, BlockFrameIndex);
@@ -864,10 +860,8 @@ namespace HarmonixMetasound
 					CurrentCellNotes[i].GetChannelAndNote(MidiCh, MidiNote);
 					FMidiStreamEvent MidiEvent(this, FMidiMsg::CreateNoteOff(MidiCh, MidiNote));
 					MidiEvent.BlockSampleFrameIndex  = CurrentBlockSpanStart;
-					MidiEvent.BlockSampleFrameOffset = 0.0f;
 					MidiEvent.AuthoredMidiTick       = 0;
 					MidiEvent.CurrentMidiTick        = 0;
-					MidiEvent.MsOffset               = 0.0f;
 					MidiEvent.TrackIndex             = 1;
 					MidiEvent.SetVoiceId(CurrentCellNotes[i]);
 					UE_LOG(LogStepSequencePlayer, Verbose, TEXT("0x%x Note-Off %d (during resize)"), (uint32)(size_t)this, MidiNote);
@@ -896,10 +890,8 @@ namespace HarmonixMetasound
 				CurrentCellNotes[i].GetChannelAndNote(MidiCh, MidiNote);
 				FMidiStreamEvent MidiEvent(this, FMidiMsg::CreateNoteOff(MidiCh, MidiNote));
 				MidiEvent.BlockSampleFrameIndex  = AtFrameIndex;
-				MidiEvent.BlockSampleFrameOffset = 0.0f;
 				MidiEvent.AuthoredMidiTick       = AbsMidiTick;
 				MidiEvent.CurrentMidiTick        = AbsMidiTick;
-				MidiEvent.MsOffset               = 0.0f;
 				MidiEvent.TrackIndex             = 1;
 				MidiEvent.SetVoiceId(CurrentCellNotes[i]);
 				MidiOutPin->AddMidiEvent(MidiEvent);
diff --git a/Engine/Plugins/Runtime/Harmonix/Source/HarmonixMetasound/Public/HarmonixMetasound/DataTypes/MidiStream.h b/Engine/Plugins/Runtime/Harmonix/Source/HarmonixMetasound/Public/HarmonixMetasound/DataTypes/MidiStream.h
--- a/Engine/Plugins/Runtime/Harmonix/Source/HarmonixMetasound/Public/HarmonixMetasound/DataTypes/MidiStream.h
+++ b/Engine/Plugins/Runtime/Harmonix/Source/HarmonixMetasound/Public/HarmonixMetasound/DataTypes/MidiStream.h
@@ -22,11 +22,9 @@ namespace HarmonixMetasound
 	struct HARMONIXMETASOUND_API FMidiStreamEvent
 	{
 		int32        BlockSampleFrameIndex  = 0;
-		float        BlockSampleFrameOffset = 0.0f;
 		int32        TrackIndex             = 0;
 		int32        AuthoredMidiTick       = 0;
 		int32        CurrentMidiTick        = 0;
-		float        MsOffset               = 0.0f;
 		FMidiMsg     MidiMessage;
 
 		FMidiStreamEvent(const FMidiVoiceGeneratorBase* Owner, const FMidiMsg& Message)
