[HarmonixMetasound] Client Crash at UMusicClockComponent::CalculateSongPosWithOffset
- Nullptr check on ClockDriver

diff --git a/Engine/Plugins/Runtime/Harmonix/Source/HarmonixMetasound/Private/HarmonixMetasound/Components/MusicClockComponent.cpp b/Engine/Plugins/Runtime/Harmonix/Source/HarmonixMetasound/Private/HarmonixMetasound/Components/MusicClockComponent.cpp
--- a/Engine/Plugins/Runtime/Harmonix/Source/HarmonixMetasound/Private/HarmonixMetasound/Components/MusicClockComponent.cpp
+++ b/Engine/Plugins/Runtime/Harmonix/Source/HarmonixMetasound/Private/HarmonixMetasound/Components/MusicClockComponent.cpp
@@ -176,14 +176,15 @@ void UMusicClockComponent::ConnectToWallClock()
 
 FMidiSongPos UMusicClockComponent::CalculateSongPosWithOffset(float MsOffset, ECalibratedMusicTimebase Timebase) const
 {
-	check(ClockDriver);
-
-	EnsureClockIsValidForGameFrame();
-
 	FMidiSongPos Result;
-	if (ClockDriver->CalculateSongPosWithOffset(MsOffset, Timebase, Result))
+	if (ClockDriver != nullptr)
 	{
-		return Result;
+		EnsureClockIsValidForGameFrame();
+
+		if (ClockDriver->CalculateSongPosWithOffset(MsOffset, Timebase, Result))
+		{
+			return Result;
+		}
 	}
 
 	// otherwise, use our song maps copy
