Rebuild midi map when we seek to clear out tempo and time/sig changes

#tests PIE, Creative

diff --git a/Engine/Plugins/Runtime/Harmonix/Source/HarmonixMetasound/Private/HarmonixMetasound/Nodes/MetronomeNode.cpp b/Engine/Plugins/Runtime/Harmonix/Source/HarmonixMetasound/Private/HarmonixMetasound/Nodes/MetronomeNode.cpp
--- a/Engine/Plugins/Runtime/Harmonix/Source/HarmonixMetasound/Private/HarmonixMetasound/Nodes/MetronomeNode.cpp
+++ b/Engine/Plugins/Runtime/Harmonix/Source/HarmonixMetasound/Private/HarmonixMetasound/Nodes/MetronomeNode.cpp
@@ -442,7 +442,9 @@ namespace HarmonixMetasound
 				return EMusicPlayerTransportState::Playing;
 
 			case EMusicPlayerTransportState::Seeking:
+				BuildMidiData();
 				DrivingMidiClock.SeekTo(StartFrameIndex, TransportInPin->GetNextSeekDestination(), SeekPreRollBarsInPin);
+				LastClockTickUpdate = DrivingMidiClock.GetCurrentMidiTick();
 				// Here we will return that we want to be in the same state we were in before this request to 
 				// seek since we can seek "instantaneously"...
 				return GetTransportState();
