[Harmonix] MIDI Stream: reduce copy/merge surface area

* Collapse merge and copy methods down into 1 copy and 2 merge methods
* Move merge voice id remapping into MIDI streams
* Guarantee that a clock (even just a dummy) will be available to downstream nodes
* Guarantee that clocks and transport state are propagated downstream

#tests automated tests, PIE Patchwork

diff --git a/Engine/Plugins/Runtime/Harmonix/Source/HarmonixMetasound/Private/HarmonixMetasound/Analysis/MidiStreamVertexAnalyzer.cpp b/Engine/Plugins/Runtime/Harmonix/Source/HarmonixMetasound/Private/HarmonixMetasound/Analysis/MidiStreamVertexAnalyzer.cpp
--- a/Engine/Plugins/Runtime/Harmonix/Source/HarmonixMetasound/Private/HarmonixMetasound/Analysis/MidiStreamVertexAnalyzer.cpp
+++ b/Engine/Plugins/Runtime/Harmonix/Source/HarmonixMetasound/Private/HarmonixMetasound/Analysis/MidiStreamVertexAnalyzer.cpp
@@ -43,9 +43,9 @@ namespace HarmonixMetasound::Analysis
 		for (const FMidiStreamEvent& Event : MidiStreamIn.GetEventsInBlock())
 		{
 			FMidiEventInfo EventOut;
-			if (const FMidiClockReadRef* Clock = MidiStreamIn.GetMidiClockSource())
+			if (const TSharedPtr<const FMidiClock, ESPMode::NotThreadSafe> Clock = MidiStreamIn.GetClock())
 			{
-				EventOut.Timestamp = (*Clock)->GetBarMap().TickToMusicTimestamp(Event.CurrentMidiTick);
+				EventOut.Timestamp = Clock->GetBarMap().TickToMusicTimestamp(Event.CurrentMidiTick);
 			}
 			EventOut.TrackIndex = Event.TrackIndex;
 			EventOut.MidiMessage = Event.MidiMessage;
diff --git a/Engine/Plugins/Runtime/Harmonix/Source/HarmonixMetasound/Private/HarmonixMetasound/DataTypes/MidiStream.cpp b/Engine/Plugins/Runtime/Harmonix/Source/HarmonixMetasound/Private/HarmonixMetasound/DataTypes/MidiStream.cpp
--- a/Engine/Plugins/Runtime/Harmonix/Source/HarmonixMetasound/Private/HarmonixMetasound/DataTypes/MidiStream.cpp
+++ b/Engine/Plugins/Runtime/Harmonix/Source/HarmonixMetasound/Private/HarmonixMetasound/DataTypes/MidiStream.cpp
@@ -1,11 +1,12 @@
 // Copyright Epic Games, Inc. All Rights Reserved.
+
 #include "HarmonixMetasound/DataTypes/MidiStream.h"
+
+#include "MetasoundDataTypeRegistrationMacro.h"
 #include "Algo/BinarySearch.h"
 #include "Containers/Array.h"
 #include "Containers/UnrealString.h"
 
-#include "MetasoundDataTypeRegistrationMacro.h"
-
 REGISTER_METASOUND_DATATYPE(HarmonixMetasound::FMidiStream, "MIDIStream")
 
 namespace HarmonixMetasound
@@ -14,29 +15,39 @@ namespace HarmonixMetasound
 
 	using namespace Metasound;
 
-	FMidiStream::FMidiStream(const FOperatorSettings& InSettings)
-		: NumFramesPerBlock(InSettings.GetNumFramesPerBlock())
-		, SampleRate(InSettings.GetSampleRate())
+	TSharedPtr<const FMidiClock, ESPMode::NotThreadSafe> FMidiStream::DummyClock{};
+	
+	FMidiStream::FMidiStream(const FOperatorSettings&)
 	{
+		// We have to create the dummy clock here because otherwise registration with the clock update subsystem fails when launching.
+		if (!DummyClock.IsValid())
+		{
+			DummyClock = MakeShared<FMidiClock, ESPMode::NotThreadSafe>(FOperatorSettings{ 48000, 100 });
+		}
+		
+		Clock = DummyClock;
+	}
 
+	void FMidiStream::SetClock(const FMidiClock& InClock)
+	{
+		Clock = InClock.AsWeak();
 	}
 
-	void FMidiStream::PrepareBlock()
+	void FMidiStream::ResetClock()
 	{
-		TransportChangesInBlock.Empty(4);
-		EventsInBlock.Empty(32);
+		Clock = DummyClock;
 	}
 
-	void FMidiStream::AddTransportStateChangeMessage(int32 SampleFrameIndexInBlock, EMusicPlayerTransportState State)
+	TSharedPtr<const FMidiClock, ESPMode::NotThreadSafe> FMidiStream::GetClock() const
 	{
-		if (State == CurrentTransportState.TransportState)
-		{
-			return;
-		}
-		CurrentTransportState.BlockSampleFrameIndex = SampleFrameIndexInBlock;
-		CurrentTransportState.BlockSampleFrameIndex = 0.0f;
-		CurrentTransportState.TransportState        = State;
-		TransportChangesInBlock.Add(CurrentTransportState);
+		const TSharedPtr<const FMidiClock, ESPMode::NotThreadSafe> ClockPtr = Clock.Pin();
+		return ClockPtr.IsValid() ? ClockPtr : DummyClock;
+	}
+
+
+	void FMidiStream::PrepareBlock()
+	{
+		EventsInBlock.Empty(32);
 	}
 
 	void FMidiStream::AddMidiEvent(const FMidiStreamEvent& Event)
@@ -101,282 +112,93 @@ namespace HarmonixMetasound
 		return &(*Repository)[TextIndex];
 	}
 
-	void FMidiStream::Copy(const TArray<FMidiStreamReadRef>& InStreams, bool EventsOnly)
+	void FMidiStream::Copy(const FMidiStream& From, FMidiStream& To, const FEventFilter& Filter, const FEventTransformer& Transformer)
 	{
-		if (!InStreams.Num())
-		{
-			return;
-		}
-
-		if (InStreams.Num() == 1)
-		{
-			Copy(InStreams[0]);
-			return;
-		}
-
-		TArray<FMidiStreamReadRef> SameClockStreams;
-		FilterArrayToStreamsWithTheSameClock(InStreams, SameClockStreams);
-		if (!EventsOnly)
+		// Copy the clock from the other stream
 		{
-			CopyTransportEvents(SameClockStreams);
-		}
-		CopyMidiEvents(SameClockStreams);
-	}
-
-	void FMidiStream::Copy(const FMidiStream& InStream, const FCopyMidiEventsPredicate& Predicate, bool IncludeTransportEvents)
-	{
-		if (IncludeTransportEvents)
-		{
-			CopyTransportEvents(InStream);
-		}
-
-		CopyMidiEvents(InStream, Predicate);
-	}
-
-	void FMidiStream::CopyTransportEvents(const FMidiStream& InStream)
-	{
-		const TArray<FMidiTimestampTransportState>& TransportChanges = InStream.GetTransportChangesInBlock();
-		TransportChangesInBlock.Empty(TransportChanges.Num());
-		for (const FMidiTimestampTransportState& Change : TransportChanges)
-		{
-			AddTransportStateChangeMessage(Change.BlockSampleFrameIndex, Change.TransportState);
-		}
-	}
-
-	void FMidiStream::CopyMidiEvents(const FMidiStream& InStream, const FCopyMidiEventsPredicate& Predicate)
-	{
-		const TArray<FMidiStreamEvent>& MidiEvents = InStream.GetEventsInBlock();
-		if (!ensureAlwaysMsgf(EventsInBlock.Num()==0, TEXT("Existing MIDI events in MIDI stream are being destroyed during copy!")))
-		{
-			EventsInBlock.Empty(FMath::Max(MidiEvents.Num(), 32));
-			ActiveVoices.Reset();
-		}
-		for (const FMidiStreamEvent& Event : MidiEvents)
-		{
-			if (Predicate(Event))
+			const TSharedPtr<const FMidiClock, ESPMode::NotThreadSafe> FromClock = From.GetClock();
+			
+			if (!ensure(FromClock.IsValid()))
 			{
-				AddMidiEvent(Event);
+				return;
 			}
-		}
-	}
-
-	void FMidiStream::Copy(FMidiStreamReadRef InStream, bool EventsOnly)
-	{
-		if (GetMidiClockSource() && GetMidiClockSource()->Get() != InStream->GetMidiClockSource()->Get())
-		{
-			UE_LOG(LogMidiStreamDataType, Warning, TEXT("Cannot copy MIDI events from one stream to another if they have different clocks."));
-			return;
-		}
-		if (EventsOnly)
-		{
-			CopyTransportEvents(InStream);
-		}
-		CopyMidiEvents(InStream);
-	}
 
-	void FMidiStream::FilterArrayToStreamsWithTheSameClock(const TArray<FMidiStreamReadRef>& InStreams, TArray<FMidiStreamReadRef>& OutStreams)
-	{
-		// This has no clock so pass all events
-		if (!GetMidiClockSource())
-		{
-			OutStreams.Append(InStreams);
-			return;
+			To.SetClock(*FromClock);
 		}
 
-		for (const FMidiStreamReadRef& Stream : InStreams)
+		// Reset the target
+		To.EventsInBlock.Reset();
+		To.ActiveVoices.Reset();
+		
+		// Copy the events
+		for (const FMidiStreamEvent& Event : From.GetEventsInBlock())
 		{
-			// Allow the copy if the stream matches or the stream has no clock
-			// Streams with no clock are ok because they won't conflict with our clock
-			// And chances are, that stream has no events
-			if (!Stream->GetMidiClockSource() || Stream->GetMidiClockSource()->Get() == GetMidiClockSource()->Get())
-			{
-				OutStreams.Add(Stream);
-			}
-			else
+			if (Filter(Event))
 			{
-				UE_LOG(LogMidiStreamDataType, Warning, TEXT("Cannot copy MIDI events from one stream to another if they have different clocks."));
+				To.AddMidiEvent(Transformer(Event));
 			}
 		}
 	}
 
-	void FMidiStream::CopyTransportEvents(const TArray<FMidiStreamReadRef>& InStreams)
+	void FMidiStream::Merge(const FMidiStream& From, FMidiStream& To, const FEventFilter& Filter, const FEventTransformer& Transformer)
 	{
-		TransportChangesInBlock.Empty(4);
-
-		using TransportIterator = TArray<FMidiTimestampTransportState>::RangedForConstIteratorType;
-
-		int32 NumInStreams = InStreams.Num();
-		TArray<TransportIterator> WalkingIterators;
-		WalkingIterators.Reserve(NumInStreams);
-		TArray<TransportIterator> IteratorEnds;
-		IteratorEnds.Reserve(NumInStreams);
-
-		for (int32 i = 0; i < NumInStreams; ++i)
-		{
-			WalkingIterators.Add(InStreams[i]->GetTransportChangesInBlock().begin());
-			IteratorEnds.Add(InStreams[i]->GetTransportChangesInBlock().end());
-		}
-
-		while(1)
+		// We need to make sure that the clocks match, or that one of them is using the dummy clock.
+		// Otherwise a merge is invalid.
 		{
-			int32 StreamWithEarliestEvent = -1;
-			int32 EarliestBlockSampleFrameIndex = MAX_int32;
+			const auto FromClock = From.GetClock();
+			const auto ToClock = To.GetClock();
 
-			for (int32 IteratorIndex = 0; IteratorIndex < NumInStreams; ++IteratorIndex)
+			// We shouldn't get a null clock at any time
+			if (!ensure(FromClock.IsValid()) || !ensure(ToClock.IsValid()))
 			{
-				if (WalkingIterators[IteratorIndex] != IteratorEnds[IteratorIndex])
-				{
-					if ((*WalkingIterators[IteratorIndex]).BlockSampleFrameIndex < EarliestBlockSampleFrameIndex)
-					{
-						EarliestBlockSampleFrameIndex = (*WalkingIterators[IteratorIndex]).BlockSampleFrameIndex;
-						StreamWithEarliestEvent = IteratorIndex;
-					}
-				}
+				return;
 			}
-
-			// done?
-			if (StreamWithEarliestEvent == -1)
-			{
-				break;
+			
+			const bool FromClockIsDummy = FromClock.Get() == DummyClock.Get();
+			const bool ToClockIsDummy = ToClock.Get() == DummyClock.Get();
+			const bool ClocksAreSame = FromClock.Get() == ToClock.Get();
+			
+			if (!ClocksAreSame && !FromClockIsDummy && !ToClockIsDummy)
+			{
+				return;
 			}
 
-			AddTransportStateChangeMessage((*WalkingIterators[StreamWithEarliestEvent]).BlockSampleFrameIndex, (*WalkingIterators[StreamWithEarliestEvent]).TransportState);
-
-			++WalkingIterators[StreamWithEarliestEvent];
-		}
-	}
-
-	void FMidiStream::CopyTransportEvents(const FMidiStreamReadRef& InStream)
-	{
-		CopyTransportEvents(*InStream);
-	}
-
-	void FMidiStream::CopyMidiEvents(const TArray<FMidiStreamReadRef>& InStreams)
-	{
-		using MidiStreamIterator = TArray<FMidiStreamEvent>::RangedForConstIteratorType;
-		
-		int32 NumInStreams = InStreams.Num();
-		TArray<MidiStreamIterator> WalkingIterators;
-		WalkingIterators.Reserve(NumInStreams);
-		TArray<MidiStreamIterator> IteratorEnds;
-		IteratorEnds.Reserve(NumInStreams);
-
-		int32 NumTotalEvents = 0;
-		for (int32 i = 0; i < NumInStreams; ++i)
-		{
-			WalkingIterators.Add(InStreams[i]->GetEventsInBlock().begin());
-			IteratorEnds.Add(InStreams[i]->GetEventsInBlock().end());
-			NumTotalEvents += InStreams[i]->GetEventsInBlock().Num();
-
-			// Ensure that the active voices from the source streams are in this stream
-			ActiveVoices.Append(InStreams[i]->ActiveVoices);
-		}
-
-		EventsInBlock.Empty(FMath::Max(32, NumTotalEvents));
-
-		while (1)
-		{
-			int32 StreamWithEarliestEvent = -1;
-			int32 EarliestBlockSampleIndex = MAX_int32;
-
-			for (int32 IteratorIndex = 0; IteratorIndex < NumInStreams; ++IteratorIndex)
+			// If the "to" clock is the dummy, and the "from" clock isn't, overwrite the "to" clock
+			if (FromClockIsDummy && !ToClockIsDummy)
 			{
-				if (WalkingIterators[IteratorIndex] != IteratorEnds[IteratorIndex])
-				{
-					if ((*WalkingIterators[IteratorIndex]).BlockSampleFrameIndex == EarliestBlockSampleIndex)
-					{
-						if ((*WalkingIterators[StreamWithEarliestEvent]).MidiMessage.IsNoteOn() && (*WalkingIterators[IteratorIndex]).MidiMessage.IsNoteOff())
-						{
-							StreamWithEarliestEvent = IteratorIndex;
-						}
-					}
-					else if ((*WalkingIterators[IteratorIndex]).BlockSampleFrameIndex < EarliestBlockSampleIndex)
-					{
-						EarliestBlockSampleIndex = (*WalkingIterators[IteratorIndex]).BlockSampleFrameIndex;
-						StreamWithEarliestEvent = IteratorIndex;
-					}
-				}
-			}
-
-			// done?
-			if (StreamWithEarliestEvent == -1)
-			{
-				break;
+				To.SetClock(*FromClock);
 			}
-
-			AddMidiEvent(*WalkingIterators[StreamWithEarliestEvent]);
-			++WalkingIterators[StreamWithEarliestEvent];
-		}
-	}
-
-	void FMidiStream::MergeMidiEvents(const TArray<FMidiStreamReadRef>& InStreams)
-	{
-		using MidiStreamIterator = TArray<FMidiStreamEvent>::RangedForConstIteratorType;
-
-		int32 NumInStreams = InStreams.Num();
-		TArray<MidiStreamIterator> WalkingIterators;
-		WalkingIterators.Reserve(NumInStreams);
-		TArray<MidiStreamIterator> IteratorEnds;
-		IteratorEnds.Reserve(NumInStreams);
-
-		int32 NumTotalEvents = 0;
-		for (int32 i = 0; i < NumInStreams; ++i)
-		{
-			WalkingIterators.Add(InStreams[i]->GetEventsInBlock().begin());
-			IteratorEnds.Add(InStreams[i]->GetEventsInBlock().end());
-			NumTotalEvents += InStreams[i]->GetEventsInBlock().Num();
 		}
 
-		EventsInBlock.Reserve(FMath::Max(32, NumTotalEvents + EventsInBlock.Num()));
-
-		while (1)
+		// Insert the events
+		for (const FMidiStreamEvent& Event : From.GetEventsInBlock())
 		{
-			int32 StreamWithEarliestEvent = -1;
-			int32 EarliestBlockSampleIndex = MAX_int32;
-
-			for (int32 IteratorIndex = 0; IteratorIndex < NumInStreams; ++IteratorIndex)
+			if (Filter(Event))
 			{
-				if (WalkingIterators[IteratorIndex] != IteratorEnds[IteratorIndex])
-				{
-					if ((*WalkingIterators[IteratorIndex]).BlockSampleFrameIndex < EarliestBlockSampleIndex)
-					{
-						EarliestBlockSampleIndex = (*WalkingIterators[IteratorIndex]).BlockSampleFrameIndex;
-						StreamWithEarliestEvent = IteratorIndex;
-					}
-				}
+				To.InsertMidiEvent(Transformer(Event));
 			}
-
-			// done?
-			if (StreamWithEarliestEvent == -1)
-			{
-				break;
-			}
-
-			InsertMidiEvent(*WalkingIterators[StreamWithEarliestEvent]);
-			++WalkingIterators[StreamWithEarliestEvent];
 		}
 	}
-
-	void FMidiStream::CopyMidiEvents(const FMidiStreamReadRef& InStream)
+	
+	void FMidiStream::Merge(
+		const FMidiStream& FromA,
+		const FMidiStream& FromB,
+		FMidiStream& To,
+		const FEventFilter& Filter,
+		const FEventTransformer& Transformer)
 	{
-		const TArray<FMidiStreamEvent>& MidiEvents = InStream->GetEventsInBlock();
-		EventsInBlock.Empty(FMath::Max(MidiEvents.Num(), 32));
-		// Ensure that the active voices from the source match this stream
-		ActiveVoices = InStream->ActiveVoices;
-		for (auto& Event : MidiEvents)
-		{
-			AddMidiEvent(Event);
-		}
-	}
+		// Merge in stream A
+		Merge(FromA, To, Filter, Transformer);
 
-	void FMidiStream::MergeMidiEvents(const FMidiStreamReadRef& InStream)
-	{
-		const TArray<FMidiStreamEvent>& NewMidiEvents = InStream->GetEventsInBlock();
-		EventsInBlock.Reserve(FMath::Max(32, EventsInBlock.Num() + NewMidiEvents.Num()));
-		for (auto& Event : NewMidiEvents)
+		// Merge in stream B and re-map the voice ids
+		const auto RemapTransform = [&To, &Transformer](const FMidiStreamEvent& Event)
 		{
-			InsertMidiEvent(Event);
-		}
+			FMidiStreamEvent TransformedEvent = Transformer(Event);
+			const auto GeneratorId = TransformedEvent.GetVoiceId().GetGeneratorId();
+			TransformedEvent.ReassignOwner(&To.GeneratorMap.FindOrAdd(GeneratorId));
+			return TransformedEvent;
+		};
+		Merge(FromB, To, Filter, RemapTransform);
 	}
 
 	void FMidiStream::UpdateActiveVoice(const FMidiStreamEvent& Event)
diff --git a/Engine/Plugins/Runtime/Harmonix/Source/HarmonixMetasound/Private/HarmonixMetasound/MidiOps/MidiChannelFilter.cpp b/Engine/Plugins/Runtime/Harmonix/Source/HarmonixMetasound/Private/HarmonixMetasound/MidiOps/MidiChannelFilter.cpp
--- a/Engine/Plugins/Runtime/Harmonix/Source/HarmonixMetasound/Private/HarmonixMetasound/MidiOps/MidiChannelFilter.cpp
+++ b/Engine/Plugins/Runtime/Harmonix/Source/HarmonixMetasound/Private/HarmonixMetasound/MidiOps/MidiChannelFilter.cpp
@@ -20,10 +20,7 @@ namespace Harmonix::Midi::Ops
 		};
 
 		// Copy events which pass the filter to the output
-		OutStream.Copy(
-			InStream,
-			Filter,
-			true /* include transport events */);
+		HarmonixMetasound::FMidiStream::Copy(InStream, OutStream, Filter);
 
 		// Unstick notes if necessary
 		StuckNoteGuard.Process(InStream, OutStream, Filter);
diff --git a/Engine/Plugins/Runtime/Harmonix/Source/HarmonixMetasound/Private/HarmonixMetasound/MidiOps/MidiNoteFilter.cpp b/Engine/Plugins/Runtime/Harmonix/Source/HarmonixMetasound/Private/HarmonixMetasound/MidiOps/MidiNoteFilter.cpp
--- a/Engine/Plugins/Runtime/Harmonix/Source/HarmonixMetasound/Private/HarmonixMetasound/MidiOps/MidiNoteFilter.cpp
+++ b/Engine/Plugins/Runtime/Harmonix/Source/HarmonixMetasound/Private/HarmonixMetasound/MidiOps/MidiNoteFilter.cpp
@@ -36,10 +36,7 @@ namespace Harmonix::Midi::Ops
 		};
 
 		// Copy events which pass the filter to the output
-		OutStream.Copy(
-			InStream,
-			Filter,
-			true /* include transport events */);
+		HarmonixMetasound::FMidiStream::Copy(InStream, OutStream, Filter);
 
 		// Unstick notes if necessary
 		StuckNoteGuard.Process(InStream, OutStream, Filter);
diff --git a/Engine/Plugins/Runtime/Harmonix/Source/HarmonixMetasound/Private/HarmonixMetasound/MidiOps/MidiTrackFilter.cpp b/Engine/Plugins/Runtime/Harmonix/Source/HarmonixMetasound/Private/HarmonixMetasound/MidiOps/MidiTrackFilter.cpp
--- a/Engine/Plugins/Runtime/Harmonix/Source/HarmonixMetasound/Private/HarmonixMetasound/MidiOps/MidiTrackFilter.cpp
+++ b/Engine/Plugins/Runtime/Harmonix/Source/HarmonixMetasound/Private/HarmonixMetasound/MidiOps/MidiTrackFilter.cpp
@@ -25,10 +25,7 @@ namespace Harmonix::Midi::Ops
 		};
 		
 		// Copy events which pass the filter to the output
-		OutStream.Copy(
-			InStream,
-			Filter,
-			true /* include transport events */);
+		HarmonixMetasound::FMidiStream::Copy(InStream, OutStream, Filter);
 
 		// Unstick notes if necessary
 		StuckNoteGuard.Process(InStream, OutStream, Filter);
diff --git a/Engine/Plugins/Runtime/Harmonix/Source/HarmonixMetasound/Private/HarmonixMetasound/Nodes/FusionSamplerNode.cpp b/Engine/Plugins/Runtime/Harmonix/Source/HarmonixMetasound/Private/HarmonixMetasound/Nodes/FusionSamplerNode.cpp
--- a/Engine/Plugins/Runtime/Harmonix/Source/HarmonixMetasound/Private/HarmonixMetasound/Nodes/FusionSamplerNode.cpp
+++ b/Engine/Plugins/Runtime/Harmonix/Source/HarmonixMetasound/Private/HarmonixMetasound/Nodes/FusionSamplerNode.cpp
@@ -613,7 +613,7 @@ namespace HarmonixMetasound
 		auto MidiEventIterator = MidiEvents.begin();
 
 		// create an iterator for the midi clock 
-		const FMidiClockReadRef* MidiClock = MidiStreamInPin->GetMidiClockSource();
+		const TSharedPtr<const FMidiClock, ESPMode::NotThreadSafe> MidiClock = MidiStreamInPin->GetClock();
 		
 		while (FramesRequired > 0)
 		{
@@ -649,13 +649,13 @@ namespace HarmonixMetasound
 				}
 			}
 
-			if (MidiClock)
+			if (MidiClock.IsValid())
 			{
-				float ClockSpeed = (*MidiClock)->GetSpeedAtBlockSampleFrame(CurrentBlockFrameIndex);
+				const float ClockSpeed = MidiClock->GetSpeedAtBlockSampleFrame(CurrentBlockFrameIndex);
 				SetSpeed(ClockSpeed, !(*ClockSpeedAffectsPitchInPin));
-				float ClockTempo = (*MidiClock)->GetTempoAtBlockSampleFrame(CurrentBlockFrameIndex);
+				const float ClockTempo = MidiClock->GetTempoAtBlockSampleFrame(CurrentBlockFrameIndex);
 				SetTempo(ClockTempo);
-				float Beat = (*MidiClock)->GetQuarterNoteIncludingCountIn();
+				const float Beat = MidiClock->GetQuarterNoteIncludingCountIn();
 				SetBeat(Beat);
 			}
 
diff --git a/Engine/Plugins/Runtime/Harmonix/Source/HarmonixMetasound/Private/HarmonixMetasound/Nodes/MidiNoteTriggerNode.cpp b/Engine/Plugins/Runtime/Harmonix/Source/HarmonixMetasound/Private/HarmonixMetasound/Nodes/MidiNoteTriggerNode.cpp
--- a/Engine/Plugins/Runtime/Harmonix/Source/HarmonixMetasound/Private/HarmonixMetasound/Nodes/MidiNoteTriggerNode.cpp
+++ b/Engine/Plugins/Runtime/Harmonix/Source/HarmonixMetasound/Private/HarmonixMetasound/Nodes/MidiNoteTriggerNode.cpp
@@ -153,15 +153,13 @@ namespace HarmonixMetasound::Nodes::MidiNoteTriggerNode
 				return;
 			}
 
-			const TArray<FMidiTimestampTransportState> TransportChanges = MidiStreamInPin->GetTransportChangesInBlock();
-			if (!TransportChanges.IsEmpty())
+			// Note off if the transport has stopped and we have a sounding note
+			if (SoundingNote >= 0)
 			{
-				if (TransportChanges.Last().TransportState != EMusicPlayerTransportState::Playing)
+				const TSharedPtr<const FMidiClock, ESPMode::NotThreadSafe> Clock = MidiStreamInPin->GetClock();
+				if (Clock.IsValid() && Clock->GetTransportStateAtEndOfBlock() != EMusicPlayerTransportState::Playing)
 				{
-					if (SoundingNote >= 0)
-					{
-						TriggerNoteOff(0, SoundingNote);
-					}
+					TriggerNoteOff(0, SoundingNote);
 					return;
 				}
 			}
diff --git a/Engine/Plugins/Runtime/Harmonix/Source/HarmonixMetasound/Private/HarmonixMetasound/Nodes/MidiPlayerNode.cpp b/Engine/Plugins/Runtime/Harmonix/Source/HarmonixMetasound/Private/HarmonixMetasound/Nodes/MidiPlayerNode.cpp
--- a/Engine/Plugins/Runtime/Harmonix/Source/HarmonixMetasound/Private/HarmonixMetasound/Nodes/MidiPlayerNode.cpp
+++ b/Engine/Plugins/Runtime/Harmonix/Source/HarmonixMetasound/Private/HarmonixMetasound/Nodes/MidiPlayerNode.cpp
@@ -260,8 +260,7 @@ namespace HarmonixMetasound
 		, BlockSize(InSettings.GetNumFramesPerBlock())
 	{
 		MidiClockOut->RegisterHiResPlayCursor(this);
-		MidiOutPin->SetClockSource(MidiClockOut);
-		MidiOutPin->AddTransportStateChangeMessage(0, EMusicPlayerTransportState::Prepared);
+		MidiOutPin->SetClock(*MidiClockOut);
 	}
 
 	FExternallyClockedMidiPlayerOperator::FExternallyClockedMidiPlayerOperator(const FOperatorSettings& InSettings,
@@ -337,14 +336,12 @@ namespace HarmonixMetasound
 			case EMusicPlayerTransportState::Stopping:
 			case EMusicPlayerTransportState::Killing:
 				MidiClockOut->AddTransportStateChangeToBlock({ 0, 0.0f, EMusicPlayerTransportState::Prepared });
-				MidiOutPin->AddTransportStateChangeMessage(0, EMusicPlayerTransportState::Prepared);
 				return EMusicPlayerTransportState::Prepared;
 
 			case EMusicPlayerTransportState::Starting:
 			case EMusicPlayerTransportState::Playing:
 			case EMusicPlayerTransportState::Continuing:
 				MidiClockOut->ResetAndStart(0, true);
-				MidiOutPin->AddTransportStateChangeMessage(0, EMusicPlayerTransportState::Playing);
 				return EMusicPlayerTransportState::Playing;
 
 			case EMusicPlayerTransportState::Seeking: // seeking is omitted from init, shouldn't happen
@@ -354,7 +351,6 @@ namespace HarmonixMetasound
 			case EMusicPlayerTransportState::Pausing:
 			case EMusicPlayerTransportState::Paused:
 				MidiClockOut->AddTransportStateChangeToBlock({ 0, 0.0f, EMusicPlayerTransportState::Paused });
-				MidiOutPin->AddTransportStateChangeMessage(0, EMusicPlayerTransportState::Paused);
 				return EMusicPlayerTransportState::Paused;
 
 			default:
@@ -492,16 +488,12 @@ namespace HarmonixMetasound
 			{
 			case EMusicPlayerTransportState::Invalid:
 			case EMusicPlayerTransportState::Preparing:
-				MidiOutPin->AddTransportStateChangeMessage(StartFrameIndex, EMusicPlayerTransportState::Prepared);
 				return EMusicPlayerTransportState::Prepared;
 
 			case EMusicPlayerTransportState::Prepared:
-				MidiOutPin->AddTransportStateChangeMessage(StartFrameIndex, CurrentState);
 				return CurrentState;
 
 			case EMusicPlayerTransportState::Starting:
-				//SampleCounter.SetNumSamples(0);
-				MidiOutPin->AddTransportStateChangeMessage(StartFrameIndex, EMusicPlayerTransportState::Playing);
 				return EMusicPlayerTransportState::Playing;
 
 			case EMusicPlayerTransportState::Playing:
@@ -512,18 +504,15 @@ namespace HarmonixMetasound
 				return GetTransportState();
 
 			case EMusicPlayerTransportState::Continuing:
-				MidiOutPin->AddTransportStateChangeMessage(StartFrameIndex, EMusicPlayerTransportState::Continuing);
 				return EMusicPlayerTransportState::Playing;
 
 			case EMusicPlayerTransportState::Pausing:
-				MidiOutPin->AddTransportStateChangeMessage(StartFrameIndex, EMusicPlayerTransportState::Pausing);
 				return EMusicPlayerTransportState::Paused;
 
 			case EMusicPlayerTransportState::Paused:
 				return EMusicPlayerTransportState::Paused;
 
 			case EMusicPlayerTransportState::Stopping:
-				MidiOutPin->AddTransportStateChangeMessage(StartFrameIndex, EMusicPlayerTransportState::Stopping);
 				{
 					FMidiStreamEvent MidiEvent(this, FMidiMsg::CreateAllNotesOff());
 					MidiEvent.BlockSampleFrameIndex = MidiClockOut->GetCurrentBlockFrameIndex();
@@ -535,7 +524,6 @@ namespace HarmonixMetasound
 				return EMusicPlayerTransportState::Prepared;
 
 			case EMusicPlayerTransportState::Killing:
-				MidiOutPin->AddTransportStateChangeMessage(StartFrameIndex, EMusicPlayerTransportState::Killing);
 				return EMusicPlayerTransportState::Prepared;
 
 			default:
@@ -559,13 +547,11 @@ namespace HarmonixMetasound
 			case EMusicPlayerTransportState::Prepared:
 			case EMusicPlayerTransportState::Stopping:
 			case EMusicPlayerTransportState::Killing:
-				MidiOutPin->AddTransportStateChangeMessage(0, EMusicPlayerTransportState::Prepared);
 				return EMusicPlayerTransportState::Prepared;
 
 			case EMusicPlayerTransportState::Starting:
 			case EMusicPlayerTransportState::Playing:
 			case EMusicPlayerTransportState::Continuing:
-				MidiOutPin->AddTransportStateChangeMessage(0, EMusicPlayerTransportState::Playing);
 				return EMusicPlayerTransportState::Playing;
 
 			case EMusicPlayerTransportState::Seeking: // seeking is omitted from init, shouldn't happen
@@ -574,7 +560,6 @@ namespace HarmonixMetasound
 
 			case EMusicPlayerTransportState::Pausing:
 			case EMusicPlayerTransportState::Paused:
-				MidiOutPin->AddTransportStateChangeMessage(0, EMusicPlayerTransportState::Paused);
 				return EMusicPlayerTransportState::Paused;
 
 			default:
@@ -637,12 +622,10 @@ namespace HarmonixMetasound
 				MidiClockOut->AddTransportStateChangeToBlock({ StartFrameIndex, 0.0f, EMusicPlayerTransportState::Prepared });
 
 				// midi out
-				MidiOutPin->AddTransportStateChangeMessage(StartFrameIndex, EMusicPlayerTransportState::Prepared);
 				return EMusicPlayerTransportState::Prepared;
 
 			case EMusicPlayerTransportState::Prepared:
 				// midi out
-				MidiOutPin->AddTransportStateChangeMessage(StartFrameIndex, CurrentState);
 				return CurrentState;
 
 			case EMusicPlayerTransportState::Starting:
@@ -650,7 +633,6 @@ namespace HarmonixMetasound
 				MidiClockOut->ResetAndStart(StartFrameIndex, !ReceivedSeekWhileStopped());
 
 				// midi out
-				MidiOutPin->AddTransportStateChangeMessage(StartFrameIndex, EMusicPlayerTransportState::Playing);
 				return EMusicPlayerTransportState::Playing;
 
 			case EMusicPlayerTransportState::Playing:
@@ -668,7 +650,6 @@ namespace HarmonixMetasound
 				MidiClockOut->AddTransportStateChangeToBlock({ StartFrameIndex, 0.0f, EMusicPlayerTransportState::Playing });
 
 				// midi out
-				MidiOutPin->AddTransportStateChangeMessage(StartFrameIndex, EMusicPlayerTransportState::Continuing);
 				return EMusicPlayerTransportState::Playing;
 
 			case EMusicPlayerTransportState::Pausing:
@@ -687,7 +668,6 @@ namespace HarmonixMetasound
 				MidiClockOut->AddTransportStateChangeToBlock({ StartFrameIndex, 0.0f, EMusicPlayerTransportState::Prepared });
 
 				// midi out
-				MidiOutPin->AddTransportStateChangeMessage(StartFrameIndex, EMusicPlayerTransportState::Stopping);
 				return EMusicPlayerTransportState::Prepared;
 
 			case EMusicPlayerTransportState::Killing:
@@ -695,7 +675,6 @@ namespace HarmonixMetasound
 				MidiClockOut->AddTransportStateChangeToBlock({ StartFrameIndex, 0.0f, EMusicPlayerTransportState::Prepared });
 
 				// midi out
-				MidiOutPin->AddTransportStateChangeMessage(StartFrameIndex, EMusicPlayerTransportState::Killing);
 				return EMusicPlayerTransportState::Prepared;
 
 			default:
@@ -736,7 +715,7 @@ namespace HarmonixMetasound
 		BlockSize = Params.OperatorSettings.GetNumFramesPerBlock();
 		CurrentBlockSpanStart = 0;
 
-		MidiOutPin->SetClockSource(MidiClockOut);
+		MidiOutPin->SetClock(*MidiClockOut);
 		MidiClockOut->ResetAndStart(0, true);
 
 		NeedsTransportInit = true;
diff --git a/Engine/Plugins/Runtime/Harmonix/Source/HarmonixMetasound/Private/HarmonixMetasound/Nodes/MidiStreamLoggerNode.cpp b/Engine/Plugins/Runtime/Harmonix/Source/HarmonixMetasound/Private/HarmonixMetasound/Nodes/MidiStreamLoggerNode.cpp
--- a/Engine/Plugins/Runtime/Harmonix/Source/HarmonixMetasound/Private/HarmonixMetasound/Nodes/MidiStreamLoggerNode.cpp
+++ b/Engine/Plugins/Runtime/Harmonix/Source/HarmonixMetasound/Private/HarmonixMetasound/Nodes/MidiStreamLoggerNode.cpp
@@ -46,7 +46,6 @@ namespace HarmonixMetasound
 		FMidiStreamReadRef MidiStreamInPin;
 
 		//** DATA
-		void DumpTransportEvent(const FMidiTimestampTransportState& TransportEvent);
 		void DumpMidiEvent(const FMidiStreamEvent& MidiEvent);
 	};
 
@@ -160,43 +159,12 @@ namespace HarmonixMetasound
 			return;
 		}
 
-		const TArray<FMidiTimestampTransportState>& TransportEvents = MidiStreamInPin->GetTransportChangesInBlock();
-		const TArray<FMidiStreamEvent>& MidiEvents = MidiStreamInPin->GetEventsInBlock();
-		auto TransportEventIterator = TransportEvents.begin();
-		auto MidiEventIterator = MidiEvents.begin();
-		while (TransportEventIterator != TransportEvents.end() || MidiEventIterator != MidiEvents.end())
+		for (const FMidiStreamEvent& Event : MidiStreamInPin->GetEventsInBlock())
 		{
-			if (TransportEventIterator != TransportEvents.end() && MidiEventIterator != MidiEvents.end())
-			{
-				if ((*TransportEventIterator).BlockSampleFrameIndex <= (*MidiEventIterator).BlockSampleFrameIndex)
-				{
-					DumpTransportEvent(*TransportEventIterator);
-					++TransportEventIterator;
-				}
-				else
-				{
-					DumpMidiEvent(*MidiEventIterator);
-					++MidiEventIterator;
-				}
-			}
-			else if (TransportEventIterator != TransportEvents.end())
-			{
-				DumpTransportEvent(*TransportEventIterator);
-				++TransportEventIterator;
-			}
-			else
-			{
-				DumpMidiEvent(*MidiEventIterator);
-				++MidiEventIterator;
-			}
+			DumpMidiEvent(Event);
 		}
 	}
 
-	void FMidiStreamLoggerOperator::DumpTransportEvent(const FMidiTimestampTransportState& TransportEvent)
-	{
-		UE_LOG(LogMIDIStreamLogger, Log, TEXT("[%d (%f)] Transport: %s"), TransportEvent.BlockSampleFrameIndex, TransportEvent.BlockSampleFrameOffset, *FMusicTransportControllable::StateToString(TransportEvent.TransportState));
-	}
-
 	void FMidiStreamLoggerOperator::DumpMidiEvent(const FMidiStreamEvent& MidiEvent)
 	{
 		UE_LOG(LogMIDIStreamLogger, Log, TEXT("[%d] Track: %d, Auth Tick %d, Render Tick %d, Message: %s"),
diff --git a/Engine/Plugins/Runtime/Harmonix/Source/HarmonixMetasound/Private/HarmonixMetasound/Nodes/MidiStreamMergeNode.cpp b/Engine/Plugins/Runtime/Harmonix/Source/HarmonixMetasound/Private/HarmonixMetasound/Nodes/MidiStreamMergeNode.cpp
--- a/Engine/Plugins/Runtime/Harmonix/Source/HarmonixMetasound/Private/HarmonixMetasound/Nodes/MidiStreamMergeNode.cpp
+++ b/Engine/Plugins/Runtime/Harmonix/Source/HarmonixMetasound/Private/HarmonixMetasound/Nodes/MidiStreamMergeNode.cpp
@@ -43,11 +43,6 @@ namespace HarmonixMetasound
 
 		//** OUTPUTS
 		FMidiStreamWriteRef MidiStreamOutPin;
-
-		//** DATA
-		TMap<uint32, FMidiVoiceGeneratorBase> GeneratorMap;
-		
-		bool bNeedsClockRefresh { true };
 	};
 
 	class FMidiStreamMergeNode : public FNodeFacade
@@ -136,8 +131,6 @@ namespace HarmonixMetasound
 
 		InVertexData.BindReadVertex(METASOUND_GET_PARAM_NAME(InputMidiStreamA), MidiStreamAInPin);
 		InVertexData.BindReadVertex(METASOUND_GET_PARAM_NAME(InputMidiStreamB), MidiStreamBInPin);
-
-		bNeedsClockRefresh = true;
 	}
 
 	void FMidiStreamMergeOperator::BindOutputs(FOutputVertexInterfaceData& InVertexData)
@@ -145,47 +138,16 @@ namespace HarmonixMetasound
 		using namespace CommonPinNames;
 
 		InVertexData.BindReadVertex(METASOUND_GET_PARAM_NAME(Outputs::MidiStream), MidiStreamOutPin);
-
-		bNeedsClockRefresh = true;
 	}
 
 	void FMidiStreamMergeOperator::Reset(const FResetParams&)
 	{
-		bNeedsClockRefresh = true;
 	}
 
 	void FMidiStreamMergeOperator::Execute()
 	{
-		// Refresh the clock for the output if inputs may have changed
-		if (bNeedsClockRefresh)
-		{
-			MidiStreamOutPin->ResetClockSource();
-			
-			// If there is a clock on the first stream, use that
-			if (const FMidiClockReadRef* ClockA = MidiStreamAInPin->GetMidiClockSource())
-			{
-				MidiStreamOutPin->SetClockSource(*ClockA);
-			}
-			// If not, and there's a clock on the second stream, use that
-			else if (const FMidiClockReadRef* ClockB = MidiStreamBInPin->GetMidiClockSource())
-			{
-				MidiStreamOutPin->SetClockSource(*ClockB);
-			}
-
-			bNeedsClockRefresh = false;
-		}
-		
 		MidiStreamOutPin->PrepareBlock();
-		TArray<FMidiStreamReadRef> SameClockStreams;
-		MidiStreamOutPin->FilterArrayToStreamsWithTheSameClock({MidiStreamAInPin, MidiStreamBInPin}, SameClockStreams);
-		MidiStreamOutPin->CopyTransportEvents(SameClockStreams);
-		MidiStreamOutPin->MergeMidiEvents(MidiStreamAInPin);
-		for (FMidiStreamEvent Event : MidiStreamBInPin->GetEventsInBlock())
-		{
-			uint32 GenId = Event.GetVoiceId().GetGeneratorId();
-			Event.ReassignOwner(&GeneratorMap.FindOrAdd(GenId));
-			MidiStreamOutPin->InsertMidiEvent(Event);
-		}
+		FMidiStream::Merge(*MidiStreamAInPin, *MidiStreamBInPin, *MidiStreamOutPin);
 		
 	}
 }
diff --git a/Engine/Plugins/Runtime/Harmonix/Source/HarmonixMetasound/Private/HarmonixMetasound/Nodes/MidiStreamSelectNode.cpp b/Engine/Plugins/Runtime/Harmonix/Source/HarmonixMetasound/Private/HarmonixMetasound/Nodes/MidiStreamSelectNode.cpp
--- a/Engine/Plugins/Runtime/Harmonix/Source/HarmonixMetasound/Private/HarmonixMetasound/Nodes/MidiStreamSelectNode.cpp
+++ b/Engine/Plugins/Runtime/Harmonix/Source/HarmonixMetasound/Private/HarmonixMetasound/Nodes/MidiStreamSelectNode.cpp
@@ -210,31 +210,30 @@ namespace HarmonixMetasound
 
 	void FMidiStreamSelectOperator::CopyFromSelectedStream(int32 SelectedStreamIndex)
 	{
-		FMidiStreamReadRef SourceStream = SelectedStreamIndex == 0 ? MidiStreamAInPin : MidiStreamBInPin;
+		const FMidiStreamReadRef SourceStream = SelectedStreamIndex == 0 ? MidiStreamAInPin : MidiStreamBInPin;
 
 		// Copy, but keep track of 'note-on' messages so we can appropriately send 'note-off' messages later...
-		MidiStreamOutPin->Copy(SourceStream,
-			[&](const FMidiStreamEvent& MidiEvent)
+		FMidiStream::Copy(
+			*SourceStream,
+			*MidiStreamOutPin,
+			[this](const FMidiStreamEvent& Event)
 			{
-				if (MidiEvent.MidiMessage.IsStd())
+				if (Event.MidiMessage.IsNoteOn())
 				{
-					if (MidiEvent.MidiMessage.IsNoteOn())
-					{
-						PlayingVoices.Add(MidiEvent.GetVoiceId());
-					}
-					else
-					{
-						PlayingVoices.Remove(MidiEvent.GetVoiceId());
-					}
+					PlayingVoices.Add(Event.GetVoiceId());
+				}
+				else if (Event.MidiMessage.IsNoteOff())
+				{
+					PlayingVoices.Remove(Event.GetVoiceId());
 				}
+				
 				return true;
-			},
-			true);
+			});
 	}
 
 	void FMidiStreamSelectOperator::MergeOrCreateNoteOffsForDeselectedStream(int32 SelectedStreamIndex)
 	{
-		FMidiStreamReadRef InActiveStream = SelectedStreamIndex == 0 ? MidiStreamBInPin : MidiStreamAInPin;
+		FMidiStreamReadRef InactiveStream = SelectedStreamIndex == 0 ? MidiStreamBInPin : MidiStreamAInPin;
 
 		if (*ImmediateNoteOffInPin)
 		{
@@ -260,13 +259,17 @@ namespace HarmonixMetasound
 		else if (*CopyInactiveNoteOffsInPin)
 		{
 			// Merge note offs from inactive midi stream
-			MidiStreamOutPin->MergeMidiEvents(InActiveStream, [&](const FMidiStreamEvent& MidiEvent)
+			FMidiStream::Merge(
+				*InactiveStream,
+				*MidiStreamOutPin,
+				[this](const FMidiStreamEvent& Event)
 				{
-					if (MidiEvent.MidiMessage.IsStd() && MidiEvent.MidiMessage.IsNoteOff() && PlayingVoices.Contains(MidiEvent.GetVoiceId()))
+					if (Event.MidiMessage.IsNoteOff() && PlayingVoices.Contains(Event.GetVoiceId()))
 					{
-						PlayingVoices.Remove(MidiEvent.GetVoiceId());
+						PlayingVoices.Remove(Event.GetVoiceId());
 						return true;
 					}
+
 					return false;
 				});
 		}
@@ -278,16 +281,7 @@ namespace HarmonixMetasound
 		PlayingVoices.Empty();
 
 		MidiStreamOutPin->PrepareBlock();
-		
-		if (MidiStreamAInPin.Get() && MidiStreamAInPin->GetMidiClockSource())
-		{
-			MidiStreamOutPin->SetClockSource(*(MidiStreamAInPin->GetMidiClockSource()));
-		}
-		else if (MidiStreamBInPin.Get() && MidiStreamBInPin->GetMidiClockSource())
-		{
-			// Try the B stream if the A stream isn't hooked up
-			MidiStreamOutPin->SetClockSource(*(MidiStreamBInPin->GetMidiClockSource()));
-		}
+		MidiStreamOutPin->ResetClock();
 	}
 }
 
diff --git a/Engine/Plugins/Runtime/Harmonix/Source/HarmonixMetasound/Private/HarmonixMetasound/Nodes/MidiStreamTrackFilterNode.cpp b/Engine/Plugins/Runtime/Harmonix/Source/HarmonixMetasound/Private/HarmonixMetasound/Nodes/MidiStreamTrackFilterNode.cpp
--- a/Engine/Plugins/Runtime/Harmonix/Source/HarmonixMetasound/Private/HarmonixMetasound/Nodes/MidiStreamTrackFilterNode.cpp
+++ b/Engine/Plugins/Runtime/Harmonix/Source/HarmonixMetasound/Private/HarmonixMetasound/Nodes/MidiStreamTrackFilterNode.cpp
@@ -143,33 +143,19 @@ namespace HarmonixMetasound::Nodes::MidiTrackFilter
 			InVertexData.BindReadVertex(Inputs::MinTrackIndexName, Inputs.MinTrackIndex);
 			InVertexData.BindReadVertex(Inputs::MaxTrackIndexName, Inputs.MaxTrackIndex);
 			InVertexData.BindReadVertex(Inputs::IncludeConductorTrackName, Inputs.IncludeConductorTrack);
-
-			ReassignOutputClock = true;
 		}
 
 		virtual void BindOutputs(FOutputVertexInterfaceData& InVertexData) override
 		{
 			InVertexData.BindReadVertex(Outputs::MidiStreamName, Outputs.MidiStream);
-			
-			ReassignOutputClock = true;
 		}
 
-		void Reset(const FResetParams& ResetParams)
+		void Reset(const FResetParams&)
 		{
-			ReassignOutputClock = true;
 		}
 
 		void Execute()
 		{
-			if (ReassignOutputClock)
-			{
-				if (auto* Clock = Inputs.MidiStream->GetMidiClockSource())
-				{
-					Outputs.MidiStream->SetClockSource(*Clock);
-				}
-				ReassignOutputClock = false;
-			}
-
 			Filter.SetTrackRange(*Inputs.MinTrackIndex, *Inputs.MaxTrackIndex, *Inputs.IncludeConductorTrack);
 
 			Outputs.MidiStream->PrepareBlock();
@@ -183,7 +169,6 @@ namespace HarmonixMetasound::Nodes::MidiTrackFilter
 		FInputs Inputs;
 		FOutputs Outputs;
 		Harmonix::Midi::Ops::FMidiTrackFilter Filter;
-		bool ReassignOutputClock = true;
 	};
 
 	class FMidiStreamTrackFilterNode_V1 final : public FNodeFacade
diff --git a/Engine/Plugins/Runtime/Harmonix/Source/HarmonixMetasound/Private/HarmonixMetasound/Nodes/MidiStreamTransposerNode.cpp b/Engine/Plugins/Runtime/Harmonix/Source/HarmonixMetasound/Private/HarmonixMetasound/Nodes/MidiStreamTransposerNode.cpp
--- a/Engine/Plugins/Runtime/Harmonix/Source/HarmonixMetasound/Private/HarmonixMetasound/Nodes/MidiStreamTransposerNode.cpp
+++ b/Engine/Plugins/Runtime/Harmonix/Source/HarmonixMetasound/Private/HarmonixMetasound/Nodes/MidiStreamTransposerNode.cpp
@@ -127,20 +127,15 @@ namespace HarmonixMetasound::Nodes::MidiNoteTranspose
 			InVertexData.BindReadVertex(METASOUND_GET_PARAM_NAME(Inputs::Enable), Inputs.Enabled);
 			InVertexData.BindReadVertex(METASOUND_GET_PARAM_NAME(Inputs::MidiStream), Inputs.MidiStream);
 			InVertexData.BindReadVertex(METASOUND_GET_PARAM_NAME(Inputs::Transposition), Inputs.Transposition);
-
-			ReassignOutputClock = true;
 		}
 
 		virtual void BindOutputs(FOutputVertexInterfaceData& InVertexData) override
 		{
 			InVertexData.BindReadVertex(METASOUND_GET_PARAM_NAME(Outputs::MidiStream), Outputs.MidiStream);
-
-			ReassignOutputClock = true;
 		}
 
 		void Reset(const FResetParams&)
 		{
-			ReassignOutputClock = true;
 		}
 
 		void Execute()
@@ -149,16 +144,6 @@ namespace HarmonixMetasound::Nodes::MidiNoteTranspose
 			{
 				return;
 			}
-			
-			if (ReassignOutputClock)
-			{
-				if (const FMidiClockReadRef* Clock = Inputs.MidiStream->GetMidiClockSource())
-				{
-					Outputs.MidiStream->SetClockSource(*Clock);
-				}
-
-				ReassignOutputClock = false;
-			}
 
 			const int32 Transposition = *Inputs.Transposition;
 
@@ -173,14 +158,13 @@ namespace HarmonixMetasound::Nodes::MidiNoteTranspose
 				
 				return NewEvent;
 			};
-			
-			Outputs.MidiStream->Copy(Inputs.MidiStream, [](const FMidiStreamEvent&){return true;}, Transformer);
+
+			FMidiStream::Copy(*Inputs.MidiStream, *Outputs.MidiStream, FMidiStream::NoOpFilter, Transformer);
 		}
 		
 	private:
 		FInputs Inputs;
 		FOutputs Outputs;
-		bool ReassignOutputClock = true;
 	};
 
 	class FMidiNoteTransposeNode final : public FNodeFacade
diff --git a/Engine/Plugins/Runtime/Harmonix/Source/HarmonixMetasound/Private/HarmonixMetasound/Nodes/StepSequencePlayerNode.cpp b/Engine/Plugins/Runtime/Harmonix/Source/HarmonixMetasound/Private/HarmonixMetasound/Nodes/StepSequencePlayerNode.cpp
--- a/Engine/Plugins/Runtime/Harmonix/Source/HarmonixMetasound/Private/HarmonixMetasound/Nodes/StepSequencePlayerNode.cpp
+++ b/Engine/Plugins/Runtime/Harmonix/Source/HarmonixMetasound/Private/HarmonixMetasound/Nodes/StepSequencePlayerNode.cpp
@@ -349,15 +349,12 @@ namespace HarmonixMetasound
 				{
 				case EMusicPlayerTransportState::Invalid:
 				case EMusicPlayerTransportState::Preparing:
-					MidiOutPin->AddTransportStateChangeMessage(StartFrameIndex, EMusicPlayerTransportState::Prepared);
 					return EMusicPlayerTransportState::Prepared;
 				
 				case EMusicPlayerTransportState::Prepared:
-					MidiOutPin->AddTransportStateChangeMessage(StartFrameIndex, CurrentState);
 					return CurrentState;
 				
 				case EMusicPlayerTransportState::Starting:
-					MidiOutPin->AddTransportStateChangeMessage(StartFrameIndex, EMusicPlayerTransportState::Playing);
 					return EMusicPlayerTransportState::Playing;
 
 				case EMusicPlayerTransportState::Playing:
@@ -367,11 +364,9 @@ namespace HarmonixMetasound
 					return GetTransportState();
 
 				case EMusicPlayerTransportState::Continuing:
-					MidiOutPin->AddTransportStateChangeMessage(StartFrameIndex, EMusicPlayerTransportState::Continuing);
 					return EMusicPlayerTransportState::Playing;
 
 				case EMusicPlayerTransportState::Pausing:
-					MidiOutPin->AddTransportStateChangeMessage(StartFrameIndex, EMusicPlayerTransportState::Pausing);
 					return EMusicPlayerTransportState::Paused;
 
 				case EMusicPlayerTransportState::Paused:
@@ -379,12 +374,10 @@ namespace HarmonixMetasound
 
 				case EMusicPlayerTransportState::Stopping:
 					AllNotesOff(StartFrameIndex, MidiClockInPin->GetCurrentMidiTick(), true);
-					MidiOutPin->AddTransportStateChangeMessage(StartFrameIndex, EMusicPlayerTransportState::Stopping);
 					return EMusicPlayerTransportState::Prepared;
 
 				case EMusicPlayerTransportState::Killing:
 					AllNotesOff(StartFrameIndex, MidiClockInPin->GetCurrentMidiTick(), true);
-					MidiOutPin->AddTransportStateChangeMessage(StartFrameIndex, EMusicPlayerTransportState::Killing);
 					return EMusicPlayerTransportState::Prepared;
 
 				default:
@@ -400,7 +393,7 @@ namespace HarmonixMetasound
 	void FStepSequencePlayerOperator::Init()
 	{
 		MidiOutPin->PrepareBlock();
-		MidiOutPin->SetClockSource(MidiClockInPin);
+		MidiOutPin->SetClock(*MidiClockInPin);
 		MidiClockInPin->RegisterHiResPlayCursor(this);
 
 		SequenceTable = SequenceAssetInPin->GetRenderable();
@@ -424,15 +417,12 @@ namespace HarmonixMetasound
 			{
 			case EMusicPlayerTransportState::Invalid:
 			case EMusicPlayerTransportState::Preparing:
-				MidiOutPin->AddTransportStateChangeMessage(0, EMusicPlayerTransportState::Prepared);
 				return EMusicPlayerTransportState::Prepared;
 				
 			case EMusicPlayerTransportState::Prepared:
-				MidiOutPin->AddTransportStateChangeMessage(0, CurrentState);
 				return CurrentState;
 				
 			case EMusicPlayerTransportState::Starting:
-				MidiOutPin->AddTransportStateChangeMessage(0, EMusicPlayerTransportState::Playing);
 				return EMusicPlayerTransportState::Playing;
 
 			case EMusicPlayerTransportState::Playing:
@@ -442,11 +432,9 @@ namespace HarmonixMetasound
 				return GetTransportState();
 
 			case EMusicPlayerTransportState::Continuing:
-				MidiOutPin->AddTransportStateChangeMessage(0, EMusicPlayerTransportState::Continuing);
 				return EMusicPlayerTransportState::Playing;
 
 			case EMusicPlayerTransportState::Pausing:
-				MidiOutPin->AddTransportStateChangeMessage(0, EMusicPlayerTransportState::Pausing);
 				return EMusicPlayerTransportState::Paused;
 
 			case EMusicPlayerTransportState::Paused:
@@ -454,12 +442,10 @@ namespace HarmonixMetasound
 
 			case EMusicPlayerTransportState::Stopping:
 				AllNotesOff(0, MidiClockInPin->GetCurrentMidiTick(), true);
-				MidiOutPin->AddTransportStateChangeMessage(0, EMusicPlayerTransportState::Stopping);
 				return EMusicPlayerTransportState::Prepared;
 
 			case EMusicPlayerTransportState::Killing:
 				AllNotesOff(0, MidiClockInPin->GetCurrentMidiTick(), true);
-				MidiOutPin->AddTransportStateChangeMessage(0, EMusicPlayerTransportState::Killing);
 				return EMusicPlayerTransportState::Prepared;
 
 			default:
diff --git a/Engine/Plugins/Runtime/Harmonix/Source/HarmonixMetasound/Public/HarmonixMetasound/DataTypes/MidiClock.h b/Engine/Plugins/Runtime/Harmonix/Source/HarmonixMetasound/Public/HarmonixMetasound/DataTypes/MidiClock.h
--- a/Engine/Plugins/Runtime/Harmonix/Source/HarmonixMetasound/Public/HarmonixMetasound/DataTypes/MidiClock.h
+++ b/Engine/Plugins/Runtime/Harmonix/Source/HarmonixMetasound/Public/HarmonixMetasound/DataTypes/MidiClock.h
@@ -57,7 +57,7 @@ namespace HarmonixMetasound
 
 	constexpr FMidiTimestampTempo InvalidMidiTimestampTempo{ -1, 0.0f, 0.0f };
 
-	class HARMONIXMETASOUND_API FMidiClock
+	class HARMONIXMETASOUND_API FMidiClock : public TSharedFromThis<FMidiClock, ESPMode::NotThreadSafe>
 	{
 	public:
 		static constexpr int32 kMidiGranularity = 128;
diff --git a/Engine/Plugins/Runtime/Harmonix/Source/HarmonixMetasound/Public/HarmonixMetasound/DataTypes/MidiStream.h b/Engine/Plugins/Runtime/Harmonix/Source/HarmonixMetasound/Public/HarmonixMetasound/DataTypes/MidiStream.h
--- a/Engine/Plugins/Runtime/Harmonix/Source/HarmonixMetasound/Public/HarmonixMetasound/DataTypes/MidiStream.h
+++ b/Engine/Plugins/Runtime/Harmonix/Source/HarmonixMetasound/Public/HarmonixMetasound/DataTypes/MidiStream.h
@@ -2,18 +2,14 @@
 
 #pragma once
 
-#include "HarmonixMetasound/DataTypes/MusicTransport.h"
-#include "HarmonixMidi/MidiFile.h"
-#include "HarmonixMidi/MidiMsg.h"
-#include "HarmonixMidi/MidiVoiceId.h"
-#include "Logging/LogMacros.h"
 #include "MetasoundDataReference.h"
 #include "MetasoundOperatorSettings.h"
 #include "MetasoundVariable.h"
-
 #include "MidiClock.h"
-
-#include <atomic>
+#include "HarmonixMidi/MidiFile.h"
+#include "HarmonixMidi/MidiMsg.h"
+#include "HarmonixMidi/MidiVoiceId.h"
+#include "Logging/LogMacros.h"
 
 namespace HarmonixMetasound
 {
@@ -61,22 +57,14 @@ namespace HarmonixMetasound
 	class HARMONIXMETASOUND_API FMidiStream
 	{
 	public:
+		explicit FMidiStream(const Metasound::FOperatorSettings& InSettings);
 
-		FMidiStream(const Metasound::FOperatorSettings& InSettings);
-
-		void SetClockSource(const FMidiClockReadRef& ClockSource)
-		{
-			MidiClockSource.Emplace(ClockSource);
-		}
-
-		void ResetClockSource()
-		{
-			MidiClockSource.Reset();
-		}
+		void SetClock(const FMidiClock& InClock);
+		void ResetClock();
+		TSharedPtr<const FMidiClock, ESPMode::NotThreadSafe> GetClock() const;
 
 		void PrepareBlock();
 
-		void AddTransportStateChangeMessage(int32 SampleFrameIndexInBlock, EMusicPlayerTransportState State);
 		void AddMidiEvent(const FMidiStreamEvent& Event);
 		void InsertMidiEvent(const FMidiStreamEvent& Event);
 		void AddNoteOffEventOrCancelPendingNoteOn(const FMidiStreamEvent& Event);
@@ -85,300 +73,55 @@ namespace HarmonixMetasound
 		void SetMidiFile(const FMidiFileProxyPtr& MidiFile) { MidiFileSourceOfEvents = MidiFile; }
 		const FString* GetMidiTrackText(int32 TrackNumber, int32 TextIndex) const;
 
-		const TArray<FMidiTimestampTransportState>& GetTransportChangesInBlock() const
-		{
-			return TransportChangesInBlock;
-		}
+
 		const TArray<FMidiStreamEvent>& GetEventsInBlock() const
 		{
 			return EventsInBlock;
 		}
 
-		const FMidiClockReadRef* GetMidiClockSource() const 
-		{
-			return MidiClockSource.GetPtrOrNull();
-		}
-
-		using FCopyMidiEventsPredicate = TFunctionRef<bool(const FMidiStreamEvent&)>;
-		void Copy(const FMidiStream& InStream, const FCopyMidiEventsPredicate& Predicate, bool IncludeTransportEvents);
-		void CopyTransportEvents(const FMidiStream& InStream);
-		void CopyMidiEvents(const FMidiStream& InStream, const FCopyMidiEventsPredicate& Predicate);
-
-		// Forward declare the read ref for copying other streams
-		using FMidiStreamReadRef = Metasound::TDataReadReference<FMidiStream>;
-
-		// Copies all midi and transport events from the in stream to this stream. Obliterates any
-		// transport/midi events already in this stream!
-		void Copy(FMidiStreamReadRef InStream, bool EventsOnly = false);
-
-		// Copies all midi and transport events from the in streams to this stream. Obliterates any
-		// transport/midi events already in this stream!
-		void Copy(const TArray<FMidiStreamReadRef>& InStreams, bool EventsOnly = false);
-
-		// Copies all transport events, and midi events that match the predicate, from the in stream to this stream.
-		// Obliterates any transport/midi events already in this stream!
-		template<typename T>
-		void Copy(FMidiStreamReadRef InStream, const T& Predicate, bool EventsOnly = false)
-		{
-			if (GetMidiClockSource() && InStream->GetMidiClockSource() && GetMidiClockSource()->Get() != InStream->GetMidiClockSource()->Get())
-			{
-				UE_LOG(LogMidiStreamDataType, Warning, TEXT("Cannot copy midi events from one stream to another if they have different clocks."));
-				return;
-			}
-			if (!EventsOnly)
-			{
-				CopyTransportEvents(InStream);
-			}
-			CopyMidiEvents(InStream, Predicate);
-		}
-
-		// Copies all transport events, and midi events that match the predicate, from the in streams to this stream.
-		// Obliterates any transport/midi events already in this stream!
-		template<typename T>
-		void Copy(const TArray<FMidiStreamReadRef>& InStreams, const T& Predicate, bool EventsOnly = false)
-		{
-			if (!InStreams.Num())
-			{
-				return;
-			}
-
-			if (InStreams.Num() == 1)
-			{
-				Copy(InStreams[0], Predicate);
-				return;
-			}
-
-			TArray<FMidiStreamReadRef> SameClockStreams;
-			FilterArrayToStreamsWithTheSameClock(InStreams, SameClockStreams);
-			if (!EventsOnly)
-			{
-				CopyTransportEvents(SameClockStreams);
-			}
-			CopyMidiEvents(SameClockStreams, Predicate);
-		}
-
-		// Copies all midi and transport events from the in stream to this stream that match the predicate. Obliterates any
-		// transport/midi events already in this stream!
-		template<typename FILTER, typename TRANSFORMER>
-		void Copy(FMidiStreamReadRef InStream, const FILTER& Predicate, const TRANSFORMER& Transformer, bool EventsOnly = false)
-		{
-			if (GetMidiClockSource() && InStream->GetMidiClockSource() && GetMidiClockSource()->Get() != InStream->GetMidiClockSource()->Get())
-			{
-				UE_LOG(LogMidiStreamDataType, Warning, TEXT("Cannot copy midi events from one stream to another if they have different clocks."));
-				return;
-			}
-			if (!EventsOnly)
-			{
-				CopyTransportEvents(InStream);
-			}
-			CopyMidiEvents(InStream, Predicate, Transformer);
-		}
-
 		int32 GetTicksPerQuarterNote() const { return TicksPerQuarterNote; }
 		void SetTicksPerQuarterNote(int32 InTicksPerQuarterNote) { TicksPerQuarterNote = InTicksPerQuarterNote; }
 
-		// Returns an array of only the in streams that have the same clock as this stream
-		void FilterArrayToStreamsWithTheSameClock(const TArray<FMidiStreamReadRef>& InStreams, TArray<FMidiStreamReadRef>& OutStreams);
-		// Copies all transport events from the in streams to this. Obliterates any existing transport events.
-		void CopyTransportEvents(const TArray<FMidiStreamReadRef>& InStreams);
-		// Copies all transport events from the in stream to this. Obliterates any existing transport events.
-		void CopyTransportEvents(const FMidiStreamReadRef& InStream);
-		// Copies all midi events from the in streams to this. Obliterates and existing midi events.
-		void CopyMidiEvents(const TArray<FMidiStreamReadRef>& InStreams);
-		// Merges all midi events from the in streams to this.
-		void MergeMidiEvents(const TArray<FMidiStreamReadRef>& InStreams);
-		// Copies all midi events from the in stream to this. Obliterates and existing midi events.
-		void CopyMidiEvents(const FMidiStreamReadRef& InStream);
-		// Merges all midi events from the in stream to this.
-		void MergeMidiEvents(const FMidiStreamReadRef& InStream);
-
-		// Copies all midi events that match the predicate from the in stream to this. Obliterates and existing midi events.
-		template<typename T>
-		void CopyMidiEvents(const FMidiStreamReadRef& InStream, const T& Predicate)
-		{
-			const TArray<FMidiStreamEvent>& MidiEvents = InStream->GetEventsInBlock();
-			if (!ensureAlwaysMsgf(EventsInBlock.Num()==0, TEXT("Existing midi events in midi stream are being destroyed during copy!")))
-			{
-				EventsInBlock.Empty(FMath::Max(MidiEvents.Num(), 32));
-				ActiveVoices.Reset();
-			}
-			for (const FMidiStreamEvent& Event : MidiEvents)
-			{
-				if (Predicate(Event))
-				{
-					AddMidiEvent(Event);
-				}
-			}
-		}
-
-		// Merges all midi events that match the predicate from the in stream to this. 
-		template<typename T>
-		void MergeMidiEvents(const FMidiStreamReadRef& InStream, const T& Predicate)
-		{
-			const TArray<FMidiStreamEvent>& MidiEvents = InStream->GetEventsInBlock();
-			for (const FMidiStreamEvent& Event : MidiEvents)
-			{
-				if (Predicate(Event))
-				{
-					InsertMidiEvent(Event);
-				}
-			}
-		}
-
-		// Copies all midi events that match the predicate from the in stream to this. Obliterates and existing midi events.
-		template<typename T>
-		void CopyMidiEvents(const TArray<FMidiStreamReadRef>& InStreams, const T& Predicate)
-		{
-			if (!ensureAlwaysMsgf(EventsInBlock.Num() == 0, TEXT("Existing midi events in midi stream are being destroyed during copy!")))
-			{
-				EventsInBlock.Empty(32);
-				ActiveVoices.Reset();
-			}
-
-			// Define the type so that we can store iterators in an array
-			using MidiStreamIterator = TArray<FMidiStreamEvent>::RangedForConstIteratorType;
-
-			TArray<MidiStreamIterator> WalkingIterators;
-			TArray<MidiStreamIterator> IteratorEnds;
-
-			for (int32 i = 0; i < InStreams.Num(); ++i)
-			{
-				WalkingIterators.Add(InStreams[i]->GetEventsInBlock().begin());
-				IteratorEnds.Add(InStreams[i]->GetEventsInBlock().end());
-			}
-
-			while (1)
-			{
-				int32 StreamWithEarliestEvent = -1;
-				int32 EarliestMidiTick = MAX_int32;
-
-				for (int32 IteratorIndex = 0; IteratorIndex < WalkingIterators.Num(); ++IteratorIndex)
-				{
-					if (WalkingIterators[IteratorIndex] != IteratorEnds[IteratorIndex])
-					{
-						if ((*WalkingIterators[IteratorIndex]).AuthoredMidiTick < EarliestMidiTick)
-						{
-							EarliestMidiTick = (*WalkingIterators[IteratorIndex]).AuthoredMidiTick;
-							StreamWithEarliestEvent = IteratorIndex;
-						}
-					}
-				}
-
-				// done?
-				if (StreamWithEarliestEvent == -1)
-				{
-					break;
-				}
-
-				if (Predicate(*WalkingIterators[StreamWithEarliestEvent]))
-				{
-					AddMidiEvent(*WalkingIterators[StreamWithEarliestEvent]);
-				}
-				++WalkingIterators[StreamWithEarliestEvent];
-			}
-		}
-
-		// Merges all midi events that match the predicate from the in stream to this.
-		template<typename T>
-		void MergeMidiEvents(const TArray<FMidiStreamReadRef>& InStreams, const T& Predicate)
-		{
-			// Define the type so that we can store iterators in an array
-			using MidiStreamIterator = TArray<FMidiStreamEvent>::RangedForConstIteratorType;
-
-			TArray<MidiStreamIterator> WalkingIterators;
-			TArray<MidiStreamIterator> IteratorEnds;
-
-			for (int32 i = 0; i < InStreams.Num(); ++i)
-			{
-				WalkingIterators.Add(InStreams[i]->GetEventsInBlock().begin());
-				IteratorEnds.Add(InStreams[i]->GetEventsInBlock().end());
-			}
-
-			while (1)
-			{
-				int32 StreamWithEarliestEvent = -1;
-				int32 EarliestBlockSampleFrameIndex = MAX_int32;
-
-				for (int32 IteratorIndex = 0; IteratorIndex < WalkingIterators.Num(); ++IteratorIndex)
-				{
-					if (WalkingIterators[IteratorIndex] != IteratorEnds[IteratorIndex])
-					{
-						if ((*WalkingIterators[IteratorIndex]).BlockSampleFrameIndex == EarliestBlockSampleFrameIndex)
-						{
-							if ((*WalkingIterators[StreamWithEarliestEvent]).MidiMessage.IsNoteOn() && (*WalkingIterators[IteratorIndex]).MidiMessage.IsNoteOff())
-							{
-								StreamWithEarliestEvent = IteratorIndex;
-							}
-						}
-						else if ((*WalkingIterators[IteratorIndex]).BlockSampleFrameIndex < EarliestBlockSampleFrameIndex)
-						{
-							EarliestBlockSampleFrameIndex = (*WalkingIterators[IteratorIndex]).BlockSampleFrameIndex;
-							StreamWithEarliestEvent = IteratorIndex;
-						}
-					}
-				}
-
-				// done?
-				if (StreamWithEarliestEvent == -1)
-				{
-					break;
-				}
-
-				if (Predicate(*WalkingIterators[StreamWithEarliestEvent]))
-				{
-					InsertMidiEvent(*WalkingIterators[StreamWithEarliestEvent]);
-				}
-				++WalkingIterators[StreamWithEarliestEvent];
-			}
-		}
-
-		// Copies all midi events that match the predicate from the in stream to this. Optionally transforms the events. Obliterates and existing midi events.
-		template<typename FILTER, typename TRANSFORMER>
-		void CopyMidiEvents(const FMidiStreamReadRef& InStream, const FILTER& Predicate, const TRANSFORMER& Transformer)
-		{
-			const TArray<FMidiStreamEvent>& MidiEvents = InStream->GetEventsInBlock();
-			EventsInBlock.Empty(FMath::Max(MidiEvents.Num(), 32));
-			ActiveVoices.Reset();
-			
-			for (const FMidiStreamEvent& Event : MidiEvents)
-			{
-				if (Predicate(Event))
-				{
-					AddMidiEvent(Transformer(Event));
-				}
-			}
-		}
-
-		// Copies all midi events that match the predicate from the in stream to this. Optionally transforms the events. Obliterates and existing midi events.
-		template<typename FILTER, typename TRANSFORMER>
-		void MergeMidiEvents(const FMidiStreamReadRef& InStream, const FILTER& Predicate, const TRANSFORMER& Transformer)
-		{
-			const TArray<FMidiStreamEvent>& MidiEvents = InStream->GetEventsInBlock();
-			for (const FMidiStreamEvent& Event : MidiEvents)
-			{
-				if (Predicate(Event))
-				{
-					InsertMidiEvent(Transformer(Event));
-				}
-			}
-		}
+		using FEventFilter = TFunction<bool(const FMidiStreamEvent&)>;
+		inline static const FEventFilter NoOpFilter = [](const FMidiStreamEvent&) { return true; };
+		
+		using FEventTransformer = TFunction<FMidiStreamEvent(const FMidiStreamEvent&)>;
+		inline static const FEventTransformer NoOpTransformer = [](const FMidiStreamEvent& Event) { return Event; };
+		
+		static void Copy(
+			const FMidiStream& From,
+			FMidiStream& To,
+			const FEventFilter& Filter = NoOpFilter,
+			const FEventTransformer& Transformer = NoOpTransformer);
+
+		static void Merge(
+			const FMidiStream& From,
+			FMidiStream& To,
+			const FEventFilter& Filter = NoOpFilter,
+			const FEventTransformer& Transformer = NoOpTransformer);
+		
+		static void Merge(
+			const FMidiStream& FromA,
+			const FMidiStream& FromB,
+			FMidiStream& To,
+			const FEventFilter& Filter = NoOpFilter,
+			const FEventTransformer& Transformer = NoOpTransformer);
 
 	private:
 		FMidiFileProxyPtr MidiFileSourceOfEvents;
-		int32 NumFramesPerBlock = 0;
-		Metasound::FSampleRate SampleRate  = 0;
 		int32 TicksPerQuarterNote = Harmonix::Midi::Constants::GTicksPerQuarterNoteInt;
 
-		FMidiTimestampTransportState CurrentTransportState;
-		TArray<FMidiTimestampTransportState> TransportChangesInBlock;
-
 		TArray<FMidiStreamEvent> EventsInBlock;
 
 		friend class FMidiVoiceTracker;
 		void UpdateActiveVoice(const FMidiStreamEvent& Event);
 		TSet<FMidiVoiceId> ActiveVoices;
 
-		TOptional<FMidiClockReadRef> MidiClockSource;
+		static TSharedPtr<const FMidiClock, ESPMode::NotThreadSafe> DummyClock;
+		TWeakPtr<const FMidiClock, ESPMode::NotThreadSafe> Clock;
+
+		// Map to handle re-mapping merged MIDI events, which helps to disambiguate split/transposed notes
+		TMap<uint32, FMidiVoiceGeneratorBase> GeneratorMap;
 	};
 
 	// Declare aliases IN the namespace...
diff --git a/Engine/Plugins/Runtime/Harmonix/Source/HarmonixMetasoundTests/Private/Analysis/MidiStreamVertexAnalyzerTests.cpp b/Engine/Plugins/Runtime/Harmonix/Source/HarmonixMetasoundTests/Private/Analysis/MidiStreamVertexAnalyzerTests.cpp
--- a/Engine/Plugins/Runtime/Harmonix/Source/HarmonixMetasoundTests/Private/Analysis/MidiStreamVertexAnalyzerTests.cpp
+++ b/Engine/Plugins/Runtime/Harmonix/Source/HarmonixMetasoundTests/Private/Analysis/MidiStreamVertexAnalyzerTests.cpp
@@ -105,7 +105,7 @@ namespace HarmonixMetasoundTests::MidiStreamVertexAnalyzer
 		UTEST_TRUE("MIDI input is valid", MidiIn.IsSet());
 		constexpr float Tempo = 87;
 		HarmonixMetasound::FMidiClockWriteRef Clock = HarmonixMetasound::FMidiClockWriteRef::CreateNew(Generator->OperatorSettings);
-		(*MidiIn)->SetClockSource(Clock);
+		(*MidiIn)->SetClock(*Clock);
 		
 		bool CallbackSuccess = false;
 
diff --git a/Engine/Plugins/Runtime/Harmonix/Source/HarmonixMetasoundTests/Private/DataTypes/MidiStreamTests.cpp b/Engine/Plugins/Runtime/Harmonix/Source/HarmonixMetasoundTests/Private/DataTypes/MidiStreamTests.cpp
new file mode 100644
--- /dev/null
+++ b/Engine/Plugins/Runtime/Harmonix/Source/HarmonixMetasoundTests/Private/DataTypes/MidiStreamTests.cpp
@@ -0,0 +1,268 @@
+// Copyright Epic Games, Inc. All Rights Reserved.
+
+#include "HarmonixMetasound/DataTypes/MidiStream.h"
+#include "Misc/AutomationTest.h"
+
+#if WITH_DEV_AUTOMATION_TESTS
+
+namespace HarmonixMetasound::MidiStream::Tests
+{
+	FMidiStream MakeTestMidiStream(const Metasound::FOperatorSettings& OperatorSettings)
+	{
+		FMidiStream Stream{ OperatorSettings };
+
+		// Make a bunch of note events
+		constexpr int32 NumNotes = 3;
+		uint8 NoteNumber = 0;
+		int32 EventAudioFrame = 0;
+		
+		for (int32 i = 0; i < NumNotes; ++i)
+		{
+			FMidiStreamEvent NoteOn{static_cast<uint32>(0), FMidiMsg::CreateNoteOn(5, NoteNumber, 123)};
+			NoteOn.BlockSampleFrameIndex = EventAudioFrame;
+			Stream.InsertMidiEvent(NoteOn);
+
+			constexpr int32 AudioFrameDelta = 11;
+			EventAudioFrame += AudioFrameDelta;
+			
+			FMidiStreamEvent NoteOff{static_cast<uint32>(0), FMidiMsg::CreateNoteOff(5, NoteNumber)};
+			NoteOff.BlockSampleFrameIndex = EventAudioFrame;
+			Stream.InsertMidiEvent(NoteOff);
+
+			constexpr uint8 NoteNumberDelta = 17;
+			NoteNumber += NoteNumberDelta;
+		}
+
+		// Make a bunch of CC events
+		for (int32 i = 0; i < OperatorSettings.GetNumFramesPerBlock(); ++i)
+		{
+			FMidiStreamEvent CC{ static_cast<uint32>(0), FMidiMsg::CreateControlChange(7, 45, i % 128)};
+			CC.BlockSampleFrameIndex = i;
+			Stream.InsertMidiEvent(CC);
+		}
+		
+		return Stream;
+	}
+
+	bool EventsAreEqual(const FMidiStreamEvent& A, const FMidiStreamEvent& B, bool MatchVoiceId = true)
+	{
+		return (!MatchVoiceId || A.GetVoiceId() == B.GetVoiceId())
+		&& A.TrackIndex == B.TrackIndex
+		&& A.BlockSampleFrameIndex == B.BlockSampleFrameIndex
+		&& A.MidiMessage == B.MidiMessage;
+	}
+	
+	IMPLEMENT_SIMPLE_AUTOMATION_TEST(
+		FMidiStreamCopyTest,
+		"Harmonix.Metasound.DataTypes.MidiStream.Copy",
+		EAutomationTestFlags::EditorContext | EAutomationTestFlags::EngineFilter)
+	bool FMidiStreamCopyTest::RunTest(const FString&)
+	{
+		Metasound::FOperatorSettings OperatorSettings{ 48000, 100 };
+		const FMidiStream FromStream = MakeTestMidiStream(OperatorSettings);
+		FMidiStream ToStream{ OperatorSettings };
+
+		// Copy with no filter or transform
+		{
+			FMidiStream::Copy(FromStream, ToStream);
+
+			const TArray<FMidiStreamEvent>& FromEvents = FromStream.GetEventsInBlock();
+			const TArray<FMidiStreamEvent>& ToEvents = ToStream.GetEventsInBlock();
+			UTEST_EQUAL("Same number of events", ToEvents.Num(), FromEvents.Num());
+			
+			for (int i = 0; i < FromEvents.Num(); ++i)
+			{
+				UTEST_TRUE("Events match", EventsAreEqual(FromEvents[i], ToEvents[i]));
+			}
+		}
+
+		// Copy with filter
+		{
+			const auto Filter = [](const FMidiStreamEvent& Event)
+			{
+				return Event.MidiMessage.IsNoteOn();
+			};
+			
+			FMidiStream::Copy(FromStream, ToStream, Filter);
+
+			const TArray<FMidiStreamEvent>& FromEvents = FromStream.GetEventsInBlock();
+			const TArray<FMidiStreamEvent>& ToEvents = ToStream.GetEventsInBlock();
+			
+			for (int i = 0; i < FromEvents.Num(); ++i)
+			{
+				if (Filter(FromEvents[i]))
+				{
+					const FMidiStreamEvent FromEvent = FromEvents[i];
+					const bool Found = ToEvents.ContainsByPredicate([FromEvent](const FMidiStreamEvent& Event)
+					{
+						return EventsAreEqual(Event, FromEvent);
+					});
+					UTEST_TRUE("Event passed through", Found);
+				}
+			}
+		}
+
+		// Copy with transformation
+		{
+			const auto Transform = [](const FMidiStreamEvent& Event)
+			{
+				if (Event.MidiMessage.IsNoteOn())
+				{
+					FMidiStreamEvent TransformedEvent = Event;
+					TransformedEvent.MidiMessage.Data1 = TransformedEvent.MidiMessage.Data1 + 2;
+					return TransformedEvent;
+				}
+
+				return Event;
+			};
+			FMidiStream::Copy(FromStream, ToStream, FMidiStream::NoOpFilter, Transform);
+
+			const TArray<FMidiStreamEvent>& FromEvents = FromStream.GetEventsInBlock();
+			const TArray<FMidiStreamEvent>& ToEvents = ToStream.GetEventsInBlock();
+			UTEST_EQUAL("Same number of events", ToEvents.Num(), FromEvents.Num());
+			
+			for (int i = 0; i < FromEvents.Num(); ++i)
+			{
+				FMidiStreamEvent ExpectedEvent = FromEvents[i];
+				
+				if (FromEvents[i].MidiMessage.IsNoteOn())
+				{
+					ExpectedEvent.MidiMessage.Data1 = ExpectedEvent.MidiMessage.Data1 + 2;
+				}
+
+				UTEST_TRUE("Events match", EventsAreEqual(ExpectedEvent, ToEvents[i]));
+			}
+		}
+		
+		return true;
+	}
+
+	void MakeMergeStreams(
+		const Metasound::FOperatorSettings& OperatorSettings,
+		FMidiStream& FromA,
+		FMidiStream& FromB,
+		FMidiStream& Expected,
+		const FMidiStream::FEventFilter& Filter = FMidiStream::NoOpFilter,
+		const FMidiStream::FEventTransformer& Transformer = FMidiStream::NoOpTransformer)
+	{
+		FromA.PrepareBlock();
+		FromA = MakeTestMidiStream(OperatorSettings);
+		
+		// Change the values of the std messages in the second stream
+		FromB.PrepareBlock();
+		for (const FMidiStreamEvent& Event : FromA.GetEventsInBlock())
+		{
+			if (Event.MidiMessage.IsStd())
+			{
+				FMidiStreamEvent Transformed = Event;
+				const uint8 Data1 = Transformed.MidiMessage.Data1;
+				Transformed.MidiMessage.Data1 = Transformed.MidiMessage.Data2;
+				Transformed.MidiMessage.Data2 = Data1;
+				FromB.InsertMidiEvent(Transformed);
+			}
+			else
+			{
+				FromB.InsertMidiEvent(Event);
+			}
+		}
+
+		// Generate the expected outcome from merging
+		Expected.PrepareBlock();
+		for (const FMidiStreamEvent& Event : FromA.GetEventsInBlock())
+		{
+			if (Filter(Event))
+			{
+				Expected.InsertMidiEvent(Transformer(Event));
+			}
+		}
+		for (const FMidiStreamEvent& Event : FromB.GetEventsInBlock())
+		{
+			if (Filter(Event))
+			{
+				Expected.InsertMidiEvent(Transformer(Event));
+			}
+		}
+	}
+
+	IMPLEMENT_SIMPLE_AUTOMATION_TEST(
+		FMidiStreamMergeTest,
+		"Harmonix.Metasound.DataTypes.MidiStream.Merge",
+		EAutomationTestFlags::EditorContext | EAutomationTestFlags::EngineFilter)
+	bool FMidiStreamMergeTest::RunTest(const FString&)
+	{
+		Metasound::FOperatorSettings OperatorSettings{ 48000, 100 };
+		FMidiStream FromA{OperatorSettings};
+		FMidiStream FromB{OperatorSettings};
+		FMidiStream Expected{OperatorSettings};
+		FMidiStream To{OperatorSettings};
+
+		// no filter or transform
+		{
+			MakeMergeStreams(OperatorSettings, FromA, FromB, Expected);
+			To.PrepareBlock();
+			FMidiStream::Merge(FromA, FromB, To);
+
+			const TArray<FMidiStreamEvent>& ExpectedEvents = Expected.GetEventsInBlock();
+			const TArray<FMidiStreamEvent>& ToEvents = To.GetEventsInBlock();
+			UTEST_EQUAL("Same number of events", ToEvents.Num(), ExpectedEvents.Num());
+			
+			for (int i = 0; i < ExpectedEvents.Num(); ++i)
+			{
+				UTEST_TRUE("Events match", EventsAreEqual(ExpectedEvents[i], ToEvents[i], false));
+			}
+		}
+
+		// with filter
+		{
+			const auto Filter = [](const FMidiStreamEvent& Event)
+			{
+				return Event.MidiMessage.IsNoteOn();
+			};
+			
+			MakeMergeStreams(OperatorSettings, FromA, FromB, Expected, Filter);
+			To.PrepareBlock();
+			FMidiStream::Merge(FromA, FromB, To, Filter);
+
+			const TArray<FMidiStreamEvent>& ExpectedEvents = Expected.GetEventsInBlock();
+			const TArray<FMidiStreamEvent>& ToEvents = To.GetEventsInBlock();
+			UTEST_EQUAL("Same number of events", ToEvents.Num(), ExpectedEvents.Num());
+			
+			for (int i = 0; i < ExpectedEvents.Num(); ++i)
+			{
+				UTEST_TRUE("Events match", EventsAreEqual(ExpectedEvents[i], ToEvents[i], false));
+			}
+		}
+
+		// with transformation
+		{
+			const auto Transform = [](const FMidiStreamEvent& Event)
+			{
+				if (Event.MidiMessage.IsControlChange())
+				{
+					FMidiStreamEvent Transformed = Event;
+					Transformed.MidiMessage.Data2 = 123;
+					return Transformed;
+				}
+
+				return Event;
+			};
+			
+			MakeMergeStreams(OperatorSettings, FromA, FromB, Expected, FMidiStream::NoOpFilter, Transform);
+			To.PrepareBlock();
+			FMidiStream::Merge(FromA, FromB, To, FMidiStream::NoOpFilter, Transform);
+
+			const TArray<FMidiStreamEvent>& ExpectedEvents = Expected.GetEventsInBlock();
+			const TArray<FMidiStreamEvent>& ToEvents = To.GetEventsInBlock();
+			UTEST_EQUAL("Same number of events", ToEvents.Num(), ExpectedEvents.Num());
+			
+			for (int i = 0; i < ExpectedEvents.Num(); ++i)
+			{
+				UTEST_TRUE("Events match", EventsAreEqual(ExpectedEvents[i], ToEvents[i], false));
+			}
+		}
+		
+		return true;
+	}
+}
+
+#endif
\ No newline at end of file
diff --git a/Engine/Plugins/Runtime/Harmonix/Source/HarmonixMetasoundTests/Private/MidiOps/MidiChannelFilterTests.cpp b/Engine/Plugins/Runtime/Harmonix/Source/HarmonixMetasoundTests/Private/MidiOps/MidiChannelFilterTests.cpp
--- a/Engine/Plugins/Runtime/Harmonix/Source/HarmonixMetasoundTests/Private/MidiOps/MidiChannelFilterTests.cpp
+++ b/Engine/Plugins/Runtime/Harmonix/Source/HarmonixMetasoundTests/Private/MidiOps/MidiChannelFilterTests.cpp
@@ -51,7 +51,6 @@ namespace Harmonix::Midi::Ops::Tests
 		HarmonixMetasound::FMidiStream OutputStream{ OperatorSettings };
 
 		// Add some events to the input stream
-		InputStream.AddTransportStateChangeMessage(0, HarmonixMetasound::EMusicPlayerTransportState::Playing);
 		constexpr uint8 NoteChannelIdx = 1;
 		constexpr uint8 NoteNumber = 76;
 		constexpr uint8 NoteVelocity = 67;
@@ -64,7 +63,6 @@ namespace Harmonix::Midi::Ops::Tests
 
 		// Default: only transport events pass
 		Filter.Process(InputStream, OutputStream);
-		UTEST_EQUAL("Got transport events", OutputStream.GetTransportChangesInBlock().Num(), 1);
 		UTEST_EQUAL("No MIDI events", OutputStream.GetEventsInBlock().Num(), 0);
 
 		// Channel mismatch: no MIDI events pass, transport state hasn't changed, so we won't get another one of those
