Fix: errors with Harmonix.Midi.Ops.PulseGenerator.Basic, in Harmonix Plugin Tests

Clock needs to prepare block at the beginning of each block

#tests Harmonix Plugin Tests

diff --git a/Engine/Plugins/Runtime/Harmonix/Source/HarmonixMetasoundTests/Private/MidiOps/PulseGeneratorTests.cpp b/Engine/Plugins/Runtime/Harmonix/Source/HarmonixMetasoundTests/Private/MidiOps/PulseGeneratorTests.cpp
--- a/Engine/Plugins/Runtime/Harmonix/Source/HarmonixMetasoundTests/Private/MidiOps/PulseGeneratorTests.cpp
+++ b/Engine/Plugins/Runtime/Harmonix/Source/HarmonixMetasoundTests/Private/MidiOps/PulseGeneratorTests.cpp
@@ -52,6 +52,7 @@ namespace Harmonix::Midi::Ops::Tests
 			while (Clock.GetCurrentMidiTick() < EndTick)
 			{
 				// Advance the clock, which will advance the play cursor in the pulse generator
+				Clock.PrepareBlock();
 				Clock.WriteAdvance(0, OperatorSettings.GetNumFramesPerBlock());
 
 				// Process, which will pop the next notes
