[Harmonix] Fix backward clock assignment logic in FMidiStream::Merge

diff --git a/Engine/Plugins/Runtime/Harmonix/Source/HarmonixMetasound/Private/HarmonixMetasound/DataTypes/MidiStream.cpp b/Engine/Plugins/Runtime/Harmonix/Source/HarmonixMetasound/Private/HarmonixMetasound/DataTypes/MidiStream.cpp
--- a/Engine/Plugins/Runtime/Harmonix/Source/HarmonixMetasound/Private/HarmonixMetasound/DataTypes/MidiStream.cpp
+++ b/Engine/Plugins/Runtime/Harmonix/Source/HarmonixMetasound/Private/HarmonixMetasound/DataTypes/MidiStream.cpp
@@ -164,7 +164,7 @@ namespace HarmonixMetasound
 			}
 
 			// If the "to" clock is the dummy, and the "from" clock isn't, overwrite the "to" clock
-			if (FromClockIsDummy && !ToClockIsDummy)
+			if (!FromClockIsDummy && ToClockIsDummy)
 			{
 				To.SetClock(*FromClock);
 			}
diff --git a/Engine/Plugins/Runtime/Harmonix/Source/HarmonixMetasoundTests/Private/DataTypes/MidiStreamTests.cpp b/Engine/Plugins/Runtime/Harmonix/Source/HarmonixMetasoundTests/Private/DataTypes/MidiStreamTests.cpp
--- a/Engine/Plugins/Runtime/Harmonix/Source/HarmonixMetasoundTests/Private/DataTypes/MidiStreamTests.cpp
+++ b/Engine/Plugins/Runtime/Harmonix/Source/HarmonixMetasoundTests/Private/DataTypes/MidiStreamTests.cpp
@@ -263,6 +263,28 @@ namespace HarmonixMetasound::MidiStream::Tests
 		
 		return true;
 	}
+
+	IMPLEMENT_SIMPLE_AUTOMATION_TEST(
+	FMidiStreamMergeClockTest,
+	"Harmonix.Metasound.DataTypes.MidiStream.Merge.ClockIsSet",
+	EAutomationTestFlags::EditorContext | EAutomationTestFlags::EngineFilter)
+	bool FMidiStreamMergeClockTest::RunTest(const FString&)
+	{
+		// Make sure merging streams results in the correct clock being set on the output stream
+		Metasound::FOperatorSettings OperatorSettings{ 48000, 100 };
+		FMidiStream FromA{OperatorSettings};
+		FMidiStream FromB{OperatorSettings};
+		FMidiStream To{OperatorSettings};
+
+		TSharedRef<const FMidiClock, ESPMode::NotThreadSafe> ClockA = MakeShared<FMidiClock, ESPMode::NotThreadSafe>(OperatorSettings);
+		FromA.SetClock(*ClockA);
+
+		FMidiStream::Merge(FromA, FromB, To);
+		TSharedPtr<const FMidiClock, ESPMode::NotThreadSafe> ToClock = To.GetClock();
+		UTEST_EQUAL("Output clock is clock A", ClockA.ToSharedPtr(), ToClock);
+		
+		return true;
+	}
 }
 
 #endif
\ No newline at end of file
