[MetasoundMusicClockDriver] Listen for generator IO changes, as they may indicate that we should be attaching to a different Clock
#tests PIE

diff --git a/Engine/Plugins/Runtime/Harmonix/Source/HarmonixMetasound/Private/HarmonixMetasound/Components/MetasoundMusicClockDriver.cpp b/Engine/Plugins/Runtime/Harmonix/Source/HarmonixMetasound/Private/HarmonixMetasound/Components/MetasoundMusicClockDriver.cpp
--- a/Engine/Plugins/Runtime/Harmonix/Source/HarmonixMetasound/Private/HarmonixMetasound/Components/MetasoundMusicClockDriver.cpp
+++ b/Engine/Plugins/Runtime/Harmonix/Source/HarmonixMetasound/Private/HarmonixMetasound/Components/MetasoundMusicClockDriver.cpp
@@ -134,6 +134,7 @@ bool FMetasoundMusicClockDriver::AttemptToConnectToAudioComponentsMetasound()
 	}
 	GeneratorAttachedCallbackHandle = CurrentGeneratorHandle->OnGeneratorHandleAttached.AddLambda([this](){OnGeneratorAttached();});
 	GeneratorDetachedCallbackHandle = CurrentGeneratorHandle->OnGeneratorHandleDetached.AddLambda([this](){OnGeneratorDetached();});
+	GeneratorIOUpdatedCallbackHandle = CurrentGeneratorHandle->OnIOUpdated.AddLambda([this](){OnGeneratorIOUpdated();});
 	UMetasoundGeneratorHandle::FOnSetGraph::FDelegate OnSetGraph;
 	OnSetGraph.BindLambda([this](){OnGraphSet();});
 	GraphChangedCallbackHandle = CurrentGeneratorHandle->AddGraphSetCallback(MoveTemp(OnSetGraph));
@@ -149,6 +150,8 @@ void FMetasoundMusicClockDriver::DetachAllCallbacks()
 		GeneratorAttachedCallbackHandle.Reset();
 		CurrentGeneratorHandle->OnGeneratorHandleDetached.Remove(GeneratorDetachedCallbackHandle);
 		GeneratorDetachedCallbackHandle.Reset();
+		CurrentGeneratorHandle->OnIOUpdated.Remove(GeneratorIOUpdatedCallbackHandle);
+		GeneratorIOUpdatedCallbackHandle.Reset();
 		CurrentGeneratorHandle->RemoveGraphSetCallback(GraphChangedCallbackHandle);
 		GraphChangedCallbackHandle.Reset();
 	}
@@ -164,6 +167,13 @@ void FMetasoundMusicClockDriver::OnGraphSet()
 	TryToRegisterPlayCursor();
 }
 
+void FMetasoundMusicClockDriver::OnGeneratorIOUpdated()
+{
+	// An output vertex update may have destroyed our Clock, reattach to the new one if it's there
+	// FUTURE: There isn't a way to listen to vertex changes on a specific Metasound output, so until then this will be called every time IO changes
+	TryToRegisterPlayCursor();
+}
+
 void FMetasoundMusicClockDriver::OnGeneratorDetached()
 {
 	check(IsInGameThread());
diff --git a/Engine/Plugins/Runtime/Harmonix/Source/HarmonixMetasound/Private/HarmonixMetasound/Components/MetasoundMusicClockDriver.h b/Engine/Plugins/Runtime/Harmonix/Source/HarmonixMetasound/Private/HarmonixMetasound/Components/MetasoundMusicClockDriver.h
--- a/Engine/Plugins/Runtime/Harmonix/Source/HarmonixMetasound/Private/HarmonixMetasound/Components/MetasoundMusicClockDriver.h
+++ b/Engine/Plugins/Runtime/Harmonix/Source/HarmonixMetasound/Private/HarmonixMetasound/Components/MetasoundMusicClockDriver.h
@@ -28,6 +28,7 @@ protected:
 	void OnGeneratorAttached();
 	void OnGeneratorDetached();
 	void OnGraphSet();
+	void OnGeneratorIOUpdated();
 
 private:
 	FName MetasoundOutputName;
@@ -45,6 +46,7 @@ private:
 
 	FDelegateHandle GeneratorAttachedCallbackHandle;
 	FDelegateHandle GeneratorDetachedCallbackHandle;
+	FDelegateHandle GeneratorIOUpdatedCallbackHandle;
 	FDelegateHandle GraphChangedCallbackHandle;
 
 	bool AttemptToConnectToAudioComponentsMetasound();
