Allow cursor owner to become nullptr when the midi clock is not found, or doesn't have a cursor manager.

diff --git a/Engine/Plugins/Runtime/Harmonix/Source/HarmonixMetasound/Private/HarmonixMetasound/Components/MetasoundMusicClockDriver.cpp b/Engine/Plugins/Runtime/Harmonix/Source/HarmonixMetasound/Private/HarmonixMetasound/Components/MetasoundMusicClockDriver.cpp
--- a/Engine/Plugins/Runtime/Harmonix/Source/HarmonixMetasound/Private/HarmonixMetasound/Components/MetasoundMusicClockDriver.cpp
+++ b/Engine/Plugins/Runtime/Harmonix/Source/HarmonixMetasound/Private/HarmonixMetasound/Components/MetasoundMusicClockDriver.cpp
@@ -117,6 +117,41 @@ bool FMetasoundMusicClockDriver::ConnectToAudioComponentsMetasound(UAudioCompone
 	return AttemptToConnectToAudioComponentsMetasound();
 }
 
+void FMetasoundMusicClockDriver::ResetCursorOwner(TSharedPtr<FMidiPlayCursorMgr> MidiPlayCursorMgr)
+{
+	check(IsInGameThread());
+
+	// Verify that the cursor owner is changing.
+	if (CursorOwner != MidiPlayCursorMgr)
+	{
+		if (MidiPlayCursorMgr)
+		{
+			// Register with the new cursor owner and broadcast the clock connection event.
+			MidiPlayCursorMgr->RegisterLowResPlayCursor(&Cursor);
+			CursorOwner = MoveTemp(MidiPlayCursorMgr);
+			WasEverConnected = true;
+
+			check(Clock);
+			Clock->MusicClockConnectedEvent.Broadcast();
+		}
+		else
+		{
+			// Unregister the old cursor owner and broadcast the clock disconnection event.
+			CursorOwner->UnregisterPlayCursor(&Cursor);
+			if (Clock->GetState() != EMusicClockState::Stopped)
+			{
+				Clock->DefaultMaps.Copy(CursorOwner->GetSongMaps(), 0, Cursor.GetCurrentTick());
+				SongPosOffsetMs = Clock->CurrentSmoothedAudioRenderSongPos.SecondsIncludingCountIn * 1000.0f;
+				FreeRunStartTimeSecs = Clock->GetWorld()->GetTimeSeconds();
+			}
+			CursorOwner.Reset();
+
+			check(Clock);
+			Clock->MusicClockDisconnectedEvent.Broadcast();
+		}
+	}
+}
+
 bool FMetasoundMusicClockDriver::AttemptToConnectToAudioComponentsMetasound()
 {
 	check(IsInGameThread());
@@ -176,21 +211,7 @@ void FMetasoundMusicClockDriver::OnGeneratorIOUpdated()
 
 void FMetasoundMusicClockDriver::OnGeneratorDetached()
 {
-	check(IsInGameThread());
-	check(Clock);
-	check(IsInGameThread());
-	if (CursorOwner)
-	{
-		CursorOwner->UnregisterPlayCursor(&Cursor);
-		if (Clock->GetState() != EMusicClockState::Stopped)
-		{
-			Clock->DefaultMaps.Copy(CursorOwner->GetSongMaps(), 0, Cursor.GetCurrentTick());
-			SongPosOffsetMs = Clock->CurrentSmoothedAudioRenderSongPos.SecondsIncludingCountIn * 1000.0f;
-			FreeRunStartTimeSecs = Clock->GetWorld()->GetTimeSeconds();
-		}
-		CursorOwner.Reset();
-		Clock->MusicClockDisconnectedEvent.Broadcast();
-	}
+	ResetCursorOwner();
 }
 
 void FMetasoundMusicClockDriver::TryToRegisterPlayCursor()
@@ -208,42 +229,33 @@ void FMetasoundMusicClockDriver::TryToRegisterPlayCursor()
 		// Send a command to the OnGenerateAudio thread, where it is safe to interact with the low level generator's output read references.
 		LowLevelGenerator->OnNextBuffer([MetasoundOutputName = MetasoundOutputName, ClockWeakPtr = TWeakObjectPtr<UMusicClockComponent>(Clock), ClockDriverWeakPtr = AsWeak()](Metasound::FMetasoundGenerator& LowLevelGenerator) mutable
 			{
+				// Try to get a cursor manager from the named midi clock output.
+				TSharedPtr<FMidiPlayCursorMgr> MidiPlayCursorMgr;
 				const TOptional<Metasound::TDataReadReference<HarmonixMetasound::FMidiClock>> MidiClockRef = LowLevelGenerator.GetOutputReadReference<HarmonixMetasound::FMidiClock>(MetasoundOutputName);
 				if (const Metasound::TDataReadReference<HarmonixMetasound::FMidiClock>* MidiClock = MidiClockRef.GetPtrOrNull())
 				{
-					if (const TSharedPtr<FMidiPlayCursorMgr>& MidiPlayCursorMgr = (*MidiClock)->GetDrivingMidiPlayCursorMgr())
-					{
-						// Send the midi clock's cursor manager to the game thread, where it is safe to modify clock component state.
-						AsyncTask(ENamedThreads::GameThread, [ClockWeakPtr = MoveTemp(ClockWeakPtr), ClockDriverWeakPtr = MoveTemp(ClockDriverWeakPtr), MidiPlayCursorMgrWeakPtr = TWeakPtr<FMidiPlayCursorMgr>(MidiPlayCursorMgr)]()
-							{
-								if (UMusicClockComponent* Clock = ClockWeakPtr.Get())
-								{
-									if (TSharedPtr<FMetasoundMusicClockDriver> ClockDriver = StaticCastSharedPtr<FMetasoundMusicClockDriver>(ClockDriverWeakPtr.Pin()))
-									{
-										// Verify that the music clock component still references the clock driver.
-										if (ClockDriver == Clock->ClockDriver)
-										{
-											if (TSharedPtr<FMidiPlayCursorMgr> MidiPlayCursorMgr = MidiPlayCursorMgrWeakPtr.Pin())
-											{
-												// Verify that the clock driver's cursor has not yet been registered with the cursor manager.
-												if (ClockDriver->CursorOwner != MidiPlayCursorMgr)
-												{
-													MidiPlayCursorMgr->RegisterLowResPlayCursor(&ClockDriver->Cursor);
-													ClockDriver->CursorOwner = MoveTemp(MidiPlayCursorMgr);
-													ClockDriver->WasEverConnected = true;
-													Clock->MusicClockConnectedEvent.Broadcast();
-												}
-											}
-										}
-									}
-								}
-							});
-					}
+					MidiPlayCursorMgr = (*MidiClock)->GetDrivingMidiPlayCursorMgr();
 				}
 				else
 				{
 					UE_LOG(LogMusicClock, Verbose, TEXT("Didn't find MIDI Clock output named \"%s\" in the Metasound!"), *MetasoundOutputName.ToString());
 				}
+
+				// Send a command to the game thread, where it is safe to modify clock component state.
+				AsyncTask(ENamedThreads::GameThread, [ClockWeakPtr = MoveTemp(ClockWeakPtr), ClockDriverWeakPtr = MoveTemp(ClockDriverWeakPtr), MidiPlayCursorMgr = MoveTemp(MidiPlayCursorMgr)]() mutable
+					{
+						if (UMusicClockComponent* Clock = ClockWeakPtr.Get())
+						{
+							if (TSharedPtr<FMetasoundMusicClockDriver> ClockDriver = StaticCastSharedPtr<FMetasoundMusicClockDriver>(ClockDriverWeakPtr.Pin()))
+							{
+								// Verify that the music clock component still references the clock driver.
+								if (ClockDriver == Clock->ClockDriver)
+								{
+									ClockDriver->ResetCursorOwner(MoveTemp(MidiPlayCursorMgr));
+								}
+							}
+						}
+					});
 			});
 	}
 }
diff --git a/Engine/Plugins/Runtime/Harmonix/Source/HarmonixMetasound/Private/HarmonixMetasound/Components/MetasoundMusicClockDriver.h b/Engine/Plugins/Runtime/Harmonix/Source/HarmonixMetasound/Private/HarmonixMetasound/Components/MetasoundMusicClockDriver.h
--- a/Engine/Plugins/Runtime/Harmonix/Source/HarmonixMetasound/Private/HarmonixMetasound/Components/MetasoundMusicClockDriver.h
+++ b/Engine/Plugins/Runtime/Harmonix/Source/HarmonixMetasound/Private/HarmonixMetasound/Components/MetasoundMusicClockDriver.h
@@ -49,6 +49,8 @@ private:
 	FDelegateHandle GeneratorIOUpdatedCallbackHandle;
 	FDelegateHandle GraphChangedCallbackHandle;
 
+	void ResetCursorOwner(TSharedPtr<FMidiPlayCursorMgr> MidiPlayCursorMgr = nullptr);
+
 	bool AttemptToConnectToAudioComponentsMetasound();
 	void DetachAllCallbacks();
 	void TryToRegisterPlayCursor();
