[Harmonix] MIDI stream/clock reference hardening

diff --git a/Engine/Plugins/Runtime/Harmonix/Source/HarmonixMetasound/Private/HarmonixMetasound/DataTypes/MidiStream.cpp b/Engine/Plugins/Runtime/Harmonix/Source/HarmonixMetasound/Private/HarmonixMetasound/DataTypes/MidiStream.cpp
--- a/Engine/Plugins/Runtime/Harmonix/Source/HarmonixMetasound/Private/HarmonixMetasound/DataTypes/MidiStream.cpp
+++ b/Engine/Plugins/Runtime/Harmonix/Source/HarmonixMetasound/Private/HarmonixMetasound/DataTypes/MidiStream.cpp
@@ -14,18 +14,9 @@ namespace HarmonixMetasound
 	DEFINE_LOG_CATEGORY(LogMidiStreamDataType);
 
 	using namespace Metasound;
-
-	TSharedPtr<const FMidiClock, ESPMode::NotThreadSafe> FMidiStream::DummyClock{};
 	
 	FMidiStream::FMidiStream(const FOperatorSettings&)
 	{
-		// We have to create the dummy clock here because otherwise registration with the clock update subsystem fails when launching.
-		if (!DummyClock.IsValid())
-		{
-			DummyClock = MakeShared<FMidiClock, ESPMode::NotThreadSafe>(FOperatorSettings{ 48000, 100 });
-		}
-		
-		Clock = DummyClock;
 	}
 
 	void FMidiStream::SetClock(const FMidiClock& InClock)
@@ -35,13 +26,12 @@ namespace HarmonixMetasound
 
 	void FMidiStream::ResetClock()
 	{
-		Clock = DummyClock;
+		Clock.Reset();
 	}
 
 	TSharedPtr<const FMidiClock, ESPMode::NotThreadSafe> FMidiStream::GetClock() const
 	{
-		const TSharedPtr<const FMidiClock, ESPMode::NotThreadSafe> ClockPtr = Clock.Pin();
-		return ClockPtr.IsValid() ? ClockPtr : DummyClock;
+		return Clock.Pin();
 	}
 
 
@@ -115,14 +105,8 @@ namespace HarmonixMetasound
 	void FMidiStream::Copy(const FMidiStream& From, FMidiStream& To, const FEventFilter& Filter, const FEventTransformer& Transformer)
 	{
 		// Copy the clock from the other stream
+		if (const TSharedPtr<const FMidiClock, ESPMode::NotThreadSafe> FromClock = From.GetClock())
 		{
-			const TSharedPtr<const FMidiClock, ESPMode::NotThreadSafe> FromClock = From.GetClock();
-			
-			if (!ensure(FromClock.IsValid()))
-			{
-				return;
-			}
-
 			To.SetClock(*FromClock);
 		}
 
@@ -142,29 +126,19 @@ namespace HarmonixMetasound
 
 	void FMidiStream::Merge(const FMidiStream& From, FMidiStream& To, const FEventFilter& Filter, const FEventTransformer& Transformer)
 	{
-		// We need to make sure that the clocks match, or that one of them is using the dummy clock.
+		// We need to make sure that the clocks match, or that one of them doesn't have a clock.
 		// Otherwise a merge is invalid.
 		{
 			const auto FromClock = From.GetClock();
 			const auto ToClock = To.GetClock();
-
-			// We shouldn't get a null clock at any time
-			if (!ensure(FromClock.IsValid()) || !ensure(ToClock.IsValid()))
-			{
-				return;
-			}
-			
-			const bool FromClockIsDummy = FromClock.Get() == DummyClock.Get();
-			const bool ToClockIsDummy = ToClock.Get() == DummyClock.Get();
-			const bool ClocksAreSame = FromClock.Get() == ToClock.Get();
 			
-			if (!ClocksAreSame && !FromClockIsDummy && !ToClockIsDummy)
+			if (FromClock.Get() != ToClock.Get() && FromClock.IsValid() && ToClock.IsValid())
 			{
 				return;
 			}
 
-			// If the "to" clock is the dummy, and the "from" clock isn't, overwrite the "to" clock
-			if (!FromClockIsDummy && ToClockIsDummy)
+			// If the "to" clock is null, and the "from" clock isn't, overwrite the "to" clock
+			if (FromClock.IsValid() && !ToClock.IsValid())
 			{
 				To.SetClock(*FromClock);
 			}
diff --git a/Engine/Plugins/Runtime/Harmonix/Source/HarmonixMetasound/Public/HarmonixMetasound/DataTypes/MidiClock.h b/Engine/Plugins/Runtime/Harmonix/Source/HarmonixMetasound/Public/HarmonixMetasound/DataTypes/MidiClock.h
--- a/Engine/Plugins/Runtime/Harmonix/Source/HarmonixMetasound/Public/HarmonixMetasound/DataTypes/MidiClock.h
+++ b/Engine/Plugins/Runtime/Harmonix/Source/HarmonixMetasound/Public/HarmonixMetasound/DataTypes/MidiClock.h
@@ -63,12 +63,11 @@ namespace HarmonixMetasound
 		static constexpr int32 kMidiGranularity = 128;
 		
 		explicit FMidiClock(const Metasound::FOperatorSettings& InSettings);
+		virtual ~FMidiClock();
 
 		FMidiClock(const FMidiClock& Other);
 		FMidiClock& operator=(const FMidiClock& Other);
 
-		virtual ~FMidiClock();
-
 		void ResetAndStart(int32 FrameIndex, bool SeekToStart = true);
 
 		void PrepareBlock();
@@ -177,7 +176,7 @@ namespace HarmonixMetasound
 		void UnregisterPlayCursor(FMidiPlayCursor* PlayCursor, bool WarnOnFail = true) const      { DrivingMidiPlayCursorMgr->UnregisterPlayCursor(PlayCursor, WarnOnFail);    }
 		void UnregisterAllPlayCursors()                                                           { DrivingMidiPlayCursorMgr->UnregisterAllPlayCursors();                      }
 
-		const TSharedPtr<FMidiPlayCursorMgr>& GetDrivingMidiPlayCursorMgr() const { return DrivingMidiPlayCursorMgr; }
+		TSharedPtr<FMidiPlayCursorMgr> GetDrivingMidiPlayCursorMgr() const { return DrivingMidiPlayCursorMgr.ToSharedPtr(); }
 		//*****************************************************************************************
 
 		void WriteAdvance(int32 StartFrameIndex, int32 EndFrameIndex, float InSpeed = 1.0f);
@@ -237,7 +236,7 @@ namespace HarmonixMetasound
 		TArray<FMidiClockEvent> MidiClockEventsInBlock;
 
 		bool SmoothingEnabled = false;
-		TSharedPtr<FMidiPlayCursorMgr> DrivingMidiPlayCursorMgr;
+		TSharedRef<FMidiPlayCursorMgr> DrivingMidiPlayCursorMgr;
 	};
 
 	// Declare aliases IN the namespace...
diff --git a/Engine/Plugins/Runtime/Harmonix/Source/HarmonixMetasound/Public/HarmonixMetasound/DataTypes/MidiStream.h b/Engine/Plugins/Runtime/Harmonix/Source/HarmonixMetasound/Public/HarmonixMetasound/DataTypes/MidiStream.h
--- a/Engine/Plugins/Runtime/Harmonix/Source/HarmonixMetasound/Public/HarmonixMetasound/DataTypes/MidiStream.h
+++ b/Engine/Plugins/Runtime/Harmonix/Source/HarmonixMetasound/Public/HarmonixMetasound/DataTypes/MidiStream.h
@@ -117,7 +117,6 @@ namespace HarmonixMetasound
 		void UpdateActiveVoice(const FMidiStreamEvent& Event);
 		TSet<FMidiVoiceId> ActiveVoices;
 
-		static TSharedPtr<const FMidiClock, ESPMode::NotThreadSafe> DummyClock;
 		TWeakPtr<const FMidiClock, ESPMode::NotThreadSafe> Clock;
 
 		// Map to handle re-mapping merged MIDI events, which helps to disambiguate split/transposed notes
