Fix Externally clocked Midi Player causing note spam when "Seeking on start"

- Metronome Node, Add whether to "ResetToStart" when rebuilding MIDI Data
- && only rebuild MIDI Data on start if we didn't already do it when we seeked
- && Only register metronome clock when looping
- MIDI Player Node should Seek when Advance is Preroll
- && checks if a seek was performed during a start

#tests PIE

diff --git a/Engine/Plugins/Runtime/Harmonix/Source/HarmonixMetasound/Private/HarmonixMetasound/Nodes/MetronomeNode.cpp b/Engine/Plugins/Runtime/Harmonix/Source/HarmonixMetasound/Private/HarmonixMetasound/Nodes/MetronomeNode.cpp
--- a/Engine/Plugins/Runtime/Harmonix/Source/HarmonixMetasound/Private/HarmonixMetasound/Nodes/MetronomeNode.cpp
+++ b/Engine/Plugins/Runtime/Harmonix/Source/HarmonixMetasound/Private/HarmonixMetasound/Nodes/MetronomeNode.cpp
@@ -84,7 +84,7 @@ namespace HarmonixMetasound
 		int32		 LastClockTickUpdate = -1;
 		int32		 LoopOffsetTick = 0;
 
-		void BuildMidiData();
+		void BuildMidiData(bool ResetToStart = true);
 		void UpdateMidi();
 		void AddTempoChangeForMidi(float TempoBPM);
 		void AddTimeSigChangeForMidi(int32 TimeSigNum, int32 TimeSigDenom);
@@ -197,10 +197,9 @@ namespace HarmonixMetasound
 		, CurrentTimeSigNum(FMath::Clamp(*TimeSigNumInPin, 1, 64))
 		, CurrentTimeSigDenom(FMath::Clamp(*TimeSigDenomInPin, 1, 64))
 	{
-		MetronomeClock.RegisterHiResPlayCursor(this);
-
 		if (LoopInPin)
 		{
+			MetronomeClock.RegisterHiResPlayCursor(this);
 			MidiClockOutPin->AttachToTimeAuthority(MetronomeClock);
 		}
 
@@ -304,6 +303,21 @@ namespace HarmonixMetasound
 	{
 		FMidiPlayCursor::AdvanceThruTick(InTick, IsPreRoll);
 
+		// if this is a pre-roll, perform a seek instead
+		// NOTE: This code is/should be identical to SeekThruTick above
+		if (IsPreRoll)
+		{
+			UpdateLoopOffsetTickFromTick(InTick);
+			int32 MappedTick = InTick - LoopOffsetTick;
+
+			float Ms = MidiClockOutPin->GetSongMaps().TickToMs(InTick + 1);
+			FMusicSeekTarget SeekTarget;
+			SeekTarget.Type = ESeekPointType::Millisecond;
+			SeekTarget.Ms = Ms;
+			MidiClockOutPin->SeekTo(MetronomeClock.GetCurrentBlockFrameIndex(), SeekTarget, SeekPreRollBarsInPin);
+			return;
+		}
+
 		int32 MappedTick = InTick - LoopOffsetTick;
 		float CurrentSpeed = MetronomeClock.GetSpeedAtBlockSampleFrame(MetronomeClock.GetCurrentBlockFrameIndex());
 		float Ms = MidiClockOutPin->GetSongMaps().TickToMs(MappedTick);
@@ -398,18 +412,6 @@ namespace HarmonixMetasound
 			UpdateMidi();
 			LastClockTickUpdate = ClockTick;
 		}
-		else if (ClockTick == -1 && LastClockTickUpdate == -1)
-		{
-			// check if we need to update our midi data based on new inputs
-			int32 InTimeSigNum = FMath::Clamp(*TimeSigNumInPin, 1, 64);
-			int32 InTimeSigDenom = FMath::Clamp(*TimeSigDenomInPin, 1, 64);
-			bool RebuildMidi = *TempoInPin > 0 && !FMath::IsNearlyEqual(CurrentTempo, *TempoInPin);
-			RebuildMidi |= InTimeSigNum != CurrentTimeSigNum || InTimeSigDenom != CurrentTimeSigDenom;
-			if (RebuildMidi)
-			{
-				BuildMidiData();
-			}
-		}
 
 		TransportSpanProcessor TransportHandler = [this, &DrivingMidiClock](int32 StartFrameIndex, int32 EndFrameIndex, EMusicPlayerTransportState CurrentState)
 		{
@@ -424,13 +426,10 @@ namespace HarmonixMetasound
 				return CurrentState;
 
 			case EMusicPlayerTransportState::Starting:
-				// We need to rebuild the midi data as there may have been a previous 'play' 
-				// and during that 'play' there may have been tempo changes added. When doing a
-				// 'start from a stop' we want a clean slate...
-				BuildMidiData();
 				// Play from the beginning if we haven't received a seek call while we were stopped...
 				if (!ReceivedSeekWhileStopped())
 				{
+					BuildMidiData(true);
 					LastClockTickUpdate = -1;
 				}
 				DrivingMidiClock.ResetAndStart(StartFrameIndex, !ReceivedSeekWhileStopped());
@@ -442,7 +441,7 @@ namespace HarmonixMetasound
 				return EMusicPlayerTransportState::Playing;
 
 			case EMusicPlayerTransportState::Seeking:
-				BuildMidiData();
+				BuildMidiData(false);
 				DrivingMidiClock.SeekTo(StartFrameIndex, TransportInPin->GetNextSeekDestination(), SeekPreRollBarsInPin);
 				LastClockTickUpdate = DrivingMidiClock.GetCurrentMidiTick();
 				// Here we will return that we want to be in the same state we were in before this request to 
@@ -504,7 +503,7 @@ namespace HarmonixMetasound
 		}
 	}
 
-	void FMetronomeOperator::BuildMidiData()
+	void FMetronomeOperator::BuildMidiData(bool ResetToStart)
 	{
 		// make sure we have valid values
 		CurrentTempo = FMath::Max(1.0f, *TempoInPin);
@@ -515,12 +514,12 @@ namespace HarmonixMetasound
 		
 		if (LoopInPin)
 		{
-			MetronomeClock.AttachToMidiResource(MidiData);
+			MetronomeClock.AttachToMidiResource(MidiData, ResetToStart);
 
 			// midi clock out will follow the tempo of the metronome
 			// so just assign it to some reasonable values. They will be ignored
 			TSharedPtr<FMidiFileData> MidiDataOut = FMidiClock::MakeClockConductorMidiData(120.0f, CurrentTimeSigNum, CurrentTimeSigDenom);
-			MidiClockOutPin->AttachToMidiResource(MidiDataOut);
+			MidiClockOutPin->AttachToMidiResource(MidiDataOut, ResetToStart);
 			
 			int32 LoopEndTick = MidiDataOut->SongMaps.GetBarMap().BarIncludingCountInToTick(FMath::Max(LoopLengthBarsInPin, 1));
 			MidiClockOutPin->SetLoop(0, LoopEndTick);
diff --git a/Engine/Plugins/Runtime/Harmonix/Source/HarmonixMetasound/Private/HarmonixMetasound/Nodes/MidiPlayerNode.cpp b/Engine/Plugins/Runtime/Harmonix/Source/HarmonixMetasound/Private/HarmonixMetasound/Nodes/MidiPlayerNode.cpp
--- a/Engine/Plugins/Runtime/Harmonix/Source/HarmonixMetasound/Private/HarmonixMetasound/Nodes/MidiPlayerNode.cpp
+++ b/Engine/Plugins/Runtime/Harmonix/Source/HarmonixMetasound/Private/HarmonixMetasound/Nodes/MidiPlayerNode.cpp
@@ -341,7 +341,7 @@ namespace HarmonixMetasound
 			case EMusicPlayerTransportState::Starting:
 			case EMusicPlayerTransportState::Playing:
 			case EMusicPlayerTransportState::Continuing:
-				MidiClockOut->ResetAndStart(0, true);
+				MidiClockOut->ResetAndStart(0, !ReceivedSeekWhileStopped());
 				return EMusicPlayerTransportState::Playing;
 
 			case EMusicPlayerTransportState::Seeking: // seeking is omitted from init, shouldn't happen
@@ -410,14 +410,9 @@ namespace HarmonixMetasound
 					{
 					case FMidiClockEvent::EType::Reset:
 					{
-						UpdateLoopOffsetTickFromTick(Event.Tick2);
-						int32 Tick = Event.Tick2 - LoopOffsetTick;
-
-						float Ms = MidiClockOut->GetSongMaps().TickToMs(Tick);
-						FMusicSeekTarget SeekTarget;
-						SeekTarget.Type = ESeekPointType::Millisecond;
-						SeekTarget.Ms = Ms;
-						MidiClockOut->SeekTo(Event.BlockFrameIndex, SeekTarget, PrerollBars);
+						// Ignore resets, because every reset event is preceded by a SeekThru event
+						// SeekThru triggers a new reset in the following clock,
+						// resulting in an exponentially growing list of reset and seek thru events
 						break;
 					}
 					case FMidiClockEvent::EType::Loop:
@@ -451,7 +446,20 @@ namespace HarmonixMetasound
 					}
 					case FMidiClockEvent::EType::AdvanceThru:
 					{
-						
+						// if this advance is a preroll, perform a seek instead so we don't trigger events doing an advance
+						// NOTE: This code is/should be identical to the SeekThru event above
+						if (Event.IsPreRoll)
+						{
+							UpdateLoopOffsetTickFromTick(Event.Tick2);
+							int32 Tick = Event.Tick2 - LoopOffsetTick;
+
+							float Ms = MidiClockOut->GetSongMaps().TickToMs(Tick + 1);
+							FMusicSeekTarget SeekTarget;
+							SeekTarget.Type = ESeekPointType::Millisecond;
+							SeekTarget.Ms = Ms;
+							MidiClockOut->SeekTo(Event.BlockFrameIndex, SeekTarget, PrerollBars);
+							break;
+						}
 						// The MidiClock handles looping on its own, so we can conveniently advance it
 						int32 Tick = Event.Tick2 - LoopOffsetTick;
 						float Ms = MidiClockOut->GetSongMaps().TickToMs(Tick);
