Don't Reset ActiveVoices on MidiStream Copy
#tests PIE

diff --git a/Engine/Plugins/Runtime/Harmonix/Source/HarmonixMetasound/Private/HarmonixMetasound/DataTypes/MidiStream.cpp b/Engine/Plugins/Runtime/Harmonix/Source/HarmonixMetasound/Private/HarmonixMetasound/DataTypes/MidiStream.cpp
--- a/Engine/Plugins/Runtime/Harmonix/Source/HarmonixMetasound/Private/HarmonixMetasound/DataTypes/MidiStream.cpp
+++ b/Engine/Plugins/Runtime/Harmonix/Source/HarmonixMetasound/Private/HarmonixMetasound/DataTypes/MidiStream.cpp
@@ -112,7 +112,7 @@ namespace HarmonixMetasound
 
 		// Reset the target
 		To.EventsInBlock.Reset();
-		To.ActiveVoices.Reset();
+		To.ActiveVoices = From.ActiveVoices;
 		
 		// Copy the events
 		for (const FMidiStreamEvent& Event : From.GetEventsInBlock())
