FMidiPlayCursorMgr::UnregisterPlayCursor calls FMidiPlayCursor::SetOwner inside the appropriate mutex-protected scope.

diff --git a/Engine/Plugins/Runtime/Harmonix/Source/HarmonixMidi/Private/HarmonixMidi/MidiPlayCursorMgr.cpp b/Engine/Plugins/Runtime/Harmonix/Source/HarmonixMidi/Private/HarmonixMidi/MidiPlayCursorMgr.cpp
--- a/Engine/Plugins/Runtime/Harmonix/Source/HarmonixMidi/Private/HarmonixMidi/MidiPlayCursorMgr.cpp
+++ b/Engine/Plugins/Runtime/Harmonix/Source/HarmonixMidi/Private/HarmonixMidi/MidiPlayCursorMgr.cpp
@@ -24,8 +24,8 @@ FMidiPlayCursorMgr::FMidiPlayCursorMgr()
 
 void FMidiPlayCursorMgr::Reset()
 {
-	FScopeLock LowResLock(&LowResCursorListLock);
-	FScopeLock HiResLock(&HiResCursorListLock);
+	FScopeLock LowResLock(&GetLowResTracker().CursorListLock);
+	FScopeLock HiResLock(&GetHiResTracker().CursorListLock);
 	UnregisterAllPlayCursors();
 	ResetTrackers();
 	MidiFileData = nullptr;
@@ -120,8 +120,8 @@ FMidiSongPos FMidiPlayCursorMgr::CalculateSongPosWithOffsetMs(float DeltaMs, boo
 
 void FMidiPlayCursorMgr::DetermineLength()
 {
-	FScopeLock LowResLock(&LowResCursorListLock);
-	FScopeLock HiResLock(&HiResCursorListLock);
+	FScopeLock LowResLock(&GetLowResTracker().CursorListLock);
+	FScopeLock HiResLock(&GetHiResTracker().CursorListLock);
 	if (!MidiFileData)
 	{
 		LengthMs = 0.0f;
@@ -171,7 +171,7 @@ void FMidiPlayCursorMgr::RegisterHiResPlayCursor(FMidiPlayCursor* PlayCursor, fl
 		PlayCursor->Owner->UnregisterPlayCursor(PlayCursor);
 	}
 	PlayCursor->UnregisterASAP = false;
-	FScopeLock HiResLock(&HiResCursorListLock);
+	FScopeLock HiResLock(&GetHiResTracker().CursorListLock);
 	PlayCursor->SetOwner(this, &GetHiResTracker(), PreRollMs);
 	GetHiResTracker().AddCursor(PlayCursor);
 }
@@ -183,7 +183,7 @@ void FMidiPlayCursorMgr::RegisterLowResPlayCursor(FMidiPlayCursor* PlayCursor, f
 		PlayCursor->Owner->UnregisterPlayCursor(PlayCursor);
 	}
 	PlayCursor->UnregisterASAP = false;
-	FScopeLock LowResLock(&LowResCursorListLock);
+	FScopeLock LowResLock(&GetLowResTracker().CursorListLock);
 	PlayCursor->SetOwner(this, &GetLowResTracker(), PreRollMs);
 	GetLowResTracker().AddCursor(PlayCursor);
 }
@@ -216,81 +216,59 @@ void FMidiPlayCursorMgr::RecalculatePreRollDueToCursorPosition(FMidiPlayCursor*
 
 void FMidiPlayCursorMgr::UnregisterPlayCursor(FMidiPlayCursor* PlayCursor, bool WarnOnFail)
 {
-	bool WasHiRes = false;
+	// search the hi-res list... then the low res...
+	for (int32 Index = 0; Index != 2; ++Index)
 	{
-		FScopeLock HiResLock(&HiResCursorListLock);
-		if (TraversingHiResCursors)
+		FScopeLock Lock(&Trackers[Index].CursorListLock);
+		if (Trackers[Index].TraversingCursors)
 		{
-			WasHiRes = PlayCursor->Tracker == &GetHiResTracker();
-		}
-		else
-		{
-			WasHiRes = GetHiResTracker().RemoveCursor(PlayCursor);
-		}
-		if (WasHiRes && TraversingHiResCursors)
-		{
-			UE_LOG(LogMidi, Error, TEXT("You cannot unregister a MidiPlayCursor while this manager is currently traversing it's cursor list!"));
-			return;
-		}
-	}
-	if (!WasHiRes)
-	{
-		// wasn't on hi-res list... try the low res...
-		FScopeLock LowResLock(&LowResCursorListLock);
-		bool WasLowRes = false;
-		if (TraversingLowResCursors)
-		{
-			WasLowRes = PlayCursor->Tracker == &GetLowResTracker();
-		}
-		else
-		{
-			WasLowRes = GetLowResTracker().RemoveCursor(PlayCursor);
-		}
-		if (!WasLowRes)
-		{
-			if (WarnOnFail)
+			if (PlayCursor->Tracker == &Trackers[Index])
 			{
-				UE_LOG(LogMidi, Warning, TEXT("Attempt to remove MidiPlayCursor from Manager that didn't own it!"));
+				UE_LOG(LogMidi, Error, TEXT("You cannot unregister a MidiPlayCursor while the cursor's manager is traversing its cursor list!"));
+				return;
 			}
-			return;
+			continue;
 		}
-		if (WasLowRes && TraversingLowResCursors)
+		if (Trackers[Index].RemoveCursor(PlayCursor))
 		{
-			UE_LOG(LogMidi, Error, TEXT("You cannot unregister a MidiPlayCursor while this manager is currently traversing it's cursor list!"));
+			// it was removed from one or the other lists...
+			PlayCursor->SetOwner(nullptr, nullptr);
 			return;
 		}
 	}
-	// it was removed from one or the other lists...
-	PlayCursor->SetOwner(nullptr, nullptr);
+	if (WarnOnFail)
+	{
+		UE_LOG(LogMidi, Warning, TEXT("Attempt to remove MidiPlayCursor from a manager that doesn't own it!"));
+	}
 }
 
 void FMidiPlayCursorMgr::UnregisterAllPlayCursors()
 {
-	FScopeLock LowResLock(&LowResCursorListLock);
-	FScopeLock HiResLock(&HiResCursorListLock);
+	FScopeLock LowResLock(&GetLowResTracker().CursorListLock);
+	FScopeLock HiResLock(&GetHiResTracker().CursorListLock);
 
-	if (TraversingHiResCursors || TraversingLowResCursors)
+	if (GetHiResTracker().TraversingCursors || GetLowResTracker().TraversingCursors)
 	{
 		UE_LOG(LogMidi, Error, TEXT("You cannot unregister all MidiPlayCursors while this manager is currently traversing it's cursor list!"));
 		return;
 	}
 
-	TraversingHiResCursors = true;
+	GetHiResTracker().TraversingCursors = true;
 	for (auto it = GetHiResTracker().Cursors.begin(); it != GetHiResTracker().Cursors.end(); ++it)
 	{
 		it.GetNode()->ManagerIsDetaching();
 		it.GetNode()->SetOwner(nullptr, nullptr);
 	}
 	GetHiResTracker().Clear();
-	TraversingHiResCursors = false;
-	TraversingLowResCursors = true;
+	GetHiResTracker().TraversingCursors = false;
+	GetLowResTracker().TraversingCursors = true;
 	for (auto it = GetLowResTracker().Cursors.begin(); it != GetLowResTracker().Cursors.end(); ++it)
 	{
 		it.GetNode()->ManagerIsDetaching();
 		it.GetNode()->SetOwner(nullptr, nullptr);
 	}
 	GetLowResTracker().Clear();
-	TraversingLowResCursors = false;
+	GetLowResTracker().TraversingCursors = false;
 }
 
 bool FMidiPlayCursorMgr::HasLowResCursors() const
@@ -300,8 +278,8 @@ bool FMidiPlayCursorMgr::HasLowResCursors() const
 
 void FMidiPlayCursorMgr::ResetTrackers()
 {
-	FScopeLock LowResLock(&LowResCursorListLock);
-	FScopeLock HiResLock(&HiResCursorListLock);
+	FScopeLock LowResLock(&GetLowResTracker().CursorListLock);
+	FScopeLock HiResLock(&GetHiResTracker().CursorListLock);
 	float Ms = GetTempoMap().TickToMs(-1);
 	for (FMidiPlayCursorTracker& Tracker : Trackers)
 	{
@@ -372,8 +350,8 @@ void FMidiPlayCursorMgr::GetCursorExtentsTicks(int32& Earliest, int32& Latest) c
 
 bool FMidiPlayCursorMgr::CursorsAllInPhase() const
 {
-	FScopeLock LowResLock(&LowResCursorListLock);
-	FScopeLock HiResLock(&HiResCursorListLock);
+	FScopeLock LowResLock(&GetLowResTracker().CursorListLock);
+	FScopeLock HiResLock(&GetHiResTracker().CursorListLock);
 
 	return CursorsAllInPhaseImpl(false) && CursorsAllInPhaseImpl(true);
 }
@@ -432,14 +410,14 @@ void FMidiPlayCursorMgr::SetLoop(int32 StartTick, int32 EndTick, bool IsDirectMa
 	// might fall in the middle of the current span of leading and lagging play cursors,
 	// which would result in ugly behavior!
 
-	FScopeLock HiResLock(&HiResCursorListLock);
+	FScopeLock HiResLock(&GetHiResTracker().CursorListLock);
 	DirectMappedTimeFollower = IsDirectMappedFollower;
 	SetLoopImpl(StartTick, EndTick, IgnoringLookAhead, false);
 }
 
 void FMidiPlayCursorMgr::ClearLoop(bool IgnoringLookAhead)
 {
-	FScopeLock HiResLock(&HiResCursorListLock);
+	FScopeLock HiResLock(&GetHiResTracker().CursorListLock);
 	ClearLoopImpl(IgnoringLookAhead, false);
 }
 
@@ -495,8 +473,8 @@ void FMidiPlayCursorMgr::ClearLoopImpl(bool IgnoringLookAhead, bool IsLowRes)
 
 void FMidiPlayCursorMgr::SeekTo(int32 Tick, int32 PreRollBars, bool IsRenderThread, bool IsLoop)
 {
-	FScopeLock LowResLock(&LowResCursorListLock);
-	FScopeLock HiResLock(&HiResCursorListLock);
+	FScopeLock LowResLock(&GetLowResTracker().CursorListLock);
+	FScopeLock HiResLock(&GetHiResTracker().CursorListLock);
 
 	// Back up one tick, as we are going to be setting the cursors'
 	// 'played through' position.
@@ -549,7 +527,7 @@ void FMidiPlayCursorMgr::SeekTo(int32 Tick, int32 PreRollBars, bool IsRenderThre
 
 void FMidiPlayCursorMgr::MoveToLoopStart()
 {
-	FScopeLock HiResLock(&HiResCursorListLock);
+	FScopeLock HiResLock(&GetHiResTracker().CursorListLock);
 	{
 		if (ensureMsgf(GetHiResTracker().Loop, TEXT("That's odd. Asked to move to the beginning of the loop... but there is no loop!")))
 		{
@@ -566,7 +544,7 @@ void FMidiPlayCursorMgr::MoveToLoopStart()
 			// The low res cursor will need to advance to the end of the loop, and then
 			// advance from loop-start to current position at some point in the future
 			// (during a low-res frame poll)
-			FScopeLock LowResLock(&LowResCursorListLock);
+			FScopeLock LowResLock(&GetLowResTracker().CursorListLock);
 			MsSinceLowResUpdate += GetHiResTracker().LoopEndMs - GetLowResTracker().CurrentMs;
 		}
 	}
@@ -579,7 +557,7 @@ bool FMidiPlayCursorMgr::IsDone() const
 		return false;
 	}
 
-	TraversingHiResCursors = true;
+	GetHiResTracker().TraversingCursors = true;
 	// This next line of ugliness is required because const iterating through a const linked list is broken
 	// and I don't have time to figure out what about the stack of templates is busted.
 	TIntrusiveDoubleLinkedList<FMidiPlayCursor>& NonConstHiResCursorList = const_cast<TIntrusiveDoubleLinkedList<FMidiPlayCursor>&>(GetHiResTracker().Cursors);
@@ -587,12 +565,12 @@ bool FMidiPlayCursorMgr::IsDone() const
 	{
 		if (!it.GetNode()->IsDone())
 		{
-			TraversingHiResCursors = false;
+			GetHiResTracker().TraversingCursors = false;
 			return false;
 		}
 	}
-	TraversingHiResCursors = false;
-	TraversingLowResCursors = true;
+	GetHiResTracker().TraversingCursors = false;
+	GetLowResTracker().TraversingCursors = true;
 	// This next line of ugliness is required because const iterating through a const linked list is broken
 	// and I don't have time to figure out what about the stack of templates is busted.
 	TIntrusiveDoubleLinkedList<FMidiPlayCursor>& NonConstLowResCursorList = const_cast<TIntrusiveDoubleLinkedList<FMidiPlayCursor>&>(GetLowResTracker().Cursors);
@@ -600,11 +578,11 @@ bool FMidiPlayCursorMgr::IsDone() const
 	{
 		if (!it.GetNode()->IsDone())
 		{
-			TraversingLowResCursors = false;
+			GetLowResTracker().TraversingCursors = false;
 			return false;
 		}
 	}
-	TraversingLowResCursors = false;
+	GetLowResTracker().TraversingCursors = false;
 	return true;
 }
 
@@ -613,8 +591,8 @@ void FMidiPlayCursorMgr::LockForMidiDataChanges()
 	if (ensureAlways(!InMidiChangeLock))
 	{
 		InMidiChangeLock = true;
-		LowResCursorListLock.Lock();
-		HiResCursorListLock.Lock();
+		GetLowResTracker().CursorListLock.Lock();
+		GetHiResTracker().CursorListLock.Lock();
 	}
 }
 
@@ -644,7 +622,7 @@ void FMidiPlayCursorMgr::MidiDataChangeComplete(EMidiChangePositionCorrectMode P
 		if (PreRollBars == 0)
 		{
 			// Now fix up the hires cursors...
-			TraversingHiResCursors = true;
+			GetHiResTracker().TraversingCursors = true;
 			for (auto it = GetHiResTracker().Cursors.begin(); it != GetHiResTracker().Cursors.end(); ++it)
 			{
 				it.GetNode()->RecalcNextEventsDueToMidiChanges(PositionMode);
@@ -657,9 +635,9 @@ void FMidiPlayCursorMgr::MidiDataChangeComplete(EMidiChangePositionCorrectMode P
 			{
 				GetHiResTracker().CurrentTick = GetTempoMap().MsToTick(GetHiResTracker().CurrentMs);
 			}
-			TraversingHiResCursors = false;
+			GetHiResTracker().TraversingCursors = false;
 			// Now the low res cursors...
-			TraversingLowResCursors = true;
+			GetLowResTracker().TraversingCursors = true;
 			for (auto it = GetLowResTracker().Cursors.begin(); it != GetLowResTracker().Cursors.end(); ++it)
 			{
 				it.GetNode()->RecalcNextEventsDueToMidiChanges(PositionMode);
@@ -688,15 +666,15 @@ void FMidiPlayCursorMgr::MidiDataChangeComplete(EMidiChangePositionCorrectMode P
 					HiResLoopedSinceLastLoResUpdate = false;
 				}
 			}
-			TraversingLowResCursors = false;
+			GetLowResTracker().TraversingCursors = false;
 		}
 		else
 		{
 			int32 SeekTick = PositionMode == EMidiChangePositionCorrectMode::MaintainTick ? GetHiResTracker().CurrentTick : GetTempoMap().MsToTick(GetHiResTracker().CurrentMs);
 			SeekTo(SeekTick, PreRollBars, true, false);
 		}
-		HiResCursorListLock.Unlock();
-		LowResCursorListLock.Unlock();
+		GetHiResTracker().CursorListLock.Unlock();
+		GetLowResTracker().CursorListLock.Unlock();
 		InMidiChangeLock = false;
 	}
 }
@@ -719,14 +697,14 @@ void FMidiPlayCursorMgr::AdvanceHiResToMs(float Ms, bool Broadcast)
 
 void FMidiPlayCursorMgr::AdvanceHiResByDeltaMs(float Ms, bool Broadcast)
 {
-	FScopeLock HiResLock(&HiResCursorListLock);
+	FScopeLock HiResLock(&GetHiResTracker().CursorListLock);
 	HiResLoopedSinceLastLoResUpdate = AdvanceTrackerByDeltaMs(Ms, GetHiResTracker(), Broadcast) || HiResLoopedSinceLastLoResUpdate;
 	MsSinceLowResUpdate += Ms;
 }
 
 void FMidiPlayCursorMgr::AdvanceHiResByDeltaTick(int32 NumTicks, bool Broadcast)
 {
-	FScopeLock HiResLock(&HiResCursorListLock);
+	FScopeLock HiResLock(&GetHiResTracker().CursorListLock);
 
 	int32 ThruTick = NumTicks + GetHiResTracker().CurrentTick;
 
@@ -737,7 +715,7 @@ void FMidiPlayCursorMgr::AdvanceHiResByDeltaTick(int32 NumTicks, bool Broadcast)
 
 void FMidiPlayCursorMgr::AdvanceHiResThruTick(int32 ThruTick, bool Broadcast, bool DontAdvancePastLoopEnd)
 {
-	FScopeLock HiResLock(&HiResCursorListLock);
+	FScopeLock HiResLock(&GetHiResTracker().CursorListLock);
 
 	// early out?...
 	if (ThruTick == GetHiResTracker().CurrentTick)
@@ -763,7 +741,7 @@ void FMidiPlayCursorMgr::AdvanceHiResThruTick(int32 ThruTick, bool Broadcast, bo
 
 bool FMidiPlayCursorMgr::AdvanceLowResCursors()
 {
-	FScopeLock LowResLock(&LowResCursorListLock);
+	FScopeLock LowResLock(&GetLowResTracker().CursorListLock);
 
 	GetLowResTracker().HandleQueuedReset();
 
@@ -775,7 +753,7 @@ bool FMidiPlayCursorMgr::AdvanceLowResCursors()
 	bool HiResLoop = false;
 	bool HiResLoopIgnoringLookAhead = false;
 	{
-		FScopeLock HiResLock(&HiResCursorListLock);
+		FScopeLock HiResLock(&GetHiResTracker().CursorListLock);
 
 		MsDiff = MsSinceLowResUpdate;
 		HiResCurrentMs = GetHiResTracker().CurrentMs;
@@ -825,7 +803,7 @@ bool FMidiPlayCursorMgr::AdvanceLowResCursors()
 
 void FMidiPlayCursorMgr::UpdateLowResCursors(FMidiPlayCursorTracker& Tracker)
 {
-	TraversingLowResCursors = true;
+	GetLowResTracker().TraversingCursors = true;
 	for (auto it = Tracker.Cursors.begin(); it != Tracker.Cursors.end();)
 	{
 		if (!it.GetNode()->UpdateWithTrackerUnchanged())
@@ -840,7 +818,7 @@ void FMidiPlayCursorMgr::UpdateLowResCursors(FMidiPlayCursorTracker& Tracker)
 			++it;
 		}
 	}
-	TraversingLowResCursors = false;
+	GetLowResTracker().TraversingCursors = false;
 }
 
 bool FMidiPlayCursorMgr::AdvanceTrackerByDeltaMs(float Ms, FMidiPlayCursorTracker& Tracker, bool IsLowRes, bool Broadcast)
@@ -866,8 +844,7 @@ bool FMidiPlayCursorMgr::AdvanceTrackerByDeltaMs(float Ms, FMidiPlayCursorTracke
 		Tracker.CurrentMs = NewMs;
 	}
 
-	bool* TraversingFlag = IsLowRes ? &TraversingLowResCursors : &TraversingHiResCursors;
-	*TraversingFlag = true;
+	Tracker.TraversingCursors = true;
 	for (auto it = Tracker.Cursors.begin(); it != Tracker.Cursors.end();)
 	{
 		if (!it.GetNode()->Advance(IsLowRes))
@@ -882,7 +859,7 @@ bool FMidiPlayCursorMgr::AdvanceTrackerByDeltaMs(float Ms, FMidiPlayCursorTracke
 			++it;
 		}
 	}
-	*TraversingFlag = false;
+	Tracker.TraversingCursors = false;
 
 	return Looped;
 }
@@ -898,8 +875,7 @@ void FMidiPlayCursorMgr::AdvanceTrackerThruTick(int32 ToTick, FMidiPlayCursorTra
 	Tracker.CurrentTick = NewTick;
 	Tracker.CurrentMs = NewMs;
 
-	bool* TraversingFlag = IsLowRes ? &TraversingLowResCursors : &TraversingHiResCursors;
-	*TraversingFlag = true;
+	Tracker.TraversingCursors = true;
 	for (auto it = Tracker.Cursors.begin(); it != Tracker.Cursors.end();)
 	{
 		if (!it.GetNode()->Advance(IsLowRes))
@@ -914,5 +890,5 @@ void FMidiPlayCursorMgr::AdvanceTrackerThruTick(int32 ToTick, FMidiPlayCursorTra
 			++it;
 		}
 	}
-	*TraversingFlag = false;
+	Tracker.TraversingCursors = false;
 }
diff --git a/Engine/Plugins/Runtime/Harmonix/Source/HarmonixMidi/Private/HarmonixMidi/MidiPlayCursorTracker.cpp b/Engine/Plugins/Runtime/Harmonix/Source/HarmonixMidi/Private/HarmonixMidi/MidiPlayCursorTracker.cpp
--- a/Engine/Plugins/Runtime/Harmonix/Source/HarmonixMidi/Private/HarmonixMidi/MidiPlayCursorTracker.cpp
+++ b/Engine/Plugins/Runtime/Harmonix/Source/HarmonixMidi/Private/HarmonixMidi/MidiPlayCursorTracker.cpp
@@ -19,6 +19,7 @@ FMidiPlayCursorTracker::FMidiPlayCursorTracker(bool InIsLowRes)
 	, LoopEndTick(0)
 	, Loop(false)
 	, LoopIgnoringLookAhead(false)
+	, TraversingCursors(false)
 	, IsLowRes(InIsLowRes)
 {
 }
diff --git a/Engine/Plugins/Runtime/Harmonix/Source/HarmonixMidi/Public/HarmonixMidi/MidiPlayCursorMgr.h b/Engine/Plugins/Runtime/Harmonix/Source/HarmonixMidi/Public/HarmonixMidi/MidiPlayCursorMgr.h
--- a/Engine/Plugins/Runtime/Harmonix/Source/HarmonixMidi/Public/HarmonixMidi/MidiPlayCursorMgr.h
+++ b/Engine/Plugins/Runtime/Harmonix/Source/HarmonixMidi/Public/HarmonixMidi/MidiPlayCursorMgr.h
@@ -177,10 +177,6 @@ private:
 	TSharedPtr<FMidiFileData> MidiFileData;
 	const FSongMaps* SongMaps;     // <-- This is a pointer to the actual midi song map data... allows for fast access
 
-	mutable FCriticalSection LowResCursorListLock;
-	mutable FCriticalSection HiResCursorListLock;
-	mutable bool TraversingHiResCursors  = false;
-	mutable bool TraversingLowResCursors = false;
 	float LengthMs;
 	int32 LengthTicks;
 	bool  DirectMappedTimeFollower;
diff --git a/Engine/Plugins/Runtime/Harmonix/Source/HarmonixMidi/Public/HarmonixMidi/MidiPlayCursorTracker.h b/Engine/Plugins/Runtime/Harmonix/Source/HarmonixMidi/Public/HarmonixMidi/MidiPlayCursorTracker.h
--- a/Engine/Plugins/Runtime/Harmonix/Source/HarmonixMidi/Public/HarmonixMidi/MidiPlayCursorTracker.h
+++ b/Engine/Plugins/Runtime/Harmonix/Source/HarmonixMidi/Public/HarmonixMidi/MidiPlayCursorTracker.h
@@ -2,11 +2,14 @@
 #pragma once
 
 #include "Containers/IntrusiveDoubleLinkedList.h"
+#include "HAL/CriticalSection.h"
 
 class FMidiPlayCursor;
 
 struct FMidiPlayCursorTracker
 {
+	mutable FCriticalSection CursorListLock;
+
 	int32   CurrentTick; // We've broadcast all events up through this tick
 	float   CurrentMs;   // We've broadcast all events up through this Ms
 	float   ElapsedMs;
@@ -27,7 +30,9 @@ struct FMidiPlayCursorTracker
 	bool    Loop;
 	bool    LoopIgnoringLookAhead;
 
-	bool    IsLowRes;
+	mutable bool TraversingCursors;
+
+	const bool IsLowRes;
 
 	TIntrusiveDoubleLinkedList<FMidiPlayCursor> Cursors;
 
