MIDI File Import Hardening:
- Lots of changes to FStdMidiFileReader to prevent writing to unallocated memory.
- Lots of error checking and messaging when import midi is malformed.
- Graceful failure when midi import fails.
- A little work on conforming midi to even bar lengths (MORE TO DO).
- Added some functions to get the data size of individual tracks and the entire midi file.
- Changed LogMidi to LogMIDI
- Added bool return values to almost all IMidiReceiver virtual functions so receivers can halt the import of midi data if they see something they don't like, or they run out of memory, etc.
- Added some data to the FMidiFileDetailCustomization to show more information about what is in the midi file in the editor.
- Added items to the context menu for Standard MIDI File asses in the editor so now MIDI assests can be exported as '.mid' files, compared, and viewed in an external editor.

diff --git a/Engine/Plugins/Runtime/Harmonix/Source/HarmonixMidi/Private/HarmonixMidi/BarMap.cpp b/Engine/Plugins/Runtime/Harmonix/Source/HarmonixMidi/Private/HarmonixMidi/BarMap.cpp
--- a/Engine/Plugins/Runtime/Harmonix/Source/HarmonixMidi/Private/HarmonixMidi/BarMap.cpp
+++ b/Engine/Plugins/Runtime/Harmonix/Source/HarmonixMidi/Private/HarmonixMidi/BarMap.cpp
@@ -8,6 +8,22 @@
 #include "Serialization/Archive.h"
 
 
+bool FBarMap::operator==(const FBarMap& Other) const
+{
+	if (StartBar != Other.StartBar || TicksPerQuarterNote != Other.TicksPerQuarterNote || Points.Num() != Other.Points.Num())
+	{
+		return false;
+	}
+	for (int32 PointIndex = 0; PointIndex < Points.Num(); ++PointIndex)
+	{
+		if (Points[PointIndex] != Other.Points[PointIndex])
+		{
+			return false;
+		}
+	}
+	return true;
+}
+
 void FBarMap::Empty()
 {
 	Points.Empty();
@@ -31,12 +47,12 @@ float FBarMap::MusicTimestampToTick(const FMusicTimestamp& Timestamp) const
 	{
 		if (Beat >= 0.0f)
 		{
-			UE_LOG(LogMidi, Warning, TEXT("Beat %f specified in Music Timestamp! Beats in Music Timestamps are '1' based. Beat 1.0 is the first beat in the specified bar. Adding 1.0 to the specified Beat!"), Beat);
+			UE_LOG(LogMIDI, Warning, TEXT("Beat %f specified in Music Timestamp! Beats in Music Timestamps are '1' based. Beat 1.0 is the first beat in the specified bar. Adding 1.0 to the specified Beat!"), Beat);
 			Beat += 1.0f;
 		}
 		else
 		{
-			UE_LOG(LogMidi, Warning, TEXT("Negative beat (%f) specified in Music Timestamp. This is not allowed! Beats are '1' based in Music Timestamps and must be positive. Using beat 1!"), Beat);
+			UE_LOG(LogMIDI, Warning, TEXT("Negative beat (%f) specified in Music Timestamp. This is not allowed! Beats are '1' based in Music Timestamps and must be positive. Using beat 1!"), Beat);
 			Beat = 1.0f;
 		}
 	}
@@ -122,7 +138,7 @@ int32 FBarMap::BarBeatTickIncludingCountInToTick(int32 BarIndex, int32 BeatInBar
 {
 	if (BeatInBar < 1)
 	{
-		UE_LOG(LogMidi, Warning, TEXT("Beat %d specified as a \"beat in bar\". Beat 1 is the first beat in a bar. Using Beat 1!"), BeatInBar);
+		UE_LOG(LogMIDI, Warning, TEXT("Beat %d specified as a \"beat in bar\". Beat 1 is the first beat in a bar. Using Beat 1!"), BeatInBar);
 		BeatInBar = 0;
 	}
 	else
@@ -154,11 +170,11 @@ int32 FBarMap::BarBeatTickIncludingCountInToTick(int32 BarIndex, int32 BeatInBar
 	int32 TicksInBeat = GetTicksInBeatAfterPoint(TimeSigIndex);
 	if (BeatInBar >= Points[TimeSigIndex].TimeSignature.Numerator)
 	{
-		UE_LOG(LogMidi, Warning, TEXT("BarBeatTickToAbsoluteTick: Supplied 'BeatInBar' is greater than the number of beats in the specified bar!"));
+		UE_LOG(LogMIDI, Warning, TEXT("BarBeatTickToAbsoluteTick: Supplied 'BeatInBar' is greater than the number of beats in the specified bar!"));
 	}
 	if (TickInBeat >= TicksInBeat)
 	{
-		UE_LOG(LogMidi, Warning, TEXT("BarBeatTickToAbsoluteTick: Supplied 'TickInBeat' is greater than the number of ticks in each beat at the specified bar!"));
+		UE_LOG(LogMIDI, Warning, TEXT("BarBeatTickToAbsoluteTick: Supplied 'TickInBeat' is greater than the number of ticks in each beat at the specified bar!"));
 	}
 	return Points[TimeSigIndex].StartTick +
 			(BarDelta * GetTicksInBarAfterPoint(TimeSigIndex)) +
diff --git a/Engine/Plugins/Runtime/Harmonix/Source/HarmonixMidi/Private/HarmonixMidi/BeatMap.cpp b/Engine/Plugins/Runtime/Harmonix/Source/HarmonixMidi/Private/HarmonixMidi/BeatMap.cpp
--- a/Engine/Plugins/Runtime/Harmonix/Source/HarmonixMidi/Private/HarmonixMidi/BeatMap.cpp
+++ b/Engine/Plugins/Runtime/Harmonix/Source/HarmonixMidi/Private/HarmonixMidi/BeatMap.cpp
@@ -3,6 +3,30 @@
 #include "HarmonixMidi/MidiConstants.h"
 #include "Misc/ScopeLock.h"
 
+bool FBeatMap::operator==(const FBeatMap& Other) const
+{
+	if (TicksPerQuarterNote != Other.TicksPerQuarterNote || Points.Num() != Other.Points.Num() || Bars.Num() != Other.Bars.Num())
+	{
+		return false;
+	}
+	for (int32 PointIndex = 0; PointIndex < Points.Num(); ++PointIndex)
+	{
+		if (Points[PointIndex] != Other.Points[PointIndex])
+		{
+			return false;
+		}
+	}
+
+	for (int32 BarIndex = 0; BarIndex < Bars.Num(); ++BarIndex)
+	{
+		if (Bars[BarIndex] != Other.Bars[BarIndex])
+		{
+			return false;
+		}
+	}
+	return true;
+}
+
 void FBeatMap::Empty()
 {
 	Points.Empty();
diff --git a/Engine/Plugins/Runtime/Harmonix/Source/HarmonixMidi/Private/HarmonixMidi/ChordMap.cpp b/Engine/Plugins/Runtime/Harmonix/Source/HarmonixMidi/Private/HarmonixMidi/ChordMap.cpp
--- a/Engine/Plugins/Runtime/Harmonix/Source/HarmonixMidi/Private/HarmonixMidi/ChordMap.cpp
+++ b/Engine/Plugins/Runtime/Harmonix/Source/HarmonixMidi/Private/HarmonixMidi/ChordMap.cpp
@@ -1,6 +1,23 @@
 // Copyright Epic Games, Inc. All Rights Reserved.
 #include "HarmonixMidi/ChordMap.h"
 
+bool FChordProgressionMap::operator==(const FChordProgressionMap& Other) const
+{
+	if (TicksPerQuarterNote != Other.TicksPerQuarterNote || Points.Num() != Other.Points.Num() || ChordTrackIndex != Other.ChordTrackIndex)
+	{
+		return false;
+	}
+
+	for (int32 PointIndex = 0; PointIndex < Points.Num(); ++PointIndex)
+	{
+		if (Points[PointIndex] != Other.Points[PointIndex])
+		{
+			return false;
+		}
+	}
+	return true;
+}
+
 void FChordProgressionMap::Finalize(int32 LastTick)
 {
 	FMusicMapUtl::Finalize(Points, LastTick);
@@ -35,6 +52,7 @@ const void FChordProgressionMap::GetChordListCopy(TArray<FChordMapPoint>& ChordL
 
 void FChordProgressionMap::Empty()
 {
+	Points.Empty();
 	ChordTrackIndex = -1;
 }
 
diff --git a/Engine/Plugins/Runtime/Harmonix/Source/HarmonixMidi/Private/HarmonixMidi/MidiConstants.cpp b/Engine/Plugins/Runtime/Harmonix/Source/HarmonixMidi/Private/HarmonixMidi/MidiConstants.cpp
--- a/Engine/Plugins/Runtime/Harmonix/Source/HarmonixMidi/Private/HarmonixMidi/MidiConstants.cpp
+++ b/Engine/Plugins/Runtime/Harmonix/Source/HarmonixMidi/Private/HarmonixMidi/MidiConstants.cpp
@@ -2,7 +2,7 @@
 
 #include "HarmonixMidi/MidiConstants.h"
 
-DEFINE_LOG_CATEGORY(LogMidi);
+DEFINE_LOG_CATEGORY(LogMIDI);
 
 namespace Harmonix::Midi::Constants
 {
diff --git a/Engine/Plugins/Runtime/Harmonix/Source/HarmonixMidi/Private/HarmonixMidi/MidiFile.cpp b/Engine/Plugins/Runtime/Harmonix/Source/HarmonixMidi/Private/HarmonixMidi/MidiFile.cpp
--- a/Engine/Plugins/Runtime/Harmonix/Source/HarmonixMidi/Private/HarmonixMidi/MidiFile.cpp
+++ b/Engine/Plugins/Runtime/Harmonix/Source/HarmonixMidi/Private/HarmonixMidi/MidiFile.cpp
@@ -28,18 +28,19 @@ public:
 	int32 GetLastTick() const { return LastTick; }
 
 private:
-	virtual void Reset() override
+	virtual bool Reset() override
 	{
 		LastTick = 0;
 		CurrentTrackHasName = false;
+		return true;
 	}
-	virtual void OnNewTrack(int32 NewTrackIndex) override;
-	virtual void OnEndOfTrack(int32 InLastTick) override;
-	virtual void OnAllTracksRead() override {}
-	virtual void OnMidiMessage(int32 Tick, uint8 Status, uint8 Data1, uint8 Data2) override;
-	virtual void OnTempo(int32 Tick, int32 Tempo) override;
-	virtual void OnText(int32 Tick, const FString& Str, uint8 Type) override;
-	virtual void OnTimeSignature(int32 Tick, int32 Numerator, int32 Denominator, bool FailOnError = true) override;
+	virtual bool OnNewTrack(int32 NewTrackIndex) override;
+	virtual bool OnEndOfTrack(int32 InLastTick) override;
+	virtual bool OnAllTracksRead() override { return true; }
+	virtual bool OnMidiMessage(int32 Tick, uint8 Status, uint8 Data1, uint8 Data2) override;
+	virtual bool OnTempo(int32 Tick, int32 Tempo) override;
+	virtual bool OnText(int32 Tick, const FString& Str, uint8 Type) override;
+	virtual bool OnTimeSignature(int32 Tick, int32 Numerator, int32 Denominator, bool FailOnError = true) override;
 
 	UMidiFile& File;
 	int32  LastTick = 0;
@@ -54,6 +55,11 @@ UMidiFile::UMidiFile()
 	TheMidiData.Tracks.Emplace("Conductor");
 }
 
+bool UMidiFile::operator==(const UMidiFile& Other) const
+{
+	return TheMidiData == Other.TheMidiData;
+}
+
 void UMidiFile::GetAssetRegistryTags(FAssetRegistryTagsContext Context) const
 {
 	Super::GetAssetRegistryTags(Context);
@@ -100,7 +106,6 @@ TSharedPtr<Audio::IProxyData> UMidiFile::CreateProxyData(const Audio::FProxyData
 	return Proxy;
 }
 
-
 void UMidiFile::BeginDestroy()
 {
 	RenderableCopyOfMidiFileData = nullptr;
@@ -122,13 +127,14 @@ void UMidiFile::PostInitProperties()
 void UMidiFile::LoadStdMidiFile(
 	const FString& FilePath,
 	int32 DesiredTicksPerQuarterNote,
-	Harmonix::Midi::Constants::EMidiTextEventEncoding InTextEncoding)
+	Harmonix::Midi::Constants::EMidiTextEventEncoding InTextEncoding,
+	bool EnsureNotFailed)
 {
 	FString Filename = FPaths::GetCleanFilename(FilePath);
 	IPlatformFile& PlatformFileApi = FPlatformFileManager::Get().GetPlatformFile();
 	IFileHandle* FileHandle = PlatformFileApi.OpenRead(*FilePath);
 	TSharedPtr<FArchiveFileReaderGeneric> Archive = MakeShared<FArchiveFileReaderGeneric>(FileHandle, *Filename, FileHandle->Size());
-	LoadStdMidiFile(Archive, Filename, DesiredTicksPerQuarterNote, InTextEncoding);
+	LoadStdMidiFile(Archive, Filename, DesiredTicksPerQuarterNote, InTextEncoding, EnsureNotFailed);
 
 #if WITH_EDITORONLY_DATA
 	if (!AssetImportData)
@@ -144,17 +150,19 @@ void UMidiFile::LoadStdMidiFile(
 	int32 BufferSize,
 	const FString& Filename,
 	int32 DesiredTicksPerQuarterNote,
-	Harmonix::Midi::Constants::EMidiTextEventEncoding InTextEncoding)
+	Harmonix::Midi::Constants::EMidiTextEventEncoding InTextEncoding,
+	bool EnsureNotFailed)
 {
 	TSharedPtr<FBufferReader> BufferArchive = MakeShared<FBufferReader>(Buffer, BufferSize, false);
-	LoadStdMidiFile(BufferArchive, Filename, DesiredTicksPerQuarterNote, InTextEncoding);
+	LoadStdMidiFile(BufferArchive, Filename, DesiredTicksPerQuarterNote, InTextEncoding, EnsureNotFailed);
 }
 
 void UMidiFile::LoadStdMidiFile(
 	TSharedPtr<FArchive> Archive,
 	const FString& Filename,
 	int32 DesiredTicksPerQuarterNote,
-	Harmonix::Midi::Constants::EMidiTextEventEncoding InTextEncoding)
+	Harmonix::Midi::Constants::EMidiTextEventEncoding InTextEncoding,
+	bool EnsureNotFailed)
 {
 	TheMidiData.TicksPerQuarterNote = DesiredTicksPerQuarterNote;
 	TheMidiData.SongMaps.Init(TheMidiData.TicksPerQuarterNote);
@@ -173,9 +181,18 @@ void UMidiFile::LoadStdMidiFile(
 	reader.ReadAllTracks();
 
 	TheMidiData.LastEventTick = EventReceiver.GetLastTick();
-	if (reader.GetFailed())
+
+	if (reader.IsFailed())
 	{
-		UE_LOG(LogMidi, Error, TEXT("MIDI import failed. Midi data is malformed."));
+		if (EnsureNotFailed)
+		{
+			ensureAlwaysMsgf(false, TEXT("%s: MIDI import failed. Midi data is malformed. Check the log for details."), *Filename);
+		}
+		else
+		{
+			UE_LOG(LogMIDI, Error, TEXT("%s MIDI import failed. Midi data is malformed. Check the log for details."), *Filename);
+		}
+		TheMidiData.Empty();
 	}
 
 	// we are now "dirty"... so make sure any new requests for renderable data 
@@ -183,19 +200,22 @@ void UMidiFile::LoadStdMidiFile(
 	RenderableCopyOfMidiFileData = nullptr;
 }
 
-void UMidiFile::SaveStdMidiFile(const FString& FilePath)
+void UMidiFile::SaveStdMidiFile(const FString& FilePath) const
 {
 	FString Filename = FPaths::GetCleanFilename(FilePath);
 	IPlatformFile& PlatformFileApi = FPlatformFileManager::Get().GetPlatformFile();
 	IFileHandle* FileHandle = PlatformFileApi.OpenWrite(*FilePath);
-	TSharedPtr<FArchiveFileWriterGeneric> Archive = MakeShared<FArchiveFileWriterGeneric>(FileHandle, *Filename, 0);
-	SaveStdMidiFile(Archive, Filename);
+	if (ensureAlwaysMsgf(FileHandle, TEXT("Failed to open file \"%s\" for writing."), *FilePath))
+	{
+		TSharedPtr<FArchiveFileWriterGeneric> Archive = MakeShared<FArchiveFileWriterGeneric>(FileHandle, *Filename, 0);
+		SaveStdMidiFile(Archive, Filename);
+	}
 }
 
-void UMidiFile::SaveStdMidiFile(TSharedPtr<FArchive> Archive, const FString& Filename)
+void UMidiFile::SaveStdMidiFile(TSharedPtr<FArchive> Archive, const FString& Filename) const
 {
 	FMidiWriter Writer(*Archive, TheMidiData.TicksPerQuarterNote);
-	Algo::ForEach(TheMidiData.Tracks, [&](FMidiTrack& track) { track.WriteStdMidi(Writer); });
+	Algo::ForEach(TheMidiData.Tracks, [&](const FMidiTrack& track) { track.WriteStdMidi(Writer); });
 }
 
 FMidiTrack* UMidiFile::AddTrack(const FString& Name)
@@ -206,32 +226,51 @@ FMidiTrack* UMidiFile::AddTrack(const FString& Name)
 
 void UMidiFile::Empty()
 {
-	TheMidiData.Tracks.Empty();
-	TheMidiData.Tracks.Emplace("Conductor");
+	TheMidiData.Empty();
 }
 
-void UMidiFile::SetConductorTrack(const FTempoMap* TempoMap, const FBarMap* BarMap)
+bool UMidiFile::IsEmpty() const 
 {
-	check(!TheMidiData.Tracks.IsEmpty());
-	// use assignment operator to properly clear string table, events, and other data.
-	TheMidiData.Tracks[0] = FMidiTrack("Conductor");
+	return TheMidiData.IsEmpty();
+}
 
-	int32 numTempoChanges = TempoMap->GetNumTempoChangePoints();
+//void UMidiFile::SetConductorTrack(const FTempoMap* TempoMap, const FBarMap* BarMap)
+void UMidiFile::BuildConductorTrack()
+{
+	if (TheMidiData.Tracks.IsEmpty())
+	{
+		TheMidiData.Tracks.Emplace("Conductor");
+	}
+	else if (TheMidiData.Tracks[0].GetName()->Compare("Conductor", ESearchCase::IgnoreCase))
+	{
+		TheMidiData.Tracks.Insert(FMidiTrack("Conductor"), 0);
+	}
+	else
+	{
+		// use assignment operator to properly clear string table, events, and other data.
+		TheMidiData.Tracks[0] = FMidiTrack("Conductor");
+	}
+
+	const FTempoMap& TempoMap = TheMidiData.SongMaps.GetTempoMap();
+
+	int32 numTempoChanges = TempoMap.GetNumTempoChangePoints();
 
 	int32 i;
 	for (i = 0; i < numTempoChanges; ++i)
 	{
-		int32 tick = TempoMap->GetTempoChangePointTick(i);
-		float msPerQuarterNote = TempoMap->GetMsPerQuarterNoteAtTick(tick);
+		int32 tick = TempoMap.GetTempoChangePointTick(i);
+		float msPerQuarterNote = TempoMap.GetMsPerQuarterNoteAtTick(tick);
 		int32 usecPerQuarterNote = int32(msPerQuarterNote * 1000);
 		TheMidiData.Tracks[0].AddEvent(FMidiEvent(tick, FMidiMsg(usecPerQuarterNote)));
 	}
 
-	int32 numTimeSigChanges = BarMap->GetNumTimeSignaturePoints();
+	const FBarMap& BarMap = TheMidiData.SongMaps.GetBarMap();
+
+	int32 numTimeSigChanges = BarMap.GetNumTimeSignaturePoints();
 
 	for (i = 0; i < numTimeSigChanges; ++i)
 	{
-		const FTimeSignaturePoint& sig = BarMap->GetTimeSignaturePoint(i);
+		const FTimeSignaturePoint& sig = BarMap.GetTimeSignaturePoint(i);
 		TheMidiData.Tracks[0].AddEvent(FMidiEvent(sig.StartTick, FMidiMsg(sig.TimeSignature.Numerator, sig.TimeSignature.Denominator)));
 	}
 
@@ -480,15 +519,19 @@ bool UMidiFile::ShouldConformMidiFileLength(EMidiFileLengthConformOption Option)
 	{
 		return true;
 	}
+	else if (!FMath::IsNearlyEqual(Timestamp.Beat, 1.0f))
+	{
+		ConformMidiFileLength(EMidiFileLengthConformOption::RoundDown, true);
+	}
 	return false;
 }
 
-void UMidiFile::ConformMidiFileLength(EMidiFileLengthConformOption Option)
+void UMidiFile::ConformMidiFileLength(EMidiFileLengthConformOption Option, bool Force)
 {
 	//check whether midi file length needs to be conformed
-	if (!ShouldConformMidiFileLength(Option))
+	if (!Force && !ShouldConformMidiFileLength(Option))
 	{
-		UE_LOG(LogMidi, Log, TEXT("Midi file does not need to be conformed."));
+		UE_LOG(LogMIDI, Log, TEXT("Midi file does not need to be conformed."));
 		return;
 	}
 	
@@ -720,7 +763,7 @@ void UMidiFile::ConformMidiFileLength(EMidiFileLengthConformOption Option)
 			if (NumControlEventsRemoved > 0) MidiFileLengthConformSummary.Append(FString::Printf(TEXT(" %d Control Change event(s),"), NumControlEventsRemoved));
 			MidiFileLengthConformSummary.Append(FString::Printf(TEXT(" at tick %d "), ConformedLastEventTick));
 		}
-		UE_LOG(LogMidi, Log, TEXT("%s"),*MidiFileLengthConformSummary);
+		UE_LOG(LogMIDI, Log, TEXT("%s"),*MidiFileLengthConformSummary);
 	}
 
 	//This prevents the current conform option affecting any midi file that is currently being obtained from RenderableCopyOfMidiFileData
@@ -748,13 +791,14 @@ FMidiEventReceiver::FMidiEventReceiver(UMidiFile& file)
 {
 }
 
-void FMidiEventReceiver::OnNewTrack(int32 newTrackIndex)
+bool FMidiEventReceiver::OnNewTrack(int32 newTrackIndex)
 {
 	File.GetTracks().Emplace();
 	CurrentTrackHasName = false;
+	return true;
 }
 
-void FMidiEventReceiver::OnEndOfTrack(int InLastTick)
+bool FMidiEventReceiver::OnEndOfTrack(int InLastTick)
 {
 	if (InLastTick > LastTick)
 	{
@@ -762,7 +806,10 @@ void FMidiEventReceiver::OnEndOfTrack(int InLastTick)
 	}
 	if (!CurrentTrackHasName)
 	{
-		check(!File.GetTracks().IsEmpty());
+		if (File.GetTracks().IsEmpty())
+		{
+			return false;
+		}
 
 		FMidiTrack& Track = File.GetTracks().Last();
 		uint16 StringIndex = Track.AddText("track_1");
@@ -771,25 +818,39 @@ void FMidiEventReceiver::OnEndOfTrack(int InLastTick)
 		Track.Sort();
 		CurrentTrackHasName = true;
 	}
+	return true;
 }
 
-void FMidiEventReceiver::OnMidiMessage(int32 Tick, uint8 Status, uint8 Data1, uint8 Data2)
+bool FMidiEventReceiver::OnMidiMessage(int32 Tick, uint8 Status, uint8 Data1, uint8 Data2)
 {
-	check(!File.GetTracks().IsEmpty());
+	if (File.GetTracks().IsEmpty())
+	{
+		return false;
+	}
 
 	File.GetTracks().Last().AddEvent(FMidiEvent(Tick, FMidiMsg(Status, Data1, Data2)));
+
+	return true;
 }
 
-void FMidiEventReceiver::OnTempo(int32 Tick, int32 Tempo)
+bool FMidiEventReceiver::OnTempo(int32 Tick, int32 Tempo)
 {
-	check(!File.GetTracks().IsEmpty());
+	if (File.GetTracks().IsEmpty())
+	{
+		return false;
+	}
 
 	File.GetTracks().Last().AddEvent(FMidiEvent(Tick, FMidiMsg(Tempo)));
+
+	return true;
 }
 
-void FMidiEventReceiver::OnText(int32 Tick, const FString& Str, uint8 Type)
+bool FMidiEventReceiver::OnText(int32 Tick, const FString& Str, uint8 Type)
 {
-	check(!File.GetTracks().IsEmpty());
+	if (File.GetTracks().IsEmpty())
+	{
+		return false;
+	}
 
 	FMidiTrack& Track = File.GetTracks().Last();
 	uint16 StringIndex = Track.AddText(Str);
@@ -799,26 +860,53 @@ void FMidiEventReceiver::OnText(int32 Tick, const FString& Str, uint8 Type)
 		// Track name should always appear at tick 0!
 		if (Tick != 0 && !CurrentTrackHasName)
 		{
-			UE_LOG(LogMidi, Warning, TEXT("MIDI track name (\"%s\") event found at tick %d but none found at tick 0! The first track name event on a midi track should be on tick 0. Shifting it to tick 0."),
+			UE_LOG(LogMIDI, Warning, TEXT("MIDI track name (\"%s\") event found at tick %d but none found at tick 0! The first track name event on a midi track should be on tick 0. Shifting it to tick 0."),
 				*Str, Tick);
 			Tick = 0;
 		}
 		else if (CurrentTrackHasName && Str != *Track.GetName())
 		{
-			UE_LOG(LogMidi, Warning, TEXT("2nd Track Name event found on MIDI track %s --> %s at tick %d"),
+			UE_LOG(LogMIDI, Warning, TEXT("2nd Track Name event found on MIDI track %s --> %s at tick %d"),
 				**Track.GetName(), *Str , Tick);
 		}
 		CurrentTrackHasName = true;
 	}
 
 	Track.AddEvent(FMidiEvent(Tick, FMidiMsg::CreateText(StringIndex, Type)));
+
+	return true;
 }
 
-void FMidiEventReceiver::OnTimeSignature(int32 Tick, int32 Numerator, int32 Denominator, bool FailOnError)
+bool FMidiEventReceiver::OnTimeSignature(int32 Tick, int32 Numerator, int32 Denominator, bool FailOnError)
 {
-	check(!File.GetTracks().IsEmpty());
+	if (File.GetTracks().IsEmpty())
+	{
+		return false;
+	}
 
 	File.GetTracks().Last().AddEvent(FMidiEvent(Tick, FMidiMsg(Numerator, Denominator)));
+
+	return true;
+}
+
+bool FMidiFileData::operator==(const FMidiFileData& Other) const
+{
+	if (Tracks.Num() != Other.Tracks.Num())
+	{
+		return false;
+	}
+	for (int32 TrackIndex = 0; TrackIndex < Tracks.Num(); ++TrackIndex)
+	{
+		if (Tracks[TrackIndex] != Other.Tracks[TrackIndex])
+		{
+			return false;
+		}
+	}
+
+	return MidiFileName == Other.MidiFileName &&
+		TicksPerQuarterNote == Other.TicksPerQuarterNote &&
+		SongMaps == Other.SongMaps &&
+		LastEventTick == Other.LastEventTick;
 }
 
 int32 FMidiFileData::FindTrackIndexByName(const FString& TrackName)
@@ -838,7 +926,7 @@ void FMidiFileData::PostSerialize(const FArchive& Ar)
 {
 	if (Ar.IsLoading() && SongMaps.BarMapIsEmpty() && Tracks.Num() > 0)
 	{
-		UE_LOG(LogMidi, Warning, TEXT("Empty bar map. Rebuilding."));
+		UE_LOG(LogMIDI, Warning, TEXT("Empty bar map. Rebuilding."));
 		const FMidiEventList& Events = Tracks[0].GetEvents();
 		FBarMap& BarMap = SongMaps.GetBarMap();
 		BarMap.SetStartBar(1);
diff --git a/Engine/Plugins/Runtime/Harmonix/Source/HarmonixMidi/Private/HarmonixMidi/MidiPlayCursor.cpp b/Engine/Plugins/Runtime/Harmonix/Source/HarmonixMidi/Private/HarmonixMidi/MidiPlayCursor.cpp
--- a/Engine/Plugins/Runtime/Harmonix/Source/HarmonixMidi/Private/HarmonixMidi/MidiPlayCursor.cpp
+++ b/Engine/Plugins/Runtime/Harmonix/Source/HarmonixMidi/Private/HarmonixMidi/MidiPlayCursor.cpp
@@ -1193,7 +1193,7 @@ void FMidiPlayCursor::BroadcastEvent(int32 TrackIndex, const FMidiEvent& Event,
 		}
 		break;
 	default:
-		UE_LOG(LogMidi, Error, TEXT("Unknown midi message type %d on track %d at tick %d, file %s"), int(Msg.MsgType()), TrackIndex, Event.GetTick(), **Owner->GetMidiFileName());
+		UE_LOG(LogMIDI, Error, TEXT("Unknown midi message type %d on track %d at tick %d, file %s"), int(Msg.MsgType()), TrackIndex, Event.GetTick(), **Owner->GetMidiFileName());
 		break;
 	}
 }
diff --git a/Engine/Plugins/Runtime/Harmonix/Source/HarmonixMidi/Private/HarmonixMidi/MidiPlayCursorMgr.cpp b/Engine/Plugins/Runtime/Harmonix/Source/HarmonixMidi/Private/HarmonixMidi/MidiPlayCursorMgr.cpp
--- a/Engine/Plugins/Runtime/Harmonix/Source/HarmonixMidi/Private/HarmonixMidi/MidiPlayCursorMgr.cpp
+++ b/Engine/Plugins/Runtime/Harmonix/Source/HarmonixMidi/Private/HarmonixMidi/MidiPlayCursorMgr.cpp
@@ -224,7 +224,7 @@ void FMidiPlayCursorMgr::UnregisterPlayCursor(FMidiPlayCursor* PlayCursor, bool
 		{
 			if (PlayCursor->Tracker == &Trackers[Index])
 			{
-				UE_LOG(LogMidi, Error, TEXT("You cannot unregister a MidiPlayCursor while the cursor's manager is traversing its cursor list!"));
+				UE_LOG(LogMIDI, Error, TEXT("You cannot unregister a MidiPlayCursor while the cursor's manager is traversing its cursor list!"));
 				return;
 			}
 			continue;
@@ -238,7 +238,7 @@ void FMidiPlayCursorMgr::UnregisterPlayCursor(FMidiPlayCursor* PlayCursor, bool
 	}
 	if (WarnOnFail)
 	{
-		UE_LOG(LogMidi, Warning, TEXT("Attempt to remove MidiPlayCursor from a manager that doesn't own it!"));
+		UE_LOG(LogMIDI, Warning, TEXT("Attempt to remove MidiPlayCursor from a manager that doesn't own it!"));
 	}
 }
 
@@ -249,7 +249,7 @@ void FMidiPlayCursorMgr::UnregisterAllPlayCursors()
 
 	if (GetHiResTracker().TraversingCursors || GetLowResTracker().TraversingCursors)
 	{
-		UE_LOG(LogMidi, Error, TEXT("You cannot unregister all MidiPlayCursors while this manager is currently traversing it's cursor list!"));
+		UE_LOG(LogMIDI, Error, TEXT("You cannot unregister all MidiPlayCursors while this manager is currently traversing it's cursor list!"));
 		return;
 	}
 
@@ -448,7 +448,7 @@ void FMidiPlayCursorMgr::SetLoopImpl(int32 StartTick, int32 EndTick, bool Ignori
 		}
 		else
 		{
-			UE_LOG(LogMidi, Warning, TEXT("FMidiPlayCursorMgr::SetLoop : not ignoring look ahead, but cursors are not in phase!"));
+			UE_LOG(LogMIDI, Warning, TEXT("FMidiPlayCursorMgr::SetLoop : not ignoring look ahead, but cursors are not in phase!"));
 		}
 	}
 }
@@ -466,7 +466,7 @@ void FMidiPlayCursorMgr::ClearLoopImpl(bool IgnoringLookAhead, bool IsLowRes)
 		}
 		else
 		{
-			UE_LOG(LogMidi, Warning, TEXT("FMidiPlayCursorMgr::ClearLoop : not ignoring look ahead, but cursors are not in phase!"));
+			UE_LOG(LogMIDI, Warning, TEXT("FMidiPlayCursorMgr::ClearLoop : not ignoring look ahead, but cursors are not in phase!"));
 		}
 	}
 }
@@ -730,7 +730,7 @@ void FMidiPlayCursorMgr::AdvanceHiResThruTick(int32 ThruTick, bool Broadcast, bo
 
 	if (ThruTick < GetHiResTracker().CurrentTick)
 	{
-		UE_LOG(LogMidi, Warning, TEXT("Asked to go back in time without a seek!"));
+		UE_LOG(LogMIDI, Warning, TEXT("Asked to go back in time without a seek!"));
 		return;
 	}
 
diff --git a/Engine/Plugins/Runtime/Harmonix/Source/HarmonixMidi/Private/HarmonixMidi/MidiReader.cpp b/Engine/Plugins/Runtime/Harmonix/Source/HarmonixMidi/Private/HarmonixMidi/MidiReader.cpp
--- a/Engine/Plugins/Runtime/Harmonix/Source/HarmonixMidi/Private/HarmonixMidi/MidiReader.cpp
+++ b/Engine/Plugins/Runtime/Harmonix/Source/HarmonixMidi/Private/HarmonixMidi/MidiReader.cpp
@@ -12,6 +12,28 @@
 // by the SMF file format.)
 const int32 kMidiChunkIDSize = 4;
 
+#define MIDIREADER_EXPECT_WITH_ERROR_RETURN_VAL(expected, ret_val_on_fail, ...)	\
+	do                                                                          \
+	{                                                                           \
+		if (!(expected))                                                        \
+		{                                                                       \
+			UE_LOG(LogMIDI, Error, __VA_ARGS__);                                \
+			State = EState::Failed;                                             \
+			return ret_val_on_fail;                                             \
+		}                                                                       \
+	} while (0)
+
+#define MIDIREADER_EXPECT_WITH_ERROR(expected, ...)	\
+	do                                              \
+	{                                               \
+		if (!(expected))                            \
+		{                                           \
+			UE_LOG(LogMIDI, Error, __VA_ARGS__);    \
+			State = EState::Failed;                 \
+			return;                                 \
+		}                                           \
+	} while (0)
+
 class FMidiChunkID
 {
 public:
@@ -39,7 +61,7 @@ const FMidiChunkID	FMidiChunkID::kMTrk("MTrk");
 
 FArchive& operator<<(FArchive& Archive, FMidiChunkID& Id)
 {
-	Archive.Serialize(Id.Str, 4);
+	Archive.Serialize(Id.Str, 4); 
 	return Archive;
 }
 
@@ -91,8 +113,7 @@ FStdMidiFileReader::FStdMidiFileReader(
 	, DestinationTicksPerQuarterNote(TicksPerQuarterNote)
 {
 	check(Receiver);
-	BufferArchive = MakeShared<FBufferReader>(Buffer, BufferSize, false);
-	InputArchive = BufferArchive;
+	Archive = MakeShared<FBufferReader>(Buffer, BufferSize, false);
 	Init();
 }
 
@@ -109,8 +130,7 @@ FStdMidiFileReader::FStdMidiFileReader(
 	Filename = FPaths::GetCleanFilename(FilePath);
 	IPlatformFile& PlatformFileApi = FPlatformFileManager::Get().GetPlatformFile();
 	IFileHandle* FileHandle = PlatformFileApi.OpenRead(*FilePath);
-	FileArchive = MakeShared<FArchiveFileReaderGeneric>(FileHandle, *Filename, FileHandle->Size());
-	InputArchive = FileArchive;
+	Archive = MakeShared<FArchiveFileReaderGeneric>(FileHandle, *Filename, FileHandle->Size());
 	Init();
 }
 
@@ -121,7 +141,7 @@ FStdMidiFileReader::FStdMidiFileReader(
 	int32 TicksPerQuarterNote,
 	Harmonix::Midi::Constants::EMidiTextEventEncoding InTextEncoding)
 	: Filename(InFilename)
-	, InputArchive(Archive)
+	, Archive(Archive.ToSharedRef())
 	, Receiver(InReceiver)
 	, TextEncoding(InTextEncoding)
 	, DestinationTicksPerQuarterNote(TicksPerQuarterNote)
@@ -130,50 +150,24 @@ FStdMidiFileReader::FStdMidiFileReader(
 	Init();
 }
 
-bool FStdMidiFileReader::ReadStdMidiFileForReceiver(
-	void* Buffer,
-	int32 BufferSize,
-	const FString& FileName,
-	 IMidiReceiver* Receiver,
-	  int32 TicksPerQuarterNote,
-	  Harmonix::Midi::Constants::EMidiTextEventEncoding InTextEncoding)
-{
-	bool Success = true;
-	FStdMidiFileReader Reader(Buffer, BufferSize, FileName, Receiver, TicksPerQuarterNote, InTextEncoding);
-
-	if (Reader.GetFailed())
-	{
-		UE_LOG(LogMidi, Error, TEXT("Unable to read MIDI file"));
-		Success = false;
-	}
-	else
-	{
-		Receiver->Reset();
-		Reader.ReadAllTracks();
-		Receiver->Finalize(Reader.LastTick);
-
-		if (Reader.GetFailed())
-		{
-			UE_LOG(LogMidi, Error, TEXT("opened but failed to parse MIDI file"));
-			Success = false;
-		}
-	}
-
-	return Success;
-}
-
 void FStdMidiFileReader::ReadAllTracks()
 {
 	// read from the top of the stream
-	InputArchive->Seek(0);
+	Archive->Seek(0);
 
 	bool TrackReadSuccessfully = false;
 	do 
 	{
 		TrackReadSuccessfully = ReadTrack();
-	} while (!Failed && TrackReadSuccessfully);
+	} while (State != EState::Failed && TrackReadSuccessfully);
 
-	Receiver->Finalize(LastTick);
+	if (State != EState::Failed)
+	{
+		if (!Receiver->Finalize(LastTick))
+		{
+			State = EState::Failed;
+		}
+	}
 }
 
 bool FStdMidiFileReader::ReadSomeEvents(int32 NumEvents)
@@ -181,7 +175,7 @@ bool FStdMidiFileReader::ReadSomeEvents(int32 NumEvents)
 	for (int32 i = 0; i < NumEvents; ++i)
 	{
 		ReadNextEvent();
-		if ((State == EState::End) || Failed)
+		if (State == EState::End || State == EState::Failed)
 		{
 			return true;
 		}
@@ -194,7 +188,7 @@ bool FStdMidiFileReader::ReadTrack()
 	do
 	{
 		ReadNextEvent();
-	} while (State != EState::End && State != EState::NewTrack && !Failed);
+	} while (State != EState::End && State != EState::NewTrack && State != EState::Failed);
 
 	return (State == EState::NewTrack);
 }
@@ -203,7 +197,7 @@ bool FStdMidiFileReader::ReadEvents(int32 Count)
 {
 	for (int32 i = 0; i < Count; ++i)
 	{
-		if (State == EState::End || Failed)
+		if (State == EState::End || State == EState::Failed)
 		{
 			return false;
 		}
@@ -221,23 +215,33 @@ void FStdMidiFileReader::SkipCurrentTrack()
 		if (CurrentTrackIndex == NumTracks - 1)
 		{
 			State = EState::End;
-			Receiver->OnEndOfTrack(CurrentTick);
-			Receiver->OnAllTracksRead();
+			if (!Receiver->OnEndOfTrack(CurrentTick) || !Receiver->OnAllTracksRead())
+			{
+				State = EState::Failed;
+			}
 		}
 		else
 		{
 			State = EState::NewTrack;
-			Receiver->OnEndOfTrack(CurrentTick);
-			InputArchive->Seek(TrackEndPos);
+			if (!Receiver->OnEndOfTrack(CurrentTick))
+			{
+				State = EState::Failed;
+			}
+			else
+			{
+				Archive->Seek(TrackEndPos);
+			}
 		}
 	}
 }
 
 void FStdMidiFileReader::Init()
 {
+	// make oursleces a "local" barmap as it can be used to 
+	// provide better error messages...
 	BarMap = MakeShared<FBarMap>();
-	check(InputArchive);
-	InputArchive->ArForceByteSwapping = PLATFORM_LITTLE_ENDIAN == 1;
+	// midi files are big-endian. If this is a little endian platform we need to byteswap!
+	Archive->ArForceByteSwapping = PLATFORM_LITTLE_ENDIAN == 1;
 	Receiver->SetMidiReader(this);
 }
 
@@ -248,135 +252,109 @@ void FStdMidiFileReader::ReadNextEvent()
 
 void FStdMidiFileReader::ReadNextEventImpl()
 {
-	if (Failed)
-	{
-		return;
-	}
-
 	switch (State)
 	{
 	case EState::InTrack:
-		ReadEvent(*InputArchive);
+		ReadEvent();
 		break;
 	case EState::NewTrack:
 		if (Format == 0 && CurrentTrackIndex == 0)
 		{
 			// we already ready the one and only track as the "conductor track".
 			// we have to read it again as "track 1", so...
-			InputArchive->Seek(LastTracksFilePosition);
+			Archive->Seek(LastTracksFilePosition);
 			TrackFilteringMode = ETrackFilteringMode::NonConductorEvents;
 		}
 		else 
 		{
-			LastTracksFilePosition = InputArchive->Tell();
+			LastTracksFilePosition = Archive->Tell();
 		}
-		ReadTrackHeader(*InputArchive);
+		ReadTrackHeader();
 		break;
 	case EState::Start:
-		ReadFileHeader(*InputArchive);
+		ReadFileHeader();
 		break;
 	case EState::End:
+	case EState::Failed:
 		break;
 	}
 }
 
 // Read the standard Midi file header chunk (MThd)
-void FStdMidiFileReader::ReadFileHeader(FArchive& Archive)
+void FStdMidiFileReader::ReadFileHeader()
 {
 	check(State == EState::Start);
 
-	FMidiChunkHeader Header(Archive);
+	FMidiChunkHeader Header(*Archive);
 
 	// Verify the MIDI header is properly formed.
-	ensureAlwaysMsgf(Header.GetID() == FMidiChunkID::kMThd && Header.GetLength() == 6, TEXT("%s: MIDI file header is corrupt"), *Filename);
+	MIDIREADER_EXPECT_WITH_ERROR(Header.GetID() == FMidiChunkID::kMThd && Header.GetLength() == 6, TEXT("%s: MIDI file header is corrupt"), *Filename);
 
 	// SMF format:
 	//   0: A single multi-channel track
 	//   1: 1 or more simultaneous tracks
 	//   2: 1 or more sequentially independent tracks
 	// we only support type 0 & 1.
-	Archive << Format;
-	ensureMsgf(Format == 1 || Format == 0, TEXT("%s: Only type 0 or 1 MIDI files are supported; this file is type %d"), *Filename, Format);
+	*Archive << Format;
+	MIDIREADER_EXPECT_WITH_ERROR(Format == 1 || Format == 0, TEXT("%s: Only type 0 or 1 MIDI files are supported; this file is type %d"), *Filename, Format);
+
+	*Archive << NumTracks;
+
+	MIDIREADER_EXPECT_WITH_ERROR(NumTracks > 0, TEXT("%s: MIDI file has no tracks"), *Filename);
 
 	if (Format == 0)
 	{
+		MIDIREADER_EXPECT_WITH_ERROR(NumTracks == 1, TEXT("%s: Format 0 file expected only one track but found %d"), *Filename, NumTracks);
+
 		// We are starting to read a format 0 midi file, so on the first pass
 		// over the one and only track in the file gather the conductor type events...
 		TrackFilteringMode = ETrackFilteringMode::ConductorEvents;
-	}
-	else
-	{
-		TrackFilteringMode = ETrackFilteringMode::None;
-	}
-
-	Archive << NumTracks;
-
-	if (!ensureMsgf(Format != 0 || NumTracks == 1, TEXT("%s: Format 0 file expected only one track but found %d"), *Filename, NumTracks))
-	{
-		Failed = true;
-		return;
-	}
-	if (Format == 0)
-	{
 		// we are going to convert it to a format 1 file, with a conductor track
 		// and one track for non-conductor events...
 		NumTracks = 2;
 	}
-
-	if (NumTracks <= 0)
-	{
-		UE_LOG(LogMidi, Error, TEXT("%s: MIDI file has no tracks"), *Filename);
-	}
 	else
 	{
-		Receiver->SetNumTracks(NumTracks);
+		TrackFilteringMode = ETrackFilteringMode::None;
 	}
 
+	MIDIREADER_EXPECT_WITH_ERROR(Receiver->SetNumTracks(NumTracks), TEXT("%s: Failed to set the number of tracks (%d)"), *Filename, NumTracks);
+
 	// Hi bit set on ticksPerQuarter indicates a SMPTE division.
 	// We don't support that.
 	int16 TicksPerQuarter;
-	Archive << TicksPerQuarter;
-	if (TicksPerQuarter & 0x8000)
-	{
-		UE_LOG(LogMidi, Error, TEXT("%s: MIDI file uses SMPTE time division; this is not supported"), *Filename);
-	}
+	*Archive << TicksPerQuarter;
+	MIDIREADER_EXPECT_WITH_ERROR((TicksPerQuarter & 0x8000) == 0, TEXT("%s: MIDI file uses SMPTE time division. This is not supported"), *Filename);
+
+	// check for rational TicksPerQuarter...
+	MIDIREADER_EXPECT_WITH_ERROR((TicksPerQuarter % 48) == 0, TEXT("%s: MIDI file Ticks Per Quarter Note is not rational (%d). Must be divisible by 48"), *Filename, TicksPerQuarter);
+
 	TickConversionFactor = float(DestinationTicksPerQuarterNote) / (float)TicksPerQuarter;
 
-	// If there are header failures, don't continue compilation. Just quit here!
-	if ((NumTracks == 0) || (Format != 1 && Format != 0) || (TicksPerQuarter & 0x8000))
-	{
-		Failed = true;
-	}
-	else
-	{
-		State = EState::NewTrack;
-	}
+	State = EState::NewTrack;
 }
 
 // Read the standard Midi track chunk header (MTrk)
-void FStdMidiFileReader::ReadTrackHeader(FArchive& Archive)
+void FStdMidiFileReader::ReadTrackHeader()
 {
 	check(State == EState::NewTrack);
 
-	FMidiChunkHeader TrackHeader(Archive);
+	FMidiChunkHeader TrackHeader(*Archive);
 
-	if (TrackHeader.GetID() != FMidiChunkID::kMTrk)
-	{
-		UE_LOG(LogMidi, Error, TEXT("%s: MIDI track header for track %d is corrupt"), *Filename, CurrentTrackIndex);
-		Failed = true;
-		return;
-	}
+	MIDIREADER_EXPECT_WITH_ERROR(TrackHeader.GetID() == FMidiChunkID::kMTrk, TEXT("%s: MIDI track header for track %d is corrupt"), *Filename, CurrentTrackIndex + 1);
 
-	TrackEndPos = int32(Archive.Tell() + TrackHeader.GetLength());
-	if (TrackEndPos > (int32)Archive.TotalSize())
-	{
-		UE_LOG(LogMidi, Error, TEXT("%s: MIDI track %d data length as recorded in the header exceeds the amount of data in the file!"), *Filename, CurrentTrackIndex);
-		Failed = true;
-		return;
-	}
+	// check for rational track data size...
+	MIDIREADER_EXPECT_WITH_ERROR((TrackHeader.GetLength() + (uint32)Archive->Tell()) <= (uint32)std::numeric_limits<int32>::max(),
+		TEXT("%s: MIDI track data length exceeds maximum length!"), *Filename, CurrentTrackIndex + 1);
+
+	TrackEndPos = int32(Archive->Tell() + TrackHeader.GetLength());
+	MIDIREADER_EXPECT_WITH_ERROR(TrackEndPos <= (int32)Archive->TotalSize(), TEXT("%s: MIDI track %d data length as recorded in the header exceeds the amount of data in the file!"), *Filename, CurrentTrackIndex + 1);
 
 	++CurrentTrackIndex; // update track num
 
+	// Let the midi receiver know that this is a new track.
+	MIDIREADER_EXPECT_WITH_ERROR(Receiver->OnNewTrack(CurrentTrackIndex), TEXT("%s: Failed to create new Track (%d)"), *Filename, CurrentTrackIndex);
+
 	// reset state for reading a new track
 	PrevStatus = 0;
 	CurrentFileTick = 0;
@@ -384,32 +362,41 @@ void FStdMidiFileReader::ReadTrackHeader(FArchive& Archive)
 	MidiListTick = -1;  // units of desired Tick (for sorting)
 	State = EState::InTrack;
 
-	// Let the midi receiver know that this is a new track.
-	Receiver->OnNewTrack(CurrentTrackIndex);
-
 	if (Format == 0)
 	{
 		if (CurrentTrackIndex == 0)
 		{
-			CurrentTrackName = TEXT("conductor");
+			CurrentTrackName = TEXT("Conductor");
 		}
 		else
 		{
 			CurrentTrackName = TEXT("Track-0");
 		}
-		Receiver->OnText(0, CurrentTrackName, Harmonix::Midi::Constants::GMeta_TrackName);
+
+		MIDIREADER_EXPECT_WITH_ERROR(Receiver->OnText(0, CurrentTrackName, Harmonix::Midi::Constants::GMeta_TrackName), TEXT("%s: Failed to add track name to Track (%d)"), *Filename, CurrentTrackIndex);
 	}
 }
 
 // Read an event from the track
-void FStdMidiFileReader::ReadEvent(FArchive& Archive)
+void FStdMidiFileReader::ReadEvent()
 {
 	check(State == EState::InTrack);
 
+	if (FailIfReadPositionIsPastTrackEnd())
+	{
+		return;
+	}
+
 	uint8 Status, Data1 = 0;    // MIDI message bytes
 	bool RunningStatus;     // true if running status is detected
 
-	CurrentFileTick += Midi::VarLenNumber::Read(Archive); // accumulate tick
+	CurrentFileTick += Midi::VarLenNumber::Read(*Archive); // accumulate tick
+
+	if (FailIfReadPositionIsPastTrackEnd())
+	{
+		return;
+	}
+
 	CurrentTick = int32(CurrentFileTick * TickConversionFactor);  // convert to kTicksPerQuarterNoteInt
 
 	if (CurrentTick > LastTick)
@@ -430,7 +417,7 @@ void FStdMidiFileReader::ReadEvent(FArchive& Archive)
 		MidiListTick = CurrentTick;
 	}
 
-	Archive << Status;
+	*Archive << Status;
 	if (Harmonix::Midi::Constants::IsStatus(Status))  // this byte is truly a status byte
 	{
 		RunningStatus = false;
@@ -451,23 +438,29 @@ void FStdMidiFileReader::ReadEvent(FArchive& Archive)
 		Status = PrevStatus;
 	}
 
+	MIDIREADER_EXPECT_WITH_ERROR(Harmonix::Midi::Constants::IsStatus(Status), TEXT("%s (%s): Found invalid MIDI status byte (%d). MIDI data appears to be currupt."), *Filename, *CurrentTrackName, Status);
+
 	// Check if this is a system (ie, non-channel) message
 	if (Harmonix::Midi::Constants::IsSystem(Status))
 	{
-		ReadSystemEvent(CurrentTick, Status, Archive);
+		ReadSystemEvent(CurrentTick, Status);
 	}                   // All other messages are channel (non-system) messages
 	else
 	{
 		// Running status means that data1 has already been read above
 		if (!RunningStatus)
 		{
-			Archive << Data1;
+			if (FailIfBytesNotAvailableInTrackData(1)) // 1 byte needed for status byte
+			{
+				return;
+			}
+			*Archive << Data1;
 		}
-		ReadMidiEvent(CurrentTick, Status, Data1, Archive);
+		ReadMidiEvent(CurrentTick, Status, Data1);
 	}
 }
 
-void FStdMidiFileReader::ReadMidiEvent(int32 Tick, uint8 Status, uint8 Data1, FArchive& Archive)
+void FStdMidiFileReader::ReadMidiEvent(int32 Tick, uint8 Status, uint8 Data1)
 {
 	using namespace Harmonix::Midi::Constants;
 	uint8 Data2 = 0;
@@ -476,7 +469,11 @@ void FStdMidiFileReader::ReadMidiEvent(int32 Tick, uint8 Status, uint8 Data1, FA
 	{
 		// special case for NoteOn (3 byte message)
 	case GNoteOn:
-		Archive << Data2;
+		if (FailIfBytesNotAvailableInTrackData(1)) // 1 byte needed for velocity
+		{
+			return;
+		}
+		*Archive << Data2;
 		ValidMidiEvent = true;
 		// convert note-on vel=0 -> note offs
 		if (Data2 == 0)
@@ -490,7 +487,11 @@ void FStdMidiFileReader::ReadMidiEvent(int32 Tick, uint8 Status, uint8 Data1, FA
 	case GControl:
 	case GPitch:
 	case GPolyPres:
-		Archive << Data2;
+		if (FailIfBytesNotAvailableInTrackData(1)) // 1 byte needed for controller value
+		{
+			return;
+		}
+		*Archive << Data2;
 		ValidMidiEvent = true;
 		break;
 
@@ -502,8 +503,9 @@ void FStdMidiFileReader::ReadMidiEvent(int32 Tick, uint8 Status, uint8 Data1, FA
 		break;
 
 	default:
-		UE_LOG(LogMidi, Error, TEXT("%s (%s): Cannot parse event %i"), *Filename, *CurrentTrackName, (Status & GMessageTypeMask));
-		break;
+		UE_LOG(LogMIDI, Error, TEXT("%s (%s): Cannot parse event %i"), *Filename, *CurrentTrackName, (Status & GMessageTypeMask));
+		State = EState::Failed;
+		return;
 	}
 
 	if (ValidMidiEvent && TrackFilteringMode != ETrackFilteringMode::ConductorEvents)
@@ -514,7 +516,7 @@ void FStdMidiFileReader::ReadMidiEvent(int32 Tick, uint8 Status, uint8 Data1, FA
 }
 
 // Read system event from the stream.  Status byte has already been read.
-void FStdMidiFileReader::ReadSystemEvent(int32 Tick, uint8 Status, FArchive& Archive)
+void FStdMidiFileReader::ReadSystemEvent(int32 Tick, uint8 Status)
 {
 	using namespace Harmonix::Midi::Constants;
 	
@@ -526,75 +528,99 @@ void FStdMidiFileReader::ReadSystemEvent(int32 Tick, uint8 Status, FArchive& Arc
 			// All system events start with length. We'll
 			// seek past that length and ignore these events
 			// entirely.
-			int32 PacketLength = Midi::VarLenNumber::Read(Archive);
-			Archive.Seek(Archive.Tell() + PacketLength);
+			int32 PacketLength = Midi::VarLenNumber::Read(*Archive);
+			if (FailIfBytesNotAvailableInTrackData((int64)PacketLength))
+			{
+				return;
+			}
+			Archive->Seek(Archive->Tell() + PacketLength);
 		}
 		break;
 	case GFile_Meta:
 		{
+			if (FailIfBytesNotAvailableInTrackData(1)) // 1 byte needed for meta event type
+			{
+				return;
+			}
 			uint8 type;
-			Archive << type;
-			ReadMetaEvent(Tick, type, Archive);
+			*Archive << type;
+			ReadMetaEvent(Tick, type);
 		}
 		break;
 	default:
-		UE_LOG(LogMidi, Error, TEXT("%s (%s): Cannot parse system event %i"), *Filename, *CurrentTrackName, Status);
+		UE_LOG(LogMIDI, Error, TEXT("%s (%s): Cannot parse system event %i"), *Filename, *CurrentTrackName, Status);
+		State = EState::Failed;
 		break;
 	}
 }
 
 // Read a Meta Event from stream.
 // Status byte and event type have already been read.
-void FStdMidiFileReader::ReadMetaEvent(int32 Tick, uint8 Type, FArchive& Archive)
+void FStdMidiFileReader::ReadMetaEvent(int32 Tick, uint8 Type)
 {
 	using namespace Harmonix::Midi::Constants;
 	
-	int32 Length = Midi::VarLenNumber::Read(Archive);
-	int32 StartPos = int32(Archive.Tell());
+	if (FailIfReadPositionIsPastTrackEnd())
+	{
+		return;
+	}
+
+	int32 Length = Midi::VarLenNumber::Read(*Archive);
+	int32 StartPos = int32(Archive->Tell());
 	FString WorkingString;
 
 	switch (Type)
 	{
 	case GMeta_TrackName:
-		CurrentTrackName = ReadText(Archive, Length);
-		Receiver->OnText(Tick, CurrentTrackName, Type);
+		CurrentTrackName = ReadText(Length);
+		if (CurrentTrackIndex == 0)
+		{
+			// This is the conductor track. Different DAWs use different names for this hidden track,
+			// so we ignore the name we got from the file and normalize it here to "conductor"...
+			CurrentTrackName = TEXT("Conductor");
+		}
+		MIDIREADER_EXPECT_WITH_ERROR(Receiver->OnText(Tick, CurrentTrackName, Type), TEXT("%s (%s): Failed to add track name to Track"), *Filename, *CurrentTrackName);
 		break;
 	case GMeta_Copyright:
 	case GMeta_Marker:
 	case GMeta_CuePoint:
-		WorkingString = ReadText(Archive, Length);
+		WorkingString = ReadText(Length);
 		if (TrackFilteringMode != ETrackFilteringMode::NonConductorEvents)
 		{
-			Receiver->OnText(Tick, WorkingString, Type);
+			MIDIREADER_EXPECT_WITH_ERROR(Receiver->OnText(Tick, WorkingString, Type), TEXT("%s (%s): Failed to add copyright, marker, or cuepoint to Track"), *Filename, *CurrentTrackName);
 		}
 		break;
 	case GMeta_Text:
 	case GMeta_Lyric:
-		WorkingString = ReadText(Archive, Length);
+		WorkingString = ReadText(Length);
 		if (TrackFilteringMode != ETrackFilteringMode::ConductorEvents)
 		{
-			Receiver->OnText(Tick, WorkingString, Type);
+			MIDIREADER_EXPECT_WITH_ERROR(Receiver->OnText(Tick, WorkingString, Type), TEXT("%s (%s): Failed to add text or lyric to Track"), *Filename, *CurrentTrackName);
 		}
 		break;
 	case GMeta_Tempo:
 		{
+			if (FailIfBytesNotAvailableInTrackData(3)) // three bytes needed for tempo
+			{
+				return;
+			}
 			uint8 Byte1, Byte2, Byte3;
-			Archive << Byte1 << Byte2 << Byte3;
+			*Archive << Byte1 << Byte2 << Byte3;
 			int32 Tempo = (Byte1 << 16) + (Byte2 << 8) + Byte3;
 			if (TrackFilteringMode != ETrackFilteringMode::NonConductorEvents)
 			{
-				Receiver->OnTempo(Tick, Tempo);
+				MIDIREADER_EXPECT_WITH_ERROR(Receiver->OnTempo(Tick, Tempo), TEXT("%s (%s): Failed to add tempo to Track"), *Filename, *CurrentTrackName);
 			}
 		}
 		break;
 
 	case GMeta_EndOfTrack:
 		ProcessMidiList();  // in case there is anything left to send out.
-		Receiver->OnEndOfTrack(Tick);
+		MIDIREADER_EXPECT_WITH_ERROR(Receiver->OnEndOfTrack(Tick), TEXT("%s (%s): Failed ending current track read"), *Filename, *CurrentTrackName);
 		if (CurrentTrackIndex == NumTracks - 1)
 		{
 			State = EState::End;
-			Receiver->OnAllTracksRead();
+			MIDIREADER_EXPECT_WITH_ERROR(Receiver->OnAllTracksRead(), TEXT("%s: Failure telling MIDI reader all tracks are read"), *Filename);
 		}
 		else
 		{
@@ -604,26 +630,21 @@ void FStdMidiFileReader::ReadMetaEvent(int32 Tick, uint8 Type, FArchive& Archive
 
 	case GMeta_TimeSig:
 		{
-			uint8 Numerator, DenominatorExp;
-			Archive << Numerator << DenominatorExp;
-
-			if (DenominatorExp > 6)
+			if (FailIfBytesNotAvailableInTrackData(2)) // 2 bytes needed for time signature
 			{
-				UE_LOG(LogMidi, Error, TEXT("%s (%s): Time signature at %s has invalid denominator (2^%d); max is 64 (2^6)"), *Filename, *CurrentTrackName, *MidiTickFormat(Tick, BarMap.Get(), Midi::EMusicTimeStringFormat::Position), DenominatorExp);
-
-				// Continuing past here could lead to an integer overflow below.
-				break;
+				return;
 			}
+			uint8 Numerator, DenominatorExp;
+			*Archive << Numerator << DenominatorExp;
+
+			MIDIREADER_EXPECT_WITH_ERROR(DenominatorExp <= 6, TEXT("%s (%s): Time signature at %s has invalid denominator (2^%d); max is 64 (2^6)"), *Filename, *CurrentTrackName, *MidiTickFormat(Tick, BarMap.Get(), Midi::EMusicTimeStringFormat::Position), DenominatorExp);
 
 			// time signature denominator is given as a power of 2
 			//  e.g. for 6/8 time, numerator=6 and denominator=3
 			int32 Denominator = int32(FMath::Pow(2.0f, DenominatorExp));
 
-			if (Numerator == 0)
-			{
-				UE_LOG(LogMidi, Error, TEXT("%s (%s): Time signature %d/%d at %s has invalid numerator (%d)"), *Filename, *CurrentTrackName, Numerator, Denominator,
-					*MidiTickFormat(Tick, BarMap.Get(), Midi::EMusicTimeStringFormat::Position), Numerator);
-			}
+			MIDIREADER_EXPECT_WITH_ERROR(Numerator != 0, TEXT("%s (%s): Time signature %d/%d at %s has invalid numerator (%d)"), *Filename, *CurrentTrackName, Numerator, Denominator,
+				*MidiTickFormat(Tick, BarMap.Get(), Midi::EMusicTimeStringFormat::Position), Numerator);
 
 			if (TrackFilteringMode != ETrackFilteringMode::NonConductorEvents)
 			{
@@ -631,18 +652,21 @@ void FStdMidiFileReader::ReadMetaEvent(int32 Tick, uint8 Type, FArchive& Archive
 				// so we can print better error and warning messages...
 				check(Tick == 0 || BarMap->GetNumTimeSignaturePoints() > 0);
 				int32 BarIndex = BarMap->TickToBarIncludingCountIn(Tick);
-				if (!BarMap->AddTimeSignatureAtBarIncludingCountIn(BarIndex, Numerator, Denominator, true, false))
-				{
-					UE_LOG(LogMidi, Warning, TEXT("%s (%s): Time signature %d/%d at %s overlaps or conflicts with nearby time signatures"),
-						*Filename, *CurrentTrackName, Numerator, Denominator, *MidiTickFormat(Tick, BarMap.Get(), Midi::EMusicTimeStringFormat::Position));
-				}
-				// fill in our bar map
-				Receiver->OnTimeSignature(Tick, Numerator, Denominator, false);
+				MIDIREADER_EXPECT_WITH_ERROR(BarMap->AddTimeSignatureAtBarIncludingCountIn(BarIndex, Numerator, Denominator, true, false),
+					TEXT("%s (%s): Time signature %d/%d at %s overlaps or conflicts with nearby time signatures"),
+					*Filename, *CurrentTrackName, Numerator, Denominator, *MidiTickFormat(Tick, BarMap.Get(), Midi::EMusicTimeStringFormat::Position));
+				// now the 'clients' view of the world...
+				MIDIREADER_EXPECT_WITH_ERROR(Receiver->OnTimeSignature(Tick, Numerator, Denominator, false),
+					TEXT("%s (%s): Failed adding time signature change to MIDI (%d/%d)"), *Filename, *CurrentTrackName, Numerator, Denominator);
 			}
 
 			// skip over next two bytes, "MIDI clocks per metronome click"
 			// and "32nd notes per MIDI quarter note"; they're not too useful
-			Archive.Seek(Archive.Tell() + 2);
+			if (FailIfBytesNotAvailableInTrackData(2))
+			{
+				return;
+			}
+			Archive->Seek(Archive->Tell() + 2);
 		}
 		break;
 
@@ -656,8 +680,9 @@ void FStdMidiFileReader::ReadMetaEvent(int32 Tick, uint8 Type, FArchive& Archive
 		break;
 
 	default:
-		UE_LOG(LogMidi, Warning, TEXT("%s (%s): Cannot parse meta event %i"), *Filename, *CurrentTrackName, Type);
-		break;
+		UE_LOG(LogMIDI, Error, TEXT("%s (%s): Cannot parse meta event %i"), *Filename, *CurrentTrackName, Type);
+		State = EState::Failed;
+		return;
 	}
 
 	// 
@@ -667,47 +692,107 @@ void FStdMidiFileReader::ReadMetaEvent(int32 Tick, uint8 Type, FArchive& Archive
 	}
 
 	// skip to end of meta event
-	Archive.Seek(StartPos + Length);
+	Archive->Seek(StartPos + Length);
 }
 
-FString FStdMidiFileReader::ReadText(FArchive& Archive, int32 Length)
+FString FStdMidiFileReader::ReadText(int32 Length)
 {
+	// make sure length is rational...
+	MIDIREADER_EXPECT_WITH_ERROR_RETURN_VAL(Length > 0 && Length <= kMaxSupportedMidiStringSize, FString(),
+		TEXT("%s (%s): %d byte long string found in MIDI data. Maximum supported string length is %d."), *Filename, *CurrentTrackName, Length, kMaxSupportedMidiStringSize);
+
+	// make sure there are enough bytes in the file...
+	if (FailIfBytesNotAvailableInTrackData(Length))
+	{
+		return FString();
+	}
+
 	unsigned char* AsUtf8 = nullptr;
+	unsigned char* AsUtf8End = nullptr;
 	if (TextEncoding == Harmonix::Midi::Constants::EMidiTextEventEncoding::Latin1)
 	{
-		unsigned char* AsLatin = (unsigned char*)FMemory::Malloc(Length + 1);
-		Archive.Serialize(AsLatin, Length);
+		unsigned char* AsLatin = (unsigned char*)FMemory::Malloc(Length + 2);
+		MIDIREADER_EXPECT_WITH_ERROR_RETURN_VAL(AsLatin, FString(),
+			TEXT("%s (%s): Couldn't allocate memory for %d byte string."), *Filename, *CurrentTrackName, Length);
+
+		Archive->Serialize(AsLatin, Length);
+		// force some null terminators...
 		AsLatin[Length] = 0;
+		AsLatin[Length+1] = 0;
+		unsigned char* AsLatinEnd = AsLatin + Length + 2;
 		AsUtf8 = (unsigned char*)FMemory::Malloc(Length * 2 + 2);
+		if (!AsUtf8)
+		{
+			FMemory::Free(AsLatin);
+			MIDIREADER_EXPECT_WITH_ERROR_RETURN_VAL(false, FString(),
+				TEXT("%s (%s): Couldn't allocate memory for %d byte string."), *Filename, *CurrentTrackName, Length);
+		}
+		FMemory::Memset(AsUtf8, 0, Length * 2 + 2);
+		AsUtf8End = AsUtf8 + (Length * 2 + 2);
 
 		unsigned char* in = AsLatin;
 		unsigned char* out = AsUtf8;
 		while (*in)
 		{
+			if (in >= AsLatinEnd)
+			{
+				break;
+			}
 			if (*in < 128)
 			{
+				if (out >= AsUtf8End)
+				{
+					break;
+				}
 				*out++ = *in++;
 			}
 			else
 			{
+				if (out >= AsUtf8End)
+				{
+					break;
+				}
 				*out++ = 0xc0 | *in >> 6;
+				if (out >= AsUtf8End)
+				{
+					break;
+				}
 				*out++ = 0x80 | (*in++ & 0x3f);
 			}
 		}
-		*out = 0;
+		if (out < AsUtf8End)
+		{
+			*out = 0;
+		}
+		else
+		{
+			*(AsUtf8End - 1) = 0;
+		}
 
 		FMemory::Free(AsLatin);
 	}
 	else
 	{
-		AsUtf8 = (unsigned char*)FMemory::Malloc(Length + 1);
-		Archive.Serialize(AsUtf8, Length);
+		AsUtf8 = (unsigned char*)FMemory::Malloc(Length + 2);
+		MIDIREADER_EXPECT_WITH_ERROR_RETURN_VAL(AsUtf8, FString(),
+			TEXT("%s (%s): Couldn't allocate memory for %d byte string."), *Filename, *CurrentTrackName, Length);
+		Archive->Serialize(AsUtf8, Length);
+		// force some null terminators...
 		AsUtf8[Length] = 0;
+		AsUtf8[Length + 1] = 0;
+	}
+
+	if (Archive->IsError())
+	{
+		FMemory::Free(AsUtf8);
+		MIDIREADER_EXPECT_WITH_ERROR_RETURN_VAL(false, FString(),
+			TEXT("%s: MIDI track %d ... Error reading string!"), *Filename, CurrentTrackIndex);
 	}
 
 	FString AsFString = StringCast<TCHAR>((const UTF8CHAR*)AsUtf8).Get();
 	FMemory::Free(AsUtf8);
 	return AsFString;
+
 }
 
 // Queue all midi that is on the same tick, for sorting. 
@@ -726,7 +811,8 @@ void FStdMidiFileReader::ProcessMidiList()
 		EventsOnSameTick.StableSort(RawMidiLess());
 		for (auto& MidiItem : EventsOnSameTick)
 		{
-			Receiver->OnMidiMessage(MidiListTick, MidiItem.Status, MidiItem.Data1, MidiItem.Data2);
+			MIDIREADER_EXPECT_WITH_ERROR(Receiver->OnMidiMessage(MidiListTick, MidiItem.Status, MidiItem.Data1, MidiItem.Data2),
+				TEXT("%s (%s): Failed to MIDI message to Track"), *Filename, *CurrentTrackName);
 
 			// state changed because SkipCurrentTrack was called. Stop processing.
 			if (State != EState::InTrack)
@@ -737,3 +823,19 @@ void FStdMidiFileReader::ProcessMidiList()
 	}
 	EventsOnSameTick.Empty();
 }
+
+bool FStdMidiFileReader::FailIfReadPositionIsPastTrackEnd()
+{
+	MIDIREADER_EXPECT_WITH_ERROR_RETURN_VAL(Archive->Tell() < (int64)TrackEndPos, true,
+		TEXT("%s: MIDI track %d ... Finding events on track passed track data block! (ie missing 'end-of-track' event)"), *Filename, CurrentTrackIndex);
+	return false; // no error
+}
+
+bool FStdMidiFileReader::FailIfBytesNotAvailableInTrackData(int64 NumBytes)
+{
+	int64 BytesAvailable = (int64)TrackEndPos - Archive->Tell();
+	MIDIREADER_EXPECT_WITH_ERROR_RETURN_VAL(BytesAvailable >= NumBytes, true,
+		TEXT("%s: MIDI track %d ... Track data underrun. Not enough bytes in track block."), *Filename, CurrentTrackIndex);
+	return false; // no error
+}
+
diff --git a/Engine/Plugins/Runtime/Harmonix/Source/HarmonixMidi/Private/HarmonixMidi/MidiTrack.cpp b/Engine/Plugins/Runtime/Harmonix/Source/HarmonixMidi/Private/HarmonixMidi/MidiTrack.cpp
--- a/Engine/Plugins/Runtime/Harmonix/Source/HarmonixMidi/Private/HarmonixMidi/MidiTrack.cpp
+++ b/Engine/Plugins/Runtime/Harmonix/Source/HarmonixMidi/Private/HarmonixMidi/MidiTrack.cpp
@@ -24,6 +24,29 @@ FMidiTrack::FMidiTrack(const FString& Name)
 	AddEvent(FMidiEvent(0, FMidiMsg::CreateText(stringIndex, Harmonix::Midi::Constants::GMeta_TrackName)));
 }
 
+bool FMidiTrack::operator==(const FMidiTrack& Other) const
+{
+	if (Events.Num() != Other.Events.Num() || Strings.Num() != Other.Strings.Num())
+	{
+		return false;
+	}
+	for (int32 EventIndex = 0; EventIndex < Events.Num(); ++EventIndex)
+	{
+		if (Events[EventIndex] != Other.Events[EventIndex])
+		{
+			return false;
+		}
+	}
+	for (int32 StringIndex = 0; StringIndex < Strings.Num(); ++StringIndex)
+	{
+		if (Strings[StringIndex] != Other.Strings[StringIndex])
+		{
+			return false;
+		}
+	}
+	return Sorted == Other.Sorted && PrimaryMidiChannel == Other.PrimaryMidiChannel;
+}
+
 const FMidiEventList& FMidiTrack::GetEvents() const
 {
 	check(Sorted);
@@ -55,7 +78,7 @@ void FMidiTrack::SetName(const FString& InName)
 			}
 		}
 	}
-	UE_LOG(LogMidi, Warning, TEXT("Cannot set name of track to %s: track does not contain a track name event"), *InName);
+	UE_LOG(LogMIDI, Warning, TEXT("Cannot set name of track to %s: track does not contain a track name event"), *InName);
 }
 
 const FString* FMidiTrack::GetName() const
@@ -117,9 +140,9 @@ void FMidiTrack::ChangeTick(FMidiEventList::TIterator Iterator, int32 NewTick)
 	Sorted = false;
 }
 
-void FMidiTrack::WriteStdMidi(FMidiWriter& Writer)
+void FMidiTrack::WriteStdMidi(FMidiWriter& Writer) const
 {
-	Sort();
+	ensureAlwaysMsgf(Sorted, TEXT("Midi events on track \"%s\" are not sorted! Code that built this midi track should call Sort() before attempting to export Standard MIDI file."), GetName() ? **GetName() : TEXT("<unnamed>"));
 
 	for (auto& Event : Events)
 	{
@@ -233,3 +256,13 @@ int32 FMidiTrack::CopyEvents(FMidiTrack& SourceTrack,	int32 FromTick, int32 Thru
 
 	return NumCopied;
 }
+
+SIZE_T FMidiTrack::GetAllocatedSize() const
+{
+	SIZE_T AllocatedSize = sizeof(FMidiTrack) + Events.GetAllocatedSize() + Strings.GetAllocatedSize();
+	for (const FString& String : Strings)
+	{
+		AllocatedSize += String.GetAllocatedSize();
+	}
+	return AllocatedSize;
+}
diff --git a/Engine/Plugins/Runtime/Harmonix/Source/HarmonixMidi/Private/HarmonixMidi/MusicTimeSpecifier.cpp b/Engine/Plugins/Runtime/Harmonix/Source/HarmonixMidi/Private/HarmonixMidi/MusicTimeSpecifier.cpp
--- a/Engine/Plugins/Runtime/Harmonix/Source/HarmonixMidi/Private/HarmonixMidi/MusicTimeSpecifier.cpp
+++ b/Engine/Plugins/Runtime/Harmonix/Source/HarmonixMidi/Private/HarmonixMidi/MusicTimeSpecifier.cpp
@@ -44,7 +44,7 @@ namespace Midi
 			Result.Bar = FCString::Atoi(*Parts[PartIndex++]);
 			if (Result.Bar < 0 && Format == EMusicTimeStringFormat::Duration)
 			{
-				UE_LOG(LogMidi, Warning, TEXT("Unexpected bar value specified in midi duration string \"%s\". Bar in duration string should be >= 0!"),*Str);
+				UE_LOG(LogMIDI, Warning, TEXT("Unexpected bar value specified in midi duration string \"%s\". Bar in duration string should be >= 0!"),*Str);
 				Result.Bar = 0;
 			}
 		}
@@ -62,12 +62,12 @@ namespace Midi
 			Result.Beat = (float)FCString::Atoi(*Parts[PartIndex++]);
 			if (Result.Beat < 1.0f && Format == EMusicTimeStringFormat::Position)
 			{
-				UE_LOG(LogMidi, Warning, TEXT("Unexpected beat value specified in midi position string %s. Beat in position string should be >= 1!"), *Str);
+				UE_LOG(LogMIDI, Warning, TEXT("Unexpected beat value specified in midi position string %s. Beat in position string should be >= 1!"), *Str);
 				Result.Beat = 1;
 			}
 			else if (Result.Beat < 0 && Format == EMusicTimeStringFormat::Duration)
 			{
-				UE_LOG(LogMidi, Warning, TEXT("Unexpected beat value specified in midi duration string %s. Beat in duration string should be >= 0!"), *Str);
+				UE_LOG(LogMIDI, Warning, TEXT("Unexpected beat value specified in midi duration string %s. Beat in duration string should be >= 0!"), *Str);
 				Result.Beat = 0;
 			}
 		}
@@ -78,7 +78,7 @@ namespace Midi
 			float num = (float)FCString::Atoi(*Parts[PartIndex]);
 			if (num < 0)
 			{
-				UE_LOG(LogMidi, Warning, TEXT("Unexpected negative value found in fractional beat portion of midi time specifier string %s. Using 0 for fractional portion!"), *Str);
+				UE_LOG(LogMIDI, Warning, TEXT("Unexpected negative value found in fractional beat portion of midi time specifier string %s. Using 0 for fractional portion!"), *Str);
 				num = 0.0f;
 			}
 			Result.Beat +=  num / div;
diff --git a/Engine/Plugins/Runtime/Harmonix/Source/HarmonixMidi/Private/HarmonixMidi/SectionMap.cpp b/Engine/Plugins/Runtime/Harmonix/Source/HarmonixMidi/Private/HarmonixMidi/SectionMap.cpp
--- a/Engine/Plugins/Runtime/Harmonix/Source/HarmonixMidi/Private/HarmonixMidi/SectionMap.cpp
+++ b/Engine/Plugins/Runtime/Harmonix/Source/HarmonixMidi/Private/HarmonixMidi/SectionMap.cpp
@@ -4,6 +4,23 @@
 #include "Misc/RuntimeErrors.h"
 #include "Algo/BinarySearch.h"
 
+bool FSectionMap::operator==(const FSectionMap& Other) const
+{
+	if (TicksPerQuarterNote != Other.TicksPerQuarterNote || Points.Num() != Other.Points.Num())
+	{
+		return false;
+	}
+
+	for (int32 PointIndex = 0; PointIndex < Points.Num(); ++PointIndex)
+	{
+		if (Points[PointIndex] != Other.Points[PointIndex])
+		{
+			return false;
+		}
+	}
+	return true;
+}
+
 void FSectionMap::Finalize(int32 LastTick)
 {
 	FMusicMapUtl::Finalize(Points, LastTick);
diff --git a/Engine/Plugins/Runtime/Harmonix/Source/HarmonixMidi/Private/HarmonixMidi/SongMapReceiver.cpp b/Engine/Plugins/Runtime/Harmonix/Source/HarmonixMidi/Private/HarmonixMidi/SongMapReceiver.cpp
--- a/Engine/Plugins/Runtime/Harmonix/Source/HarmonixMidi/Private/HarmonixMidi/SongMapReceiver.cpp
+++ b/Engine/Plugins/Runtime/Harmonix/Source/HarmonixMidi/Private/HarmonixMidi/SongMapReceiver.cpp
@@ -4,8 +4,8 @@
 #include "HarmonixMidi/MidiReader.h"
 #include "HarmonixMidi/MusicTimeSpecifier.h"
 
-#define _LOG_SONGMAP_WARNING(t,...) UE_LOG(LogMidi,Warning,TEXT("%s (%s): " t), *Reader->GetFilename(), *Reader->GetCurrentTrackName(), ##__VA_ARGS__)
-#define _LOG_SONGMAP_ERROR(t,...) UE_LOG(LogMidi,Error,TEXT("%s (%s): " t), *Reader->GetFilename(), *Reader->GetCurrentTrackName(), ##__VA_ARGS__)
+#define _LOG_SONGMAP_WARNING(t,...) UE_LOG(LogMIDI,Warning,TEXT("%s (%s): " t), *Reader->GetFilename(), *Reader->GetCurrentTrackName(), ##__VA_ARGS__)
+#define _LOG_SONGMAP_ERROR(t,...) UE_LOG(LogMIDI,Error,TEXT("%s (%s): " t), *Reader->GetFilename(), *Reader->GetCurrentTrackName(), ##__VA_ARGS__)
 
 FSongMapReceiver::FSongMapReceiver(FSongMaps* Maps)
 	: SongMaps(Maps)
@@ -14,7 +14,7 @@ FSongMapReceiver::FSongMapReceiver(FSongMaps* Maps)
 	Reset();
 }
 
-void FSongMapReceiver::Reset()
+bool FSongMapReceiver::Reset()
 {
 	CurrentTrack      = EMidiTrack::UnknownTrack;
 	LastTick          = 0;
@@ -23,25 +23,29 @@ void FSongMapReceiver::Reset()
 	LastBeatType      = EMidiNoteAssignments::Invalid;
 	bHaveBeatFailure  = false;
 	SongMaps->EmptyAllMaps();
+	return true;
 }
 
-void FSongMapReceiver::Finalize(int32 InLastFileTick)
+bool FSongMapReceiver::Finalize(int32 InLastFileTick)
 {
 	if (InLastFileTick > LastTick)
 	{
 		LastTick = InLastFileTick;
 	}
-	SongMaps->FinalizeRead(Reader);
+	return SongMaps->FinalizeRead(Reader);
 }
 
-void FSongMapReceiver::OnNewTrack(int32 NewTrackIndex)
+bool FSongMapReceiver::OnNewTrack(int32 NewTrackIndex)
 {
 	TrackIndex = NewTrackIndex;
 	if (NewTrackIndex == 0)
+	{
 		CurrentTrack = EMidiTrack::FirstTrack;
+	}
+	return true;
 }
 
-void FSongMapReceiver::OnEndOfTrack(int32 InLastTick)
+bool FSongMapReceiver::OnEndOfTrack(int32 InLastTick)
 {
 	if (InLastTick > LastTick)
 	{
@@ -49,25 +53,26 @@ void FSongMapReceiver::OnEndOfTrack(int32 InLastTick)
 	}
 	TrackIndex = -1;
 	CurrentTrack = EMidiTrack::UnknownTrack;
+	return true;
 }
 
-void FSongMapReceiver::OnMidiMessage(int32 Tick, uint8 Status, uint8 Data1, uint8 Data2)
+bool FSongMapReceiver::OnMidiMessage(int32 Tick, uint8 Status, uint8 Data1, uint8 Data2)
 {
 	if (CurrentTrack != EMidiTrack::BeatTrack)
 	{
-		return;
+		return true;
 	}
 	
 	using namespace Harmonix::Midi::Constants;
 
 	if (GetType(Status) != GNoteOn)
 	{
-		return;
+		return true;
 	}
 
 	if ((EMidiNoteAssignments)Data1 != EMidiNoteAssignments::NormalBeatPitch && (EMidiNoteAssignments)Data1 != EMidiNoteAssignments::StrongBeatPitch)
 	{
-		return;
+		return true;
 	}
 
 
@@ -99,9 +104,10 @@ void FSongMapReceiver::OnMidiMessage(int32 Tick, uint8 Status, uint8 Data1, uint
 	SongMaps->BeatMap.AddBeat(BeatType, Tick, false);
 	LastBeatTick = Tick;
 	LastBeatType = (EMidiNoteAssignments)Data1;
+	return true;
 }
 
-void FSongMapReceiver::OnText(int32 Tick, const FString& Str, uint8 Type)
+bool FSongMapReceiver::OnText(int32 Tick, const FString& Str, uint8 Type)
 {
 	// we only care about track name
 	if (Type == Harmonix::Midi::Constants::GMeta_TrackName)
@@ -120,17 +126,30 @@ void FSongMapReceiver::OnText(int32 Tick, const FString& Str, uint8 Type)
 			break;
 		}
 	}
+	return true;
 }
 
-void FSongMapReceiver::OnTempo(int32 Tick, int32 Tempo)
+bool FSongMapReceiver::OnTempo(int32 Tick, int32 Tempo)
 {
+	using namespace Harmonix::Midi::Constants;
+
+	// Filter for reasonable tempos...
+	float Bpm = MidiTempoToBPM(Tempo);
+	if (Bpm < GMinMidiFileTempo || Bpm > GMaxMidiFileTempo)
+	{
+		_LOG_SONGMAP_ERROR("Unsupported MIDI tempo encountered (%f). Tempo must be between %f and %f.", Bpm, GMinMidiFileTempo, GMaxMidiFileTempo);
+		return false;
+	}
+
 	if (!SongMaps->TempoMap.AddTempoInfoPoint(Tempo, Tick))
 	{
 		_LOG_SONGMAP_WARNING("Tempo marker at %s (%.f bpm) conflicts with other tempo markers", *FmtTick(Tick), (60000000.0f / Tempo));
+		return false;
 	}
+	return true;
 }
 
-void FSongMapReceiver::OnTimeSignature(int32 Tick, int32 Numerator, int32 Denominator, bool FailOnError)
+bool FSongMapReceiver::OnTimeSignature(int32 Tick, int32 Numerator, int32 Denominator, bool FailOnError)
 {
 	check(Tick == 0 || SongMaps->BarMap.GetNumTimeSignaturePoints() > 0);
 	int32 BarIndex = SongMaps->BarMap.TickToBarIncludingCountIn(Tick);
@@ -138,7 +157,9 @@ void FSongMapReceiver::OnTimeSignature(int32 Tick, int32 Numerator, int32 Denomi
 	if (!SongMaps->BarMap.AddTimeSignatureAtBarIncludingCountIn(BarIndex, Numerator, Denominator, true, FailOnError))
 	{
 		_LOG_SONGMAP_WARNING("Time signature %d/%d at %s overlaps or conflicts with nearby time signatures", Numerator, Denominator, *FmtTick(Tick));
+		return false;
 	}
+	return true;
 }
 
 void FSongMapReceiver::ReadSectionTrackText(int32 Tick, const FString& Str)
diff --git a/Engine/Plugins/Runtime/Harmonix/Source/HarmonixMidi/Private/HarmonixMidi/SongMapReceiver.h b/Engine/Plugins/Runtime/Harmonix/Source/HarmonixMidi/Private/HarmonixMidi/SongMapReceiver.h
--- a/Engine/Plugins/Runtime/Harmonix/Source/HarmonixMidi/Private/HarmonixMidi/SongMapReceiver.h
+++ b/Engine/Plugins/Runtime/Harmonix/Source/HarmonixMidi/Private/HarmonixMidi/SongMapReceiver.h
@@ -10,19 +10,19 @@ public:
 	FSongMapReceiver(FSongMaps* Maps);
 	virtual ~FSongMapReceiver()	{}
 
-	void Reset() override;
-	void Finalize(int32 InLastFileTick) override;
+	virtual bool Reset() override;
+	virtual bool Finalize(int32 InLastFileTick) override;
 
 	// Track begin/end:
-	virtual void OnNewTrack(int32 NewTrackIndex) override;
-	virtual void OnEndOfTrack(int32 InLastTick) override;
-	virtual void OnAllTracksRead() override {};
+	virtual bool OnNewTrack(int32 NewTrackIndex) override;
+	virtual bool OnEndOfTrack(int32 InLastTick) override;
+	virtual bool OnAllTracksRead() override { return true; };
 
 	// Standard 1- or 2-byte MIDI message:
-	virtual void OnMidiMessage(int32 Tick, uint8 Status, uint8 Data1, uint8 Data2) override;
-	virtual void OnText(int32 Tick, const FString& Str, uint8 Type) override;
-	virtual void OnTempo(int32 Tick, int32 Tempo) override;
-	virtual void OnTimeSignature(int32 Tick, int32 Numerator, int32 Denominator, bool FailOnError = true) override;
+	virtual bool OnMidiMessage(int32 Tick, uint8 Status, uint8 Data1, uint8 Data2) override;
+	virtual bool OnText(int32 Tick, const FString& Str, uint8 Type) override;
+	virtual bool OnTempo(int32 Tick, int32 Tempo) override;
+	virtual bool OnTimeSignature(int32 Tick, int32 Numerator, int32 Denominator, bool FailOnError = true) override;
 
 	int32 GetLastTick() { return LastTick; }
 
diff --git a/Engine/Plugins/Runtime/Harmonix/Source/HarmonixMidi/Private/HarmonixMidi/SongMaps.cpp b/Engine/Plugins/Runtime/Harmonix/Source/HarmonixMidi/Private/HarmonixMidi/SongMaps.cpp
--- a/Engine/Plugins/Runtime/Harmonix/Source/HarmonixMidi/Private/HarmonixMidi/SongMaps.cpp
+++ b/Engine/Plugins/Runtime/Harmonix/Source/HarmonixMidi/Private/HarmonixMidi/SongMaps.cpp
@@ -9,14 +9,41 @@ FSongMaps::FSongMaps()
 {
 }
 
+bool FSongMaps::operator==(const FSongMaps& Other) const
+{
+	if (TrackNames.Num() != Other.TrackNames.Num())
+	{
+		return false;
+	}
+	for (int32 TrackIndex = 0; TrackIndex < TrackNames.Num(); ++TrackIndex)
+	{
+		if (TrackNames[TrackIndex] != Other.TrackNames[TrackIndex])
+		{
+			return false;
+		}
+	}
+	return	TicksPerQuarterNote == Other.TicksPerQuarterNote &&
+		TempoMap == Other.TempoMap &&
+		BarMap == Other.BarMap &&
+		BeatMap == Other.BeatMap &&
+		SectionMap == Other.SectionMap &&
+		ChordMap == Other.ChordMap &&
+		LengthData == Other.LengthData;
+}
+
 void FSongMaps::Init(int32 InTicksPerQuarterNote)
 {
 	TicksPerQuarterNote = InTicksPerQuarterNote;
 	TempoMap.SetTicksPerQuarterNote(InTicksPerQuarterNote);
+	TempoMap.Empty();
 	BarMap.SetTicksPerQuarterNote(InTicksPerQuarterNote);
+	BarMap.Empty();
 	BeatMap.SetTicksPerQuarterNote(InTicksPerQuarterNote);
+	BeatMap.Empty();
 	SectionMap.SetTicksPerQuarterNote(InTicksPerQuarterNote);
+	SectionMap.Empty();
 	ChordMap.SetTicksPerQuarterNote(InTicksPerQuarterNote);
+	ChordMap.Empty();
 }
 
 void FSongMaps::Copy(const FSongMaps& Other, int32 StartTick, int32 EndTick)
@@ -76,7 +103,7 @@ float FSongMaps::TickToMs(float Tick) const
 {
 	if (TempoMap.IsEmpty())
 	{
-		UE_LOG(LogMidi, Log, TEXT("SongMaps does not contain a Tempo Map."));
+		UE_LOG(LogMIDI, Log, TEXT("SongMaps does not contain a Tempo Map."));
 		return 0.0f;
 	}
 	return TempoMap.TickToMs(Tick);
@@ -86,7 +113,7 @@ float FSongMaps::MsToTick(float Ms) const
 {
 	if (TempoMap.IsEmpty())
 	{
-		UE_LOG(LogMidi, Log, TEXT("SongMaps does not contain a Tempo Map."));
+		UE_LOG(LogMIDI, Log, TEXT("SongMaps does not contain a Tempo Map."));
 		return 0;
 	}
 	return TempoMap.MsToTick(Ms);
@@ -150,6 +177,22 @@ void FSongMaps::EmptyAllMaps()
 	SectionMap.Empty();
 	ChordMap.Empty();
 	TrackNames.Empty();
+	LengthData.LastTick = 0;
+	LengthData.LengthBars = 0;
+	LengthData.LengthTicks = 0;
+}
+
+bool FSongMaps::IsEmpty() const
+{
+	return 	TempoMap.IsEmpty() &&
+			BarMap.IsEmpty() &&
+			BeatMap.IsEmpty() &&
+			SectionMap.IsEmpty() &&
+			ChordMap.IsEmpty() &&
+			TrackNames.IsEmpty() &&
+			LengthData.LastTick == 0 &&
+			LengthData.LengthBars == 0 &&
+			LengthData.LengthTicks == 0;
 }
 
 float FSongMaps::GetSongLengthMs() const
@@ -171,7 +214,7 @@ void FSongMaps::StringLengthToMT(const FString& LengthString, int32& OutBars, in
 {
 	if (BarMap.IsEmpty())
 	{
-		UE_LOG(LogMidi, Log, TEXT("SongMaps does not contain a Bar Map."));
+		UE_LOG(LogMIDI, Log, TEXT("SongMaps does not contain a Bar Map."));
 		return;
 	}
 
@@ -221,7 +264,7 @@ void FSongMaps::StringLengthToMT(const FString& LengthString, int32& OutBars, in
 	const FTimeSignaturePoint* TimeSignaturePoint = &BarMap.GetTimeSignaturePoint(0);
 	if (!TimeSignaturePoint)
 	{
-		UE_LOG(LogMidi, Log, TEXT("No Time Signature found in SongMaps."));
+		UE_LOG(LogMIDI, Log, TEXT("No Time Signature found in SongMaps."));
 		return;
 	}
 	OutBars += (Beats == 0 && OutTicks == 0) ? 0 : 1;
@@ -246,7 +289,7 @@ const FTempoInfoPoint* FSongMaps::GetTempoInfoForMs(float Ms) const
 {
 	if (TempoMap.IsEmpty())
 	{
-		UE_LOG(LogMidi, Log, TEXT("SongMaps does not contain a Tempo Map."));
+		UE_LOG(LogMIDI, Log, TEXT("SongMaps does not contain a Tempo Map."));
 		return nullptr;
 	}
 	int32 Tick = int32(MsToTick(Ms));
@@ -257,7 +300,7 @@ const FTempoInfoPoint* FSongMaps::GetTempoInfoForTick(int32 Tick) const
 {
 	if (TempoMap.IsEmpty())
 	{
-		UE_LOG(LogMidi, Log, TEXT("SongMaps does not contain a Tempo Map."));
+		UE_LOG(LogMIDI, Log, TEXT("SongMaps does not contain a Tempo Map."));
 		return nullptr;
 	}
 	return TempoMap.GetTempoPointAtTick(Tick);
@@ -267,7 +310,7 @@ float FSongMaps::GetTempoAtMs(float Ms) const
 {
 	if (TempoMap.IsEmpty())
 	{
-		UE_LOG(LogMidi, Log, TEXT("SongMaps does not contain a Tempo Map."));
+		UE_LOG(LogMIDI, Log, TEXT("SongMaps does not contain a Tempo Map."));
 		return 0.0f;
 	}
 	int32 Tick = int32(MsToTick(Ms));
@@ -278,7 +321,7 @@ float FSongMaps::GetTempoAtTick(int32 Tick) const
 {
 	if (TempoMap.IsEmpty())
 	{
-		UE_LOG(LogMidi, Log, TEXT("SongMaps does not contain a Tempo Map."));
+		UE_LOG(LogMIDI, Log, TEXT("SongMaps does not contain a Tempo Map."));
 		return 0.0f;
 	}
 	return TempoMap.GetTempoAtTick(Tick);
@@ -290,7 +333,7 @@ const FBeatMapPoint* FSongMaps::GetBeatAtMs(float Ms) const
 {
 	if (BeatMap.IsEmpty())
 	{
-		UE_LOG(LogMidi, Log, TEXT("SongMaps does not contain a Beat Map."));
+		UE_LOG(LogMIDI, Log, TEXT("SongMaps does not contain a Beat Map."));
 		return nullptr;
 	}
 	int32 Tick = int32(MsToTick(Ms));
@@ -337,7 +380,7 @@ const FBeatMapPoint* FSongMaps::GetBeatAtTick(int32 Tick) const
 {
 	if (BeatMap.IsEmpty())
 	{
-		UE_LOG(LogMidi, Log, TEXT("SongMaps does not contain a Beat Map."));
+		UE_LOG(LogMIDI, Log, TEXT("SongMaps does not contain a Beat Map."));
 		return nullptr;
 	}
 	return BeatMap.GetPointInfoForTick(Tick);
@@ -396,7 +439,7 @@ int32 FSongMaps::GetBeatIndexAtTick(int32 Tick) const
 {
 	if (BeatMap.IsEmpty())
 	{
-		UE_LOG(LogMidi, Log, TEXT("SongMaps does not contain a Beat Map."));
+		UE_LOG(LogMIDI, Log, TEXT("SongMaps does not contain a Beat Map."));
 		return -1;
 	}
 	return BeatMap.GetPointIndexForTick(Tick);
@@ -412,7 +455,7 @@ EMusicalBeatType FSongMaps::GetBeatTypeAtTick(int32 Tick) const
 {
 	if (BeatMap.IsEmpty())
 	{
-		UE_LOG(LogMidi, Log, TEXT("SongMaps does not contain a Beat Map."));
+		UE_LOG(LogMIDI, Log, TEXT("SongMaps does not contain a Beat Map."));
 		return EMusicalBeatType::Normal;
 	}
 	return BeatMap.GetBeatTypeAtTick(Tick);
@@ -428,7 +471,7 @@ float FSongMaps::GetBeatInPulseBarAtTick(float Tick) const
 {
 	if (BeatMap.IsEmpty())
 	{
-		UE_LOG(LogMidi, Log, TEXT("SongMaps does not contain a Beat Map."));
+		UE_LOG(LogMIDI, Log, TEXT("SongMaps does not contain a Beat Map."));
 		return 0.0f;
 	}
 	return BeatMap.GetBeatInPulseBarAtTick(Tick);
@@ -444,7 +487,7 @@ int32 FSongMaps::GetNumBeatsInPulseBarAtTick(int32 Tick) const
 {
 	if (BeatMap.IsEmpty())
 	{
-		UE_LOG(LogMidi, Log, TEXT("SongMaps does not contain a Beat Map."));
+		UE_LOG(LogMIDI, Log, TEXT("SongMaps does not contain a Beat Map."));
 		return 0;
 	}
 	return BeatMap.GetNumBeatsInPulseBarAt(Tick);
@@ -463,7 +506,7 @@ const FTimeSignature* FSongMaps::GetTimeSignatureAtTick(int32 Tick) const
 {
 	if (BarMap.IsEmpty())
 	{
-		UE_LOG(LogMidi, Log, TEXT("SongMaps does not contain a Bar Map."));
+		UE_LOG(LogMIDI, Log, TEXT("SongMaps does not contain a Bar Map."));
 		return nullptr;
 	}
 	return &BarMap.GetTimeSignatureAtTick(Tick);
@@ -473,12 +516,12 @@ const FTimeSignature* FSongMaps::GetTimeSignatureAtBar(int32 Bar) const
 {
 	if (Bar < 1)
 	{
-		UE_LOG(LogMidi, Warning, TEXT("Bar < 1 (%d) specified as a musical position! Bars are '1' based in musical positions. Using bar 1!"), Bar);
+		UE_LOG(LogMIDI, Warning, TEXT("Bar < 1 (%d) specified as a musical position! Bars are '1' based in musical positions. Using bar 1!"), Bar);
 	}
 
 	if (BarMap.IsEmpty())
 	{
-		UE_LOG(LogMidi, Log, TEXT("SongMaps does not contain a Bar Map."));
+		UE_LOG(LogMIDI, Log, TEXT("SongMaps does not contain a Bar Map."));
 		return nullptr;
 	}
 	return &BarMap.GetTimeSignatureAtBar(Bar);
@@ -494,7 +537,7 @@ float FSongMaps::GetBarIncludingCountInAtTick(float Tick) const
 {
 	if (BarMap.IsEmpty())
 	{
-		UE_LOG(LogMidi, Log, TEXT("SongMaps does not contain a Bar Map."));
+		UE_LOG(LogMIDI, Log, TEXT("SongMaps does not contain a Bar Map."));
 		return 1.0f; // 1 based music positions!
 	}
 	return BarMap.TickToFractionalBarIncludingCountIn(Tick);
@@ -510,12 +553,12 @@ float FSongMaps::GetMsPerBarAtTick(float Tick) const
 {
 	if (TempoMap.IsEmpty())
 	{
-		UE_LOG(LogMidi, Log, TEXT("SongMaps does not contain a Tempo Map."));
+		UE_LOG(LogMIDI, Log, TEXT("SongMaps does not contain a Tempo Map."));
 		return 0.0f;
 	}
 	if (BarMap.IsEmpty())
 	{
-		UE_LOG(LogMidi, Log, TEXT("SongMaps does not contain a Bar Map."));
+		UE_LOG(LogMIDI, Log, TEXT("SongMaps does not contain a Bar Map."));
 		return 0.0f;
 	}
 	float Bpm = TempoMap.GetTempoAtTick(int32(Tick)); // quarter notes per minute
@@ -673,7 +716,7 @@ const FSongSection* FSongMaps::GetSectionAtTick(int32 Tick) const
 {
 	if (SectionMap.IsEmpty())
 	{
-		UE_LOG(LogMidi, Log, TEXT("SongMaps does not contain a Section Map."));
+		UE_LOG(LogMIDI, Log, TEXT("SongMaps does not contain a Section Map."));
 		return nullptr;
 	}
 	return SectionMap.TickToSection(Tick);
@@ -683,7 +726,7 @@ const FSongSection* FSongMaps::GetSectionWithName(const FString& Name) const
 {
 	if (SectionMap.IsEmpty())
 	{
-		UE_LOG(LogMidi, Log, TEXT("SongMaps does not contain a Section Map."));
+		UE_LOG(LogMIDI, Log, TEXT("SongMaps does not contain a Section Map."));
 		return nullptr;
 	}
 	return SectionMap.FindSectionInfo(Name);
@@ -699,7 +742,7 @@ FString FSongMaps::GetSectionNameAtTick(int32 Tick) const
 {
 	if (SectionMap.IsEmpty())
 	{
-		UE_LOG(LogMidi, Log, TEXT("SongMaps does not contain a Section Map."));
+		UE_LOG(LogMIDI, Log, TEXT("SongMaps does not contain a Section Map."));
 		return FString();
 	}
 	return SectionMap.GetSectionNameAtTick(Tick);
@@ -761,7 +804,7 @@ const FChordMapPoint* FSongMaps::GetChordAtTick(int32 Tick) const
 {
 	if (ChordMap.IsEmpty())
 	{
-		UE_LOG(LogMidi, Log, TEXT("SongMaps does not contain a Chord Map."));
+		UE_LOG(LogMIDI, Log, TEXT("SongMaps does not contain a Chord Map."));
 		return nullptr;
 	}
 	return ChordMap.GetPointInfoForTick(Tick);
@@ -777,7 +820,7 @@ FName FSongMaps::GetChordNameAtTick(int32 Tick) const
 {
 	if (ChordMap.IsEmpty())
 	{
-		UE_LOG(LogMidi, Log, TEXT("SongMaps does not contain a Chord Map."));
+		UE_LOG(LogMIDI, Log, TEXT("SongMaps does not contain a Chord Map."));
 		return FName();
 	}
 	return ChordMap.GetChordNameAtTick(Tick);
diff --git a/Engine/Plugins/Runtime/Harmonix/Source/HarmonixMidi/Private/HarmonixMidi/TempoMap.cpp b/Engine/Plugins/Runtime/Harmonix/Source/HarmonixMidi/Private/HarmonixMidi/TempoMap.cpp
--- a/Engine/Plugins/Runtime/Harmonix/Source/HarmonixMidi/Private/HarmonixMidi/TempoMap.cpp
+++ b/Engine/Plugins/Runtime/Harmonix/Source/HarmonixMidi/Private/HarmonixMidi/TempoMap.cpp
@@ -38,6 +38,22 @@ float FTempoMap::GetTempoAtTick(int32 Tick) const
 	return Harmonix::Midi::Constants::MidiTempoToBPM(GetMicrosecondsPerQuarterNoteAtTick(Tick));
 }
 
+bool FTempoMap::operator==(const FTempoMap& Other) const
+{
+	if (TicksPerQuarterNote != Other.TicksPerQuarterNote || Points.Num() != Other.Points.Num())
+	{
+		return false;
+	}
+	for (int32 PointIndex = 0; PointIndex < Points.Num(); ++PointIndex)
+	{
+		if (Points[PointIndex] != Other.Points[PointIndex])
+		{
+			return false;
+		}
+	}
+	return true;
+}
+
 void FTempoMap::Empty()
 {
 	Points.Empty();
diff --git a/Engine/Plugins/Runtime/Harmonix/Source/HarmonixMidi/Public/HarmonixMidi/BarMap.h b/Engine/Plugins/Runtime/Harmonix/Source/HarmonixMidi/Public/HarmonixMidi/BarMap.h
--- a/Engine/Plugins/Runtime/Harmonix/Source/HarmonixMidi/Public/HarmonixMidi/BarMap.h
+++ b/Engine/Plugins/Runtime/Harmonix/Source/HarmonixMidi/Public/HarmonixMidi/BarMap.h
@@ -145,7 +145,7 @@ public:
 	FBarMap()
 		: TicksPerQuarterNote(Harmonix::Midi::Constants::GTicksPerQuarterNoteInt)
 	{}
-	friend bool operator==(const FBarMap& Left, const FBarMap& Right);
+	bool operator==(const FBarMap& Other) const;
 
 	void Empty();
 	void Copy(const FBarMap& Other, int32 StartTick = 0, int32 EndTick = -1);
diff --git a/Engine/Plugins/Runtime/Harmonix/Source/HarmonixMidi/Public/HarmonixMidi/BeatMap.h b/Engine/Plugins/Runtime/Harmonix/Source/HarmonixMidi/Public/HarmonixMidi/BeatMap.h
--- a/Engine/Plugins/Runtime/Harmonix/Source/HarmonixMidi/Public/HarmonixMidi/BeatMap.h
+++ b/Engine/Plugins/Runtime/Harmonix/Source/HarmonixMidi/Public/HarmonixMidi/BeatMap.h
@@ -101,7 +101,7 @@ public:
 	FBeatMap()
 		: TicksPerQuarterNote(Harmonix::Midi::Constants::GTicksPerQuarterNoteInt)
 	{}
-	friend bool operator==(const FBeatMap& Left, const FBeatMap& Right);
+	bool operator==(const FBeatMap& Other) const;
 
 	void Empty();
 	void Copy(const FBeatMap& Other, int32 StartTick = 0, int32 EndTick = -1);
diff --git a/Engine/Plugins/Runtime/Harmonix/Source/HarmonixMidi/Public/HarmonixMidi/ChordMap.h b/Engine/Plugins/Runtime/Harmonix/Source/HarmonixMidi/Public/HarmonixMidi/ChordMap.h
--- a/Engine/Plugins/Runtime/Harmonix/Source/HarmonixMidi/Public/HarmonixMidi/ChordMap.h
+++ b/Engine/Plugins/Runtime/Harmonix/Source/HarmonixMidi/Public/HarmonixMidi/ChordMap.h
@@ -49,7 +49,7 @@ public:
 	FChordProgressionMap()
 		: TicksPerQuarterNote(Harmonix::Midi::Constants::GTicksPerQuarterNoteInt)
 	{}
-	friend bool operator==(const FChordProgressionMap& Left, const FChordProgressionMap& Right);
+	bool operator==(const FChordProgressionMap& Other) const;
 
 	void Finalize(int32 LastTick);
 
diff --git a/Engine/Plugins/Runtime/Harmonix/Source/HarmonixMidi/Public/HarmonixMidi/MidiConstants.h b/Engine/Plugins/Runtime/Harmonix/Source/HarmonixMidi/Public/HarmonixMidi/MidiConstants.h
--- a/Engine/Plugins/Runtime/Harmonix/Source/HarmonixMidi/Public/HarmonixMidi/MidiConstants.h
+++ b/Engine/Plugins/Runtime/Harmonix/Source/HarmonixMidi/Public/HarmonixMidi/MidiConstants.h
@@ -4,7 +4,7 @@
 
 #include "Logging/LogMacros.h"
 
-HARMONIXMIDI_API DECLARE_LOG_CATEGORY_EXTERN(LogMidi, Log, All);
+HARMONIXMIDI_API DECLARE_LOG_CATEGORY_EXTERN(LogMIDI, Log, All);
 
 namespace Harmonix::Midi::Constants
 {
@@ -29,10 +29,13 @@ namespace Harmonix::Midi::Constants
 	inline constexpr uint8 GMinVelocity    = 0;
 	inline constexpr uint8 GMaxVelocity    = 127;
 
-	inline constexpr int32   GTicksPerQuarterNoteInt  = 960;
+	inline constexpr int32 GTicksPerQuarterNoteInt  = 960;
 	inline constexpr float GTicksPerQuarterNote = 960.0f;
 	inline constexpr float GQuarterNotesPerTick = 1.0f / GTicksPerQuarterNote;
 
+	inline constexpr float GMinMidiFileTempo = 10.0f;
+	inline constexpr float GMaxMidiFileTempo = 960.0;
+
 	//////////////////////////////////////////////////////////////////////////
 	// Constants for the SMF (Standard Midi File) format.
 
diff --git a/Engine/Plugins/Runtime/Harmonix/Source/HarmonixMidi/Public/HarmonixMidi/MidiFile.h b/Engine/Plugins/Runtime/Harmonix/Source/HarmonixMidi/Public/HarmonixMidi/MidiFile.h
--- a/Engine/Plugins/Runtime/Harmonix/Source/HarmonixMidi/Public/HarmonixMidi/MidiFile.h
+++ b/Engine/Plugins/Runtime/Harmonix/Source/HarmonixMidi/Public/HarmonixMidi/MidiFile.h
@@ -33,6 +33,28 @@ struct HARMONIXMIDI_API FMidiFileData
 		: TicksPerQuarterNote(Harmonix::Midi::Constants::GTicksPerQuarterNoteInt)
 		, LastEventTick(0)
 	{}
+
+	bool operator==(const FMidiFileData& Other) const;
+
+	void Empty()
+	{
+		MidiFileName.Empty();
+		TicksPerQuarterNote = Harmonix::Midi::Constants::GTicksPerQuarterNoteInt;
+		SongMaps.EmptyAllMaps();
+		Tracks.Empty();
+		Tracks.Emplace("Conductor");
+		LastEventTick = 0;
+	}
+
+	bool IsEmpty() const 
+	{
+		return MidiFileName.IsEmpty() &&
+			TicksPerQuarterNote == Harmonix::Midi::Constants::GTicksPerQuarterNoteInt &&
+			SongMaps.IsEmpty() &&
+			(Tracks.IsEmpty() || (Tracks.Num() == 1 && Tracks[0].GetNumEvents() == 1 && Tracks[0].GetName()->Compare("Conductor", ESearchCase::IgnoreCase) == 0)) &&
+			LastEventTick == 0;
+	}
+
 	UPROPERTY(BlueprintReadOnly, Category = "MidiFile")
 	FString MidiFileName;
 	UPROPERTY(BlueprintReadOnly, Category = "MidiFile")
@@ -64,8 +86,6 @@ struct HARMONIXMIDI_API FMidiFileData
 	* asserts valid TrackIdx
 	*/
 	void AddTimeSigChange(int32 TrackIdx, int32 Tick, int32 TimeSigNum, int32 TimeSigDenom);
-
-	friend bool operator==(const FMidiFileData& Left, const FMidiFileData& Right);
 };
 
 template<>
@@ -96,7 +116,7 @@ public:
 	// A comparison operator. This allows for differences in members related to how
 	// a midi file was generated/imported, but compares the underlying "renderable"
 	// midi data.
-	friend bool operator==(const UMidiFile& Left, const UMidiFile& Right);
+	bool operator==(const UMidiFile& Other) const;
 
 	virtual void GetAssetRegistryTags(FAssetRegistryTagsContext Context) const override;
 	virtual void PostInitProperties() override;
@@ -105,31 +125,33 @@ public:
 	void LoadStdMidiFile(
 		const FString& FilePath,
 		int32 DesiredTicksPerQuarterNote = Harmonix::Midi::Constants::GTicksPerQuarterNoteInt,
-		Harmonix::Midi::Constants::EMidiTextEventEncoding InTextEncoding = Harmonix::Midi::Constants::EMidiTextEventEncoding::UTF8);
+		Harmonix::Midi::Constants::EMidiTextEventEncoding InTextEncoding = Harmonix::Midi::Constants::EMidiTextEventEncoding::UTF8,
+		bool EnsureNotFailed = true);
 	/** A method for importing a standard midi file, with the option of providing a pointer to an FSongMaps instance that will be populated during the load. */
 	void LoadStdMidiFile(
 		void* Buffer,
 		int32 BufferSize,
 		const FString& FileName,
 		int32 DesiredTicksPerQuarterNote = Harmonix::Midi::Constants::GTicksPerQuarterNoteInt,
-		Harmonix::Midi::Constants::EMidiTextEventEncoding InTextEncoding = Harmonix::Midi::Constants::EMidiTextEventEncoding::UTF8);
+		Harmonix::Midi::Constants::EMidiTextEventEncoding InTextEncoding = Harmonix::Midi::Constants::EMidiTextEventEncoding::UTF8,
+		bool EnsureNotFailed = true);
 	/** A method for importing a standard midi file, with the option of providing a pointer to an FSongMaps instance that will be populated during the load. */
 	void LoadStdMidiFile(
 		TSharedPtr<FArchive> Archive,
 		const FString& Filename,
 		int32 DesiredTicksPerQuarterNote = Harmonix::Midi::Constants::GTicksPerQuarterNoteInt,
-		Harmonix::Midi::Constants::EMidiTextEventEncoding InTextEncoding = Harmonix::Midi::Constants::EMidiTextEventEncoding::UTF8);
+		Harmonix::Midi::Constants::EMidiTextEventEncoding InTextEncoding = Harmonix::Midi::Constants::EMidiTextEventEncoding::UTF8,
+		bool EnsureNotFailed = true);
 
 	/** A method for exporting the midi track data to a standard midi file. */
-	void SaveStdMidiFile(const FString& FilePath);
+	void SaveStdMidiFile(const FString& FilePath) const;
 	/** A method for exporting the midi track data to a standard midi file. */
-	void SaveStdMidiFile(TSharedPtr<FArchive> Archive, const FString& Filename = FString());
+	void SaveStdMidiFile(TSharedPtr<FArchive> Archive, const FString& Filename = FString()) const;
 
 	FMidiTrack* AddTrack(const FString& Name);
 
 	void Empty();
-
-	void SetConductorTrack(const FTempoMap*, const FBarMap*);
+	bool IsEmpty() const;
 
 	void SortAllTracks();
 
@@ -151,6 +173,8 @@ public:
 	FMidiTrackList& GetTracks() { return TheMidiData.Tracks; }
 	const FMidiTrackList& GetTracks() const { return TheMidiData.Tracks; }
 
+	void BuildConductorTrack();
+
 	/**
 	 * This function must be called if any changes are made to any of the tracks of this midi file. It will 
 	 * do various internal tasks necessary to assure the midi data is consistent and "playable".
@@ -206,7 +230,7 @@ public:
 	* e.g. if a midi file has length 5.875 bars, Round Down conforms it to 5 bars, Round Up conforms it to 6 bars,
 	* and Round To Nearest conforms it to 6 bars (5.875 > 5.5)
 	*/
-	void ConformMidiFileLength(EMidiFileLengthConformOption Option);
+	void ConformMidiFileLength(EMidiFileLengthConformOption Option, bool Force = false);
 
 	const FSongMaps* GetSongMaps() const { return &TheMidiData.SongMaps; }
 	FSongMaps* GetSongMaps() { return &TheMidiData.SongMaps; }
@@ -221,7 +245,6 @@ protected:
 	UPROPERTY()
 	FMidiFileData TheMidiData;
 
-
 	TSharedPtr<FMidiFileData> RenderableCopyOfMidiFileData;
 };
 
diff --git a/Engine/Plugins/Runtime/Harmonix/Source/HarmonixMidi/Public/HarmonixMidi/MidiReader.h b/Engine/Plugins/Runtime/Harmonix/Source/HarmonixMidi/Public/HarmonixMidi/MidiReader.h
--- a/Engine/Plugins/Runtime/Harmonix/Source/HarmonixMidi/Public/HarmonixMidi/MidiReader.h
+++ b/Engine/Plugins/Runtime/Harmonix/Source/HarmonixMidi/Public/HarmonixMidi/MidiReader.h
@@ -62,14 +62,6 @@ public:
 		int32 TicksPerQuarterNote = Harmonix::Midi::Constants::GTicksPerQuarterNoteInt,
 		Harmonix::Midi::Constants::EMidiTextEventEncoding InTextEncoding = Harmonix::Midi::Constants::EMidiTextEventEncoding::UTF8);
 
-	static bool ReadStdMidiFileForReceiver(
-		void* Buffer,
-		int32 BufferSize,
-		const FString& FileName,
-		IMidiReceiver* Receiver,
-		int32 TicksPerQuarterNote = Harmonix::Midi::Constants::GTicksPerQuarterNoteInt,
-		Harmonix::Midi::Constants::EMidiTextEventEncoding InTextEncoding = Harmonix::Midi::Constants::EMidiTextEventEncoding::UTF8);
-
 	virtual ~FStdMidiFileReader() {}
 
 	void ReadAllTracks();                // Read all tracks
@@ -84,41 +76,40 @@ public:
 	virtual const FString& GetCurrentTrackName() const override { return CurrentTrackName; }
 	virtual const TSharedPtr<FBarMap> GetBarMap() const override { return BarMap; }
 
-	bool GetFailed() const { return Failed; }
-
 	virtual int32 GetLastTick() const override { return LastTick; }
 
-	virtual bool IsFailed() const override { return Failed; }
+	virtual bool IsFailed() const override { return State == EState::Failed; }
 
 private:
 	void Init();
 	void ReadNextEvent();
 	void ReadNextEventImpl();
 
-	void ReadFileHeader(FArchive& Archive);
-	void ReadTrackHeader(FArchive& Archive);
-	void ReadEvent(FArchive& Archive);
-	void ReadMidiEvent(int32 Tick, uint8 Status, uint8 Data1, FArchive& Archive);
-	void ReadSystemEvent(int32 Tick, uint8 Status, FArchive& Archive);
-	void ReadMetaEvent(int32 Tick, uint8 Type, FArchive& Archive);
-	FString ReadText(FArchive& Archive, int32 Length);
+	void ReadFileHeader();
+	void ReadTrackHeader();
+	void ReadEvent();
+	void ReadMidiEvent(int32 Tick, uint8 Status, uint8 Data1);
+	void ReadSystemEvent(int32 Tick, uint8 Status);
+	void ReadMetaEvent(int32 Tick, uint8 Type);
+	FString ReadText(int32 Length);
 
 	// functions for sorting the midi before sending it out to the receivers
 	void QueueChannelMsg(int32 Tick, uint8 Status, uint8 Data1, uint8 Data2);
 	void ProcessMidiList();
+	bool FailIfReadPositionIsPastTrackEnd();
+	bool FailIfBytesNotAvailableInTrackData(int64 NumBytes);
 
 	enum class EState
 	{
+		Failed,
 		Start,     // at the very beginning (before MThd chunk)
 		NewTrack,  // at the beginning of a track (before MTrk chunk)
 		InTrack,   // in the middle of reading a track
-		End        // done
+		End,        // done
 	};
 
 	FString Filename;
-	TSharedPtr<FArchive> InputArchive;
-	TSharedPtr<FBufferReader> BufferArchive;
-	TSharedPtr<FArchiveFileReaderGeneric> FileArchive;
+	TSharedRef<FArchive> Archive;
 
 	IMidiReceiver* Receiver = nullptr;
 	Harmonix::Midi::Constants::EMidiTextEventEncoding TextEncoding = Harmonix::Midi::Constants::EMidiTextEventEncoding::Latin1;
@@ -138,6 +129,8 @@ private:
 	int32 CurrentTrackIndex = -1; // What track are we currently reading?
 	FString CurrentTrackName;
 
+	inline static constexpr int32 kMaxSupportedMidiStringSize = 2048;
+
 	enum class ETrackFilteringMode
 	{
 		None,
@@ -182,6 +175,4 @@ private:
 	// we'll keep track of our own bar map so we can display
 	// more meaningful error/warning messages 
 	TSharedPtr<FBarMap> BarMap;
-
-	bool Failed = false;
 };
diff --git a/Engine/Plugins/Runtime/Harmonix/Source/HarmonixMidi/Public/HarmonixMidi/MidiReceiver.h b/Engine/Plugins/Runtime/Harmonix/Source/HarmonixMidi/Public/HarmonixMidi/MidiReceiver.h
--- a/Engine/Plugins/Runtime/Harmonix/Source/HarmonixMidi/Public/HarmonixMidi/MidiReceiver.h
+++ b/Engine/Plugins/Runtime/Harmonix/Source/HarmonixMidi/Public/HarmonixMidi/MidiReceiver.h
@@ -26,37 +26,37 @@ public:
 	virtual ~IMidiReceiver() {}
 
 	/** The reader is about to read a new MIDI file. Now is your chance to reset your state. */
-	virtual void Reset() = 0;
+	virtual bool Reset() = 0;
 
 	/** Called when the reader has determined how many tracks are in the midi data */
-	virtual void SetNumTracks(int32 Num) {}
+	virtual bool SetNumTracks(int32 Num) { return true; }
 
 	/** called when all tracks have been read. */
-	virtual void Finalize(int32 InLastFileTick) {}
+	virtual bool Finalize(int32 InLastFileTick) { return true; }
 
-	virtual void OnNewTrack(int32 NewTrackIndex) = 0;
-	virtual void OnEndOfTrack(int32 LastTick) = 0;
-	virtual void OnAllTracksRead() = 0;
+	virtual bool OnNewTrack(int32 NewTrackIndex) = 0;
+	virtual bool OnEndOfTrack(int32 LastTick) = 0;
+	virtual bool OnAllTracksRead() = 0;
 
 	/** For standard 1- or 2-byte MIDI message */
-	virtual void OnMidiMessage(int32 Tick, uint8 Status, uint8 Data1, uint8 Data2) = 0;
+	virtual bool OnMidiMessage(int32 Tick, uint8 Status, uint8 Data1, uint8 Data2) = 0;
 
 	/**
 	 * Text, Copyright, TrackName, InstName, Lyric, Marker, CuePoint meta-event:
 	 * "type" is the type of meta-event (constants defined in MidiConstants.h) 
 	 */
-	virtual void OnText(int32 Tick, const FString& Str, uint8 Type) = 0;
+	virtual bool OnText(int32 Tick, const FString& Str, uint8 Type) = 0;
 
 	/**
 	 * Tempo Change meta - event:
 	 * tempo is in microseconds per quarter-note
 	 */
-	virtual void OnTempo(int32 Tick, int32 Tempo) {}
+	virtual bool OnTempo(int32 Tick, int32 Tempo) { return true; }
 
 	/** Time Signature meta - event:
 	 *   time signature is numerator/denominator
 	 */
-	virtual void OnTimeSignature(int32 Tick, int32 Numerator, int32 Denominator, bool FailOnError = true) {}
+	virtual bool OnTimeSignature(int32 Tick, int32 Numerator, int32 Denominator, bool FailOnError = true) { return true; }
 
 	/**
 	 * Called by the reader before it starts processing the midi data stream.
@@ -86,56 +86,75 @@ protected:
  */
 class FMidiReceiverList : public IMidiReceiver
 {
-#define _FOR_EACH_MIDI_RECEIVER(func_body) Algo::ForEach(ReceiverList, [=](IMidiReceiver* recvr)func_body);
+#define _FOR_EACH_MIDI_RECEIVER(func_body) Algo::ForEach(ReceiverList, [&](IMidiReceiver* recvr){ if (Result) func_body });
 public:
-	virtual void Reset()
+	virtual bool Reset()
 	{
-		_FOR_EACH_MIDI_RECEIVER({ recvr->Reset(); });
+		bool Result = true;
+		_FOR_EACH_MIDI_RECEIVER({ Result = recvr->Reset(); });
+		return Result;
 	}
 
-	virtual void Finalize(int32 LastFileTick)
+	virtual bool Finalize(int32 LastFileTick)
 	{
-		_FOR_EACH_MIDI_RECEIVER({ recvr->Finalize(LastFileTick); });
+		bool Result = true;
+		_FOR_EACH_MIDI_RECEIVER({ Result =  recvr->Finalize(LastFileTick); });
+		return Result;
 	}
 
-	virtual void OnNewTrack(int32 NewTrackIndex)
+	virtual bool OnNewTrack(int32 NewTrackIndex)
 	{
-		_FOR_EACH_MIDI_RECEIVER({ recvr->OnNewTrack(NewTrackIndex); });
+		bool Result = true;
+		_FOR_EACH_MIDI_RECEIVER({ Result = recvr->OnNewTrack(NewTrackIndex); });
+		return Result;
 	}
 
-	virtual void OnEndOfTrack(int32 LastTick)
+	virtual bool OnEndOfTrack(int32 LastTick)
 	{
-		_FOR_EACH_MIDI_RECEIVER({ recvr->OnEndOfTrack(LastTick); });
+		bool Result = true;
+		_FOR_EACH_MIDI_RECEIVER({ Result = recvr->OnEndOfTrack(LastTick); });
+		return Result;
 	}
 
-	virtual void OnAllTracksRead()
+	virtual bool OnAllTracksRead()
 	{
-		_FOR_EACH_MIDI_RECEIVER({ recvr->OnAllTracksRead(); });
+		bool Result = true;
+		_FOR_EACH_MIDI_RECEIVER({ Result = recvr->OnAllTracksRead(); });
+		return Result;
 	}
 
-	virtual void OnMidiMessage(int32 Tick, uint8 Status, uint8 Data1, uint8 Data2)
+	virtual bool OnMidiMessage(int32 Tick, uint8 Status, uint8 Data1, uint8 Data2)
 	{
-		_FOR_EACH_MIDI_RECEIVER({ recvr->OnMidiMessage(Tick, Status, Data1, Data2); });
+		bool Result = true;
+		_FOR_EACH_MIDI_RECEIVER({ Result = recvr->OnMidiMessage(Tick, Status, Data1, Data2); });
+		return Result;
 	}
 
-	virtual void OnText(int32 Tick, const FString& Str, uint8 Type)
+	virtual bool OnText(int32 Tick, const FString& Str, uint8 Type)
 	{
-		_FOR_EACH_MIDI_RECEIVER({ recvr->OnText(Tick, Str, Type); });
+		bool Result = true;
+		_FOR_EACH_MIDI_RECEIVER({ Result = recvr->OnText(Tick, Str, Type); });
+		return Result;
 	}
 
-	virtual void OnTempo(int32 Tick, int32 Tempo)
+	virtual bool OnTempo(int32 Tick, int32 Tempo)
 	{
-		_FOR_EACH_MIDI_RECEIVER({ recvr->OnTempo(Tick, Tempo); });
+		bool Result = true;
+		_FOR_EACH_MIDI_RECEIVER({ Result = recvr->OnTempo(Tick, Tempo); });
+		return Result;
 	}
 
-	virtual void OnTimeSignature(int32 Tick, int32 Numerator, int32 Denominator, bool FailOnError = true)
+	virtual bool OnTimeSignature(int32 Tick, int32 Numerator, int32 Denominator, bool FailOnError = true)
 	{
-		_FOR_EACH_MIDI_RECEIVER({ recvr->OnTimeSignature(Tick, Numerator, Denominator, FailOnError); });
+		bool Result = true;
+		_FOR_EACH_MIDI_RECEIVER({ Result = recvr->OnTimeSignature(Tick, Numerator, Denominator, FailOnError); });
+		return Result;
 	}
 
 	virtual void SetMidiReader(IMidiReader* InReader)
 	{
 		IMidiReceiver::SetMidiReader(InReader);
+		bool Result = true;
 		_FOR_EACH_MIDI_RECEIVER({ recvr->SetMidiReader(InReader); });
 	}
 
diff --git a/Engine/Plugins/Runtime/Harmonix/Source/HarmonixMidi/Public/HarmonixMidi/MidiTrack.h b/Engine/Plugins/Runtime/Harmonix/Source/HarmonixMidi/Public/HarmonixMidi/MidiTrack.h
--- a/Engine/Plugins/Runtime/Harmonix/Source/HarmonixMidi/Public/HarmonixMidi/MidiTrack.h
+++ b/Engine/Plugins/Runtime/Harmonix/Source/HarmonixMidi/Public/HarmonixMidi/MidiTrack.h
@@ -23,7 +23,7 @@ public:
 	FMidiTrack();
 	FMidiTrack(const FString& name);
 
-	friend bool operator==(const FMidiTrack& Left, const FMidiTrack& Right);
+	bool operator==(const FMidiTrack& Other) const;
 
 	const FMidiEventList& GetEvents() const;  		// this will ASSERT if events aren't sorted!
 	const FMidiEventList& GetUnsortedEvents() const; // this won't sort
@@ -48,7 +48,7 @@ public:
 		*/
 	void ChangeTick(FMidiEventList::TIterator Iterator, int32 NewTick);
 
-	void WriteStdMidi(FMidiWriter& writer);
+	void WriteStdMidi(FMidiWriter& writer) const;
 
 	void Sort();
 
@@ -72,6 +72,8 @@ public:
 	const FMidiTextRepository* GetTextRepository() const { return &Strings; }
 	FString GetTextForMsg(const FMidiMsg& Message) const { check(Message.MsgType() == FMidiMsg::EType::Text); return GetTextAtIndex(Message.GetTextIndex()); }
 
+	SIZE_T GetAllocatedSize() const;
+
 private:
 	UPROPERTY()
 	TArray<FMidiEvent> Events; // All the midi events
diff --git a/Engine/Plugins/Runtime/Harmonix/Source/HarmonixMidi/Public/HarmonixMidi/SectionMap.h b/Engine/Plugins/Runtime/Harmonix/Source/HarmonixMidi/Public/HarmonixMidi/SectionMap.h
--- a/Engine/Plugins/Runtime/Harmonix/Source/HarmonixMidi/Public/HarmonixMidi/SectionMap.h
+++ b/Engine/Plugins/Runtime/Harmonix/Source/HarmonixMidi/Public/HarmonixMidi/SectionMap.h
@@ -44,7 +44,7 @@ public:
 	FSectionMap()
 		: TicksPerQuarterNote(Harmonix::Midi::Constants::GTicksPerQuarterNoteInt)
 	{}
-	friend bool operator==(const FSectionMap& Left, const FSectionMap& Right);
+	bool operator==(const FSectionMap& Other) const;
 
 	void Finalize(int32 LastTick);
 
diff --git a/Engine/Plugins/Runtime/Harmonix/Source/HarmonixMidi/Public/HarmonixMidi/SongMaps.h b/Engine/Plugins/Runtime/Harmonix/Source/HarmonixMidi/Public/HarmonixMidi/SongMaps.h
--- a/Engine/Plugins/Runtime/Harmonix/Source/HarmonixMidi/Public/HarmonixMidi/SongMaps.h
+++ b/Engine/Plugins/Runtime/Harmonix/Source/HarmonixMidi/Public/HarmonixMidi/SongMaps.h
@@ -73,7 +73,7 @@ struct HARMONIXMIDI_API FSongMaps
 
 public:
 	FSongMaps();
-	friend bool operator==(const FSongMaps& Left, const FSongMaps& Right);
+	bool operator==(const FSongMaps& Other) const;
 
 	void Init(int32 InTicksPerQuarterNote);
 	void Copy(const FSongMaps& Other, int32 StartTick = 0, int32 EndTick = -1);
@@ -171,6 +171,7 @@ public:
 	void                  EmptyChordMap() { ChordMap.Empty(); }
 
 	void EmptyAllMaps();
+	bool IsEmpty() const;
 
 	FSongLengthData& GetSongLengthData() { return LengthData; }
 	const FSongLengthData& GetSongLengthData() const { return LengthData; }
diff --git a/Engine/Plugins/Runtime/Harmonix/Source/HarmonixMidi/Public/HarmonixMidi/TempoMap.h b/Engine/Plugins/Runtime/Harmonix/Source/HarmonixMidi/Public/HarmonixMidi/TempoMap.h
--- a/Engine/Plugins/Runtime/Harmonix/Source/HarmonixMidi/Public/HarmonixMidi/TempoMap.h
+++ b/Engine/Plugins/Runtime/Harmonix/Source/HarmonixMidi/Public/HarmonixMidi/TempoMap.h
@@ -65,7 +65,7 @@ public:
 	{}
 	virtual ~FTempoMap() {}
 
-	friend bool operator==(const FTempoMap& Left, const FTempoMap& Right);
+	bool operator==(const FTempoMap& Other) const;
 
 	void Empty();
 	void Copy(const FTempoMap& Other, int32 StartTick = 0, int32 EndTick = -1);
diff --git a/Engine/Plugins/Runtime/Harmonix/Source/HarmonixMidiEditor/Private/AssetDefinition_MidiFile.cpp b/Engine/Plugins/Runtime/Harmonix/Source/HarmonixMidiEditor/Private/AssetDefinition_MidiFile.cpp
--- a/Engine/Plugins/Runtime/Harmonix/Source/HarmonixMidiEditor/Private/AssetDefinition_MidiFile.cpp
+++ b/Engine/Plugins/Runtime/Harmonix/Source/HarmonixMidiEditor/Private/AssetDefinition_MidiFile.cpp
@@ -8,7 +8,11 @@
 #include "DesktopPlatformModule.h"
 #include "Framework/Application/SlateApplication.h"
 #include "Misc/MessageDialog.h"
+#include "GenericPlatform/GenericPlatformFile.h"
 #include "GenericPlatform/GenericPlatformMisc.h"
+#include "GenericPlatform/GenericPlatformProcess.h"
+#include "HAL/FileManager.h"
+#include "HAL/PlatformFileManager.h"
 
 #define LOCTEXT_NAMESPACE "Harmonix_Midi"
 
@@ -54,7 +58,32 @@ void UAssetDefinition_MidiFile::RegisterContextMenu()
 				const FSlateIcon Icon = FSlateIcon(FAppStyle::GetAppStyleSetName(), "ClassIcon.MidiFile");
 				const FToolMenuExecuteAction UIAction = FToolMenuExecuteAction::CreateStatic(&UAssetDefinition_MidiFile::ExecuteExportMidiFile);
 				InSection.AddMenuEntry("MidiFile_ExportMid", Label, ToolTip, Icon, UIAction);
-			}));
+			})
+		);
+	Section.AddDynamicEntry("MidiFile_Compare",
+		FNewToolMenuSectionDelegate::CreateLambda([](FToolMenuSection& InSection)
+			{
+				const TAttribute<FText> Label = LOCTEXT("MidiFile_Compare", "Compare Standard MIDI File Assets");
+				const TAttribute<FText> ToolTip = LOCTEXT("MidiFile_CompareMidToolTip", "Compares selected Midi File Assets");
+				const FSlateIcon Icon = FSlateIcon(FAppStyle::GetAppStyleSetName(), "ClassIcon.MidiFile");
+				FToolUIAction UIAction;
+				UIAction.ExecuteAction = FToolMenuExecuteAction::CreateStatic(&UAssetDefinition_MidiFile::ExecuteCompareMidiFiles);
+				UIAction.CanExecuteAction = FToolMenuCanExecuteAction::CreateStatic(&UAssetDefinition_MidiFile::CanExecuteCompareMidiFiles);
+				InSection.AddMenuEntry("MidiFile_CompareMid", Label, ToolTip, FSlateIcon(), UIAction);
+			})
+		);
+	Section.AddDynamicEntry("MidiFile_OpenExtern",
+		FNewToolMenuSectionDelegate::CreateLambda([](FToolMenuSection& InSection)
+			{
+				const TAttribute<FText> Label = LOCTEXT("MidiFile_OpenExtern", "View Standard MIDI File In External Application");
+				const TAttribute<FText> ToolTip = LOCTEXT("MidiFile_OpenExternToolTip", "Opens the selected Standard MIDI File in the default viewer for '.mid' files on you PC.");
+				const FSlateIcon Icon = FSlateIcon(FAppStyle::GetAppStyleSetName(), "ClassIcon.MidiFile");
+				FToolUIAction UIAction;
+				UIAction.ExecuteAction = FToolMenuExecuteAction::CreateStatic(&UAssetDefinition_MidiFile::ExecuteOpenMidiFileInExternalEditor);
+				UIAction.CanExecuteAction = FToolMenuCanExecuteAction::CreateStatic(&UAssetDefinition_MidiFile::CanExecuteOpenMidiFileInExternalEditor);
+				InSection.AddMenuEntry("MidiFile_OpenExtern", Label, ToolTip, FSlateIcon(), UIAction);
+			})
+	);
 }
 
 void UAssetDefinition_MidiFile::ExecuteExportMidiFile(const FToolMenuContext& MenuContext)
@@ -133,6 +162,98 @@ void UAssetDefinition_MidiFile::ExportAllMidiToFolder(const UContentBrowserAsset
 	}
 }
 
+void UAssetDefinition_MidiFile::ExecuteCompareMidiFiles(const FToolMenuContext& MenuContext)
+{
+	if (const UContentBrowserAssetContextMenuContext* Context = UContentBrowserAssetContextMenuContext::FindContextWithAssets(MenuContext))
+	{
+		TArray<UMidiFile*> AsMidiFiles = Context->LoadSelectedObjects<UMidiFile>();
+		if (AsMidiFiles.Num() < 2)
+		{
+			return;
+		}
+		bool TheSame = true;
+		for (int32 i = 1; i < AsMidiFiles.Num(); ++i)
+		{
+			if (*AsMidiFiles[0] != *AsMidiFiles[i])
+			{
+				TheSame = false;
+				break;
+			}
+		}
+		FText Equal = LOCTEXT("MidiFile_Compare_Equal", "The selected MIDI files are the same.");
+		FText Different = LOCTEXT("MidiFile_Compare_Different", "The selected MIDI files are different.");
+		FMessageDialog::Open(EAppMsgType::Ok, TheSame ? Equal : Different, LOCTEXT("MidiFile_Compare_Result_Title", "MIDI File Comparison..."));
+	}
+}
+
+bool UAssetDefinition_MidiFile::CanExecuteCompareMidiFiles(const FToolMenuContext& MenuContext)
+{
+	if (const UContentBrowserAssetContextMenuContext* Context = UContentBrowserAssetContextMenuContext::FindContextWithAssets(MenuContext))
+	{
+		if (Context->SelectedAssets.Num() < 2)
+		{
+			return false;
+		}
+		int32 MidiCount = 0;
+		for (const FAssetData& AssetData : Context->SelectedAssets)
+		{
+			if (AssetData.AssetClassPath.GetAssetName() == "MidiFile")
+			{
+				++MidiCount;
+			}
+		}
+		return MidiCount > 1;
+	}
+	return false;
+}
+
+void UAssetDefinition_MidiFile::ExecuteOpenMidiFileInExternalEditor(const FToolMenuContext& MenuContext)
+{
+	if (const UContentBrowserAssetContextMenuContext* Context = UContentBrowserAssetContextMenuContext::FindContextWithAssets(MenuContext))
+	{
+		TArray<UMidiFile*> AsMidiFiles = Context->LoadSelectedObjects<UMidiFile>();
+
+		IFileManager& FileManager = IFileManager::Get();
+		IPlatformFile& PlatformFile = FPlatformFileManager::Get().GetPlatformFile();
+		
+		FString UserTempPath = FDesktopPlatformModule::Get()->GetUserTempPath();
+		FString MidiTempPath = FPaths::Combine(UserTempPath, "UnrealEditor", "TempMidi");
+
+		FileManager.DeleteDirectory(*MidiTempPath, false, true);
+		FileManager.MakeDirectory(*MidiTempPath, true);
+
+		FString DestFilePrefix = FString::Format(TEXT("VIEW_ONLY_{0}"), { *AsMidiFiles[0]->GetName() });
+
+		FString DestFilePath = FPaths::CreateTempFilename(*MidiTempPath, *DestFilePrefix, TEXT(".mid"));
+		AsMidiFiles[0]->SaveStdMidiFile(*DestFilePath);
+		
+		FString DestFileAsArg = FString::Format(TEXT("\"{0}\""), {*DestFilePath});
+
+		FPlatformProcess::LaunchFileInDefaultExternalApplication(*DestFileAsArg, NULL, ELaunchVerb::Open);
+	}
+}
+
+bool UAssetDefinition_MidiFile::CanExecuteOpenMidiFileInExternalEditor(const FToolMenuContext& MenuContext)
+{
+	if (const UContentBrowserAssetContextMenuContext* Context = UContentBrowserAssetContextMenuContext::FindContextWithAssets(MenuContext))
+	{
+		if (Context->SelectedAssets.Num() != 1)
+		{
+			return false;
+		}
+		int32 MidiCount = 0;
+		for (const FAssetData& AssetData : Context->SelectedAssets)
+		{
+			if (AssetData.AssetClassPath.GetAssetName() == "MidiFile")
+			{
+				++MidiCount;
+			}
+		}
+		return MidiCount == 1;
+	}
+	return false;
+}
+
 EAppReturnType::Type UAssetDefinition_MidiFile::AskOverwrite(FString& OutPath)
 {
 	FText OutPathText = FText::FromString(*OutPath);
diff --git a/Engine/Plugins/Runtime/Harmonix/Source/HarmonixMidiEditor/Private/AssetDefinition_MidiFile.h b/Engine/Plugins/Runtime/Harmonix/Source/HarmonixMidiEditor/Private/AssetDefinition_MidiFile.h
--- a/Engine/Plugins/Runtime/Harmonix/Source/HarmonixMidiEditor/Private/AssetDefinition_MidiFile.h
+++ b/Engine/Plugins/Runtime/Harmonix/Source/HarmonixMidiEditor/Private/AssetDefinition_MidiFile.h
@@ -27,6 +27,11 @@ public:
 private:
 	static void ExecuteExportMidiFile(const FToolMenuContext& MenuContext);
 	static void ExportAllMidiToFolder(const UContentBrowserAssetContextMenuContext* Context);
+	static void ExecuteCompareMidiFiles(const FToolMenuContext& MenuContext);
+	static bool CanExecuteCompareMidiFiles(const FToolMenuContext& MenuContext);
+	static void ExecuteOpenMidiFileInExternalEditor(const FToolMenuContext& MenuContext);
+	static bool CanExecuteOpenMidiFileInExternalEditor(const FToolMenuContext& MenuContext);
+
 	static EAppReturnType::Type AskOverwrite(FString& OutPath);
 	static FString LastMidiExportFolder;
 };
diff --git a/Engine/Plugins/Runtime/Harmonix/Source/HarmonixMidiEditor/Private/MidiFileDetailCustomization.cpp b/Engine/Plugins/Runtime/Harmonix/Source/HarmonixMidiEditor/Private/MidiFileDetailCustomization.cpp
--- a/Engine/Plugins/Runtime/Harmonix/Source/HarmonixMidiEditor/Private/MidiFileDetailCustomization.cpp
+++ b/Engine/Plugins/Runtime/Harmonix/Source/HarmonixMidiEditor/Private/MidiFileDetailCustomization.cpp
@@ -2,6 +2,9 @@
 
 #include "MidiFileDetailCustomization.h"
 #include "HarmonixMidi/MidiFile.h"
+#include "Math/UnitConversion.h"
+
+#define LOCTEXT_NAMESPACE "HarmonixMIDIEditor"
 
 void FMidiFileDetailCustomization::CustomizeDetails(IDetailLayoutBuilder& DetailLayout)
 {
@@ -207,4 +210,54 @@ void FMidiFileDetailCustomization::CustomizeDetails(IDetailLayoutBuilder& Detail
 
 		]
 	];
-}
\ No newline at end of file
+
+	IDetailCategoryBuilder& TrackInfoCategory = DetailLayout.EditCategory(TEXT("Track Info"), LOCTEXT("TrackInfo", "Tracks:"));
+	const UMidiFile::FMidiTrackList& Tracks = MidiFileBeingEdited->GetTracks();
+	for (const FMidiTrack& Track : Tracks)
+	{
+		const FNumericUnit<uint64> Unit = FUnitConversion::QuantizeUnitsToBestFit(Track.GetAllocatedSize(), EUnit::Bytes);
+		FString SizeInBytes = FString::Printf(TEXT("%lld %s"), Unit.Value, FUnitConversion::GetUnitDisplayString(Unit.Units));
+
+		TrackInfoCategory.AddCustomRow(FText::FromString("Tracks"))
+		.NameContent()
+		[
+			SNew(STextBlock)
+				.Text(&Track == &Tracks[0] ? FText::FromString("Conductor (tempo & time signature)") : FText::FromString(*Track.GetName()))
+				.Font(IDetailLayoutBuilder::GetDetailFont())
+		].ValueContent()
+		[
+			SNew(SHorizontalBox)
+				+ SHorizontalBox::Slot()
+				.HAlign(HAlign_Left)
+				.VAlign(VAlign_Center)
+				.AutoWidth()
+				.Padding(8.0f, 0.0f)
+				[
+					SNew(STextBlock)
+						.Text(FText::Format(FTextFormat::FromString("{0} MIDI Events"), { Track.GetEvents().Num() }))
+						.Font(IDetailLayoutBuilder::GetDetailFont())
+				]
+				+ SHorizontalBox::Slot()
+				.HAlign(HAlign_Left)
+				.VAlign(VAlign_Center)
+				.AutoWidth()
+				.Padding(8.0f, 0.0f)
+				[
+					SNew(STextBlock)
+						.Text(FText::Format(FTextFormat::FromString("{0} Strings (including track name)"), { Track.GetTextRepository() ? Track.GetTextRepository()->Num() : 0 }))
+						.Font(IDetailLayoutBuilder::GetDetailFont())
+				]
+				+ SHorizontalBox::Slot()
+				.HAlign(HAlign_Left)
+				.VAlign(VAlign_Center)
+				.AutoWidth()
+				.Padding(8.0f, 0.0f)
+				[
+					SNew(STextBlock)
+						.Text(FText::FromString(*SizeInBytes))
+						.Font(IDetailLayoutBuilder::GetDetailFont())
+				]
+		];
+	}
+}
+#undef LOCTEXT_NAMESPACE
diff --git a/Engine/Plugins/Runtime/Harmonix/Source/HarmonixMidiTests/HarmonixMidiTests.Build.cs b/Engine/Plugins/Runtime/Harmonix/Source/HarmonixMidiTests/HarmonixMidiTests.Build.cs
--- a/Engine/Plugins/Runtime/Harmonix/Source/HarmonixMidiTests/HarmonixMidiTests.Build.cs
+++ b/Engine/Plugins/Runtime/Harmonix/Source/HarmonixMidiTests/HarmonixMidiTests.Build.cs
@@ -1,5 +1,6 @@
 // Copyright Epic Games, Inc. All Rights Reserved.
 
+using System.IO;
 using UnrealBuildTool;
 
 public class HarmonixMidiTests : ModuleRules
@@ -17,5 +18,13 @@ public class HarmonixMidiTests : ModuleRules
 				"HarmonixMidi",
 			}
 		);
+
+		// Because we are a TEST module, we are reaching "deep" into the source code
+		// of the module we are testing, which under normal circumstances one would
+		// not do. 
+		PrivateIncludePaths.AddRange(
+			new string[] {
+				Path.GetFullPath(Path.Combine(ModuleDirectory, "../HarmonixMidi/Private")),
+			});
 	}
 }
\ No newline at end of file
diff --git a/Engine/Plugins/Runtime/Harmonix/Source/HarmonixMidiTests/Private/MidiFileImportExportTests.cpp b/Engine/Plugins/Runtime/Harmonix/Source/HarmonixMidiTests/Private/MidiFileImportExportTests.cpp
--- a/Engine/Plugins/Runtime/Harmonix/Source/HarmonixMidiTests/Private/MidiFileImportExportTests.cpp
+++ b/Engine/Plugins/Runtime/Harmonix/Source/HarmonixMidiTests/Private/MidiFileImportExportTests.cpp
@@ -1,13 +1,17 @@
 // Copyright Epic Games, Inc. All Rights Reserved.
 
 #include "MidiTestUtility.h"
-#include "Serialization/MemoryWriter.h"
-#include "Serialization/MemoryReader.h"
+#include "HarmonixMidi/VarLenNumber.h"
+#include "Logging/LogMacros.h"
 #include "Misc/AutomationTest.h"
 #include "Misc/Paths.h"
+#include "Serialization/MemoryWriter.h"
+#include "Serialization/MemoryReader.h"
 
 #if WITH_DEV_AUTOMATION_TESTS
 
+DEFINE_LOG_CATEGORY_STATIC(LogMIDITest, Log, All);
+
 namespace HarmonixMidiTests::ImportExportTests
 {
 	using namespace Harmonix::Testing::Utility::MidiTestUtility;
@@ -126,5 +130,265 @@ namespace HarmonixMidiTests::ImportExportTests
 
 		return true;
 	}
+
+	#define ADD_EXPECTED_MIDI_ERROR(errmsg) AddExpectedError(errmsg, EAutomationExpectedMessageFlags::Contains, 0)
+
+	IMPLEMENT_SIMPLE_AUTOMATION_TEST(
+		FTestMidiFileMalformed,
+		"Harmonix.Midi.MidiFile.MalformedDataHandling",
+		EAutomationTestFlags::EditorContext | EAutomationTestFlags::EngineFilter)
+		bool FTestMidiFileMalformed::RunTest(const FString&)
+	{
+		//input test values for creating a midi file
+		const float FileLengthBars = 5.5f;
+		const int32 NumChannels = 2;
+		const int32 NumTracks = 4;
+		const int32 TimeSigNum6 = 6;
+		const int32 TimeSigDenum8 = 8;
+		const int32 TimeSigNum4 = 4;
+		const int32 TimeSigDenum4 = 4;
+		const int32 Tempo = 120;
+
+		//create an empty midi file to test for rounding to nearest (down)
+		UMidiFile* GeneratedMidiFile = BuildMidiFile(FileLengthBars, NumChannels, NumTracks, TimeSigNum4, TimeSigDenum4, Tempo);
+		//Add some events to the midi file for testing 
+		AddEventsToTestMidiFile(GeneratedMidiFile, NumTracks, NumChannels, FileLengthBars);
+
+		TArray<uint8> StdMidiFileBytes;
+		TSharedPtr<FMemoryWriter> StdMidiFileOut = MakeShared<FMemoryWriter>(StdMidiFileBytes, true);
+		GeneratedMidiFile->SaveStdMidiFile(StdMidiFileOut);
+
+		// Put errors in the data and check for proper handling!
+		UMidiFile* ReimportedMidiFile = NewObject<UMidiFile>();
+		TArray<uint8> FutzedMidiFileBytes = StdMidiFileBytes;
+		TSharedPtr<FMemoryReader> FutzedMidiFileIn = MakeShared<FMemoryReader>(FutzedMidiFileBytes, true);
+		FMemoryWriter FutzedMidiFileOut(FutzedMidiFileBytes, true);
+		FutzedMidiFileOut.SetByteSwapping(true);
+
+		ADD_EXPECTED_MIDI_ERROR(TEXT("MIDI import failed"));
+		ADD_EXPECTED_MIDI_ERROR(TEXT("Found invalid MIDI status byte"));
+		ADD_EXPECTED_MIDI_ERROR(TEXT("Finding events on track passed track data block!"));
+		ADD_EXPECTED_MIDI_ERROR(TEXT("MIDI file header is corrupt"));
+		ADD_EXPECTED_MIDI_ERROR(TEXT("Only type 0 or 1 MIDI files are supported"));
+		ADD_EXPECTED_MIDI_ERROR(TEXT("Format 0 file expected only one track"));
+		ADD_EXPECTED_MIDI_ERROR(TEXT("MIDI file has no tracks"));
+		ADD_EXPECTED_MIDI_ERROR(TEXT("MIDI file uses SMPTE time division."));
+		ADD_EXPECTED_MIDI_ERROR(TEXT("MIDI file Ticks Per Quarter Note is not rational"));
+		ADD_EXPECTED_MIDI_ERROR(TEXT("MIDI track header for track 0 is corrupt"));
+		ADD_EXPECTED_MIDI_ERROR(TEXT("MIDI track data length exceeds maximum length!"));
+		ADD_EXPECTED_MIDI_ERROR(TEXT("MIDI track 0 data length as recorded in the header exceeds the amount of data in the file!"));
+		ADD_EXPECTED_MIDI_ERROR(TEXT("MIDI track 0 ... Track data underrun"));
+		ADD_EXPECTED_MIDI_ERROR(TEXT("0 byte long string found in MIDI data"));
+		ADD_EXPECTED_MIDI_ERROR(TEXT("Found invalid MIDI status byte"));
+		ADD_EXPECTED_MIDI_ERROR(TEXT("Unsupported MIDI tempo encountered"));
+		ADD_EXPECTED_MIDI_ERROR(TEXT("Failed to add tempo to Track"));
+		ADD_EXPECTED_MIDI_ERROR(TEXT("Time signature 0/1 at 1:1.000 has invalid numerator"));
+		ADD_EXPECTED_MIDI_ERROR(TEXT("Time signature at 1:1.000 has invalid denominator"));
+		ADD_EXPECTED_MIDI_ERROR(TEXT("Cannot parse meta event"));
+		ADD_EXPECTED_MIDI_ERROR(TEXT("Cannot parse system event"));
+
+		// Put an error in the header ID...
+		{
+			FutzedMidiFileBytes[1] = 0xFF;
+			ReimportedMidiFile->LoadStdMidiFile(FutzedMidiFileIn, TEXT("InMemoryMidiFile"), Harmonix::Midi::Constants::GTicksPerQuarterNoteInt, Harmonix::Midi::Constants::EMidiTextEventEncoding::UTF8, false);
+			TestTrue(TEXT("Currupt Midi file header ID appropriately failed."), ReimportedMidiFile->IsEmpty());
+		}
+
+		// Put an error in the header size...
+		{
+			FutzedMidiFileBytes = StdMidiFileBytes;
+			FutzedMidiFileIn->Seek(0);
+			FutzedMidiFileBytes[4] = 0xFF;
+			ReimportedMidiFile->LoadStdMidiFile(FutzedMidiFileIn, TEXT("InMemoryMidiFile"), Harmonix::Midi::Constants::GTicksPerQuarterNoteInt, Harmonix::Midi::Constants::EMidiTextEventEncoding::UTF8, false);
+			TestTrue(TEXT("Currupt Midi file header SIZE appropriately failed."), ReimportedMidiFile->IsEmpty());
+		}
+
+		// Put an error in the Format...
+		{
+			FutzedMidiFileBytes = StdMidiFileBytes;
+			FutzedMidiFileIn->Seek(0);
+			FutzedMidiFileOut.Seek(8);
+			int16 BadFormat = 0x7EEF;
+			FutzedMidiFileOut << BadFormat;
+			ReimportedMidiFile->LoadStdMidiFile(FutzedMidiFileIn, TEXT("InMemoryMidiFile"), Harmonix::Midi::Constants::GTicksPerQuarterNoteInt, Harmonix::Midi::Constants::EMidiTextEventEncoding::UTF8, false);
+			TestTrue(TEXT("Currupt Midi file format failed."), ReimportedMidiFile->IsEmpty());
+		}
+
+		// Put an error in the Number of tracks...
+		{
+			FutzedMidiFileBytes = StdMidiFileBytes;
+			FutzedMidiFileIn->Seek(0);
+			FutzedMidiFileOut.Seek(8);
+			int16 Format = 0;
+			int16 TrackNum = 3; // only one track allowed in format 0 file!
+			FutzedMidiFileOut << Format;
+			FutzedMidiFileOut << TrackNum;
+			ReimportedMidiFile->LoadStdMidiFile(FutzedMidiFileIn, TEXT("InMemoryMidiFile"), Harmonix::Midi::Constants::GTicksPerQuarterNoteInt, Harmonix::Midi::Constants::EMidiTextEventEncoding::UTF8, false);
+			TestTrue(TEXT("Currupt Midi file: Wrong number of tracks!"), ReimportedMidiFile->IsEmpty());
+		}
+
+		// Put an error in the Number of tracks...
+		{
+			FutzedMidiFileBytes = StdMidiFileBytes;
+			FutzedMidiFileIn->Seek(0);
+			FutzedMidiFileOut.Seek(8);
+			int16 Format = 1;
+			int16 TrackNum = -1; // negative number of tracks!
+			FutzedMidiFileOut << Format;
+			FutzedMidiFileOut << TrackNum;
+			ReimportedMidiFile->LoadStdMidiFile(FutzedMidiFileIn, TEXT("InMemoryMidiFile"), Harmonix::Midi::Constants::GTicksPerQuarterNoteInt, Harmonix::Midi::Constants::EMidiTextEventEncoding::UTF8, false);
+			TestTrue(TEXT("Currupt Midi file: Wrong number of tracks!"), ReimportedMidiFile->IsEmpty());
+		}
+
+		// Put an error in the TicksPerQuarter...
+		{
+			FutzedMidiFileBytes = StdMidiFileBytes;
+			FutzedMidiFileIn->Seek(0);
+			FutzedMidiFileOut.Seek(12);
+			uint16 TPQ = 0x8000;
+			FutzedMidiFileOut << TPQ;
+			ReimportedMidiFile->LoadStdMidiFile(FutzedMidiFileIn, TEXT("InMemoryMidiFile"), Harmonix::Midi::Constants::GTicksPerQuarterNoteInt, Harmonix::Midi::Constants::EMidiTextEventEncoding::UTF8, false);
+			TestTrue(TEXT("Currupt Midi file: Bad Ticks Per Quarter Note. File is SMPTE!"), ReimportedMidiFile->IsEmpty());
+		}
+
+		// Put an error in the TicksPerQuarter...
+		{
+			FutzedMidiFileBytes = StdMidiFileBytes;
+			FutzedMidiFileIn->Seek(0);
+			FutzedMidiFileOut.Seek(12);
+			uint16 TPQ = 481;
+			FutzedMidiFileOut << TPQ;
+			ReimportedMidiFile->LoadStdMidiFile(FutzedMidiFileIn, TEXT("InMemoryMidiFile"), Harmonix::Midi::Constants::GTicksPerQuarterNoteInt, Harmonix::Midi::Constants::EMidiTextEventEncoding::UTF8, false);
+			TestTrue(TEXT("Currupt Midi file: Bad Ticks Per Quarter Note. Not evenly divisible by 48!"), ReimportedMidiFile->IsEmpty());
+		}
+
+		// Put an error in the Track Header ID...
+		{
+			FutzedMidiFileBytes = StdMidiFileBytes;
+			FutzedMidiFileIn->Seek(0);
+			FutzedMidiFileBytes[14] = 'Z';
+			ReimportedMidiFile->LoadStdMidiFile(FutzedMidiFileIn, TEXT("InMemoryMidiFile"), Harmonix::Midi::Constants::GTicksPerQuarterNoteInt, Harmonix::Midi::Constants::EMidiTextEventEncoding::UTF8, false);
+			TestTrue(TEXT("Currupt Midi file: Bad Track Header ID!"), ReimportedMidiFile->IsEmpty());
+		}
+
+		// Put an error in the Track Header Size...
+		{
+			FutzedMidiFileBytes = StdMidiFileBytes;
+			FutzedMidiFileIn->Seek(0);
+			FutzedMidiFileOut.Seek(18);
+			int32 TrackDataSize = std::numeric_limits<int32>::max();
+			FutzedMidiFileOut << TrackDataSize;
+			ReimportedMidiFile->LoadStdMidiFile(FutzedMidiFileIn, TEXT("InMemoryMidiFile"), Harmonix::Midi::Constants::GTicksPerQuarterNoteInt, Harmonix::Midi::Constants::EMidiTextEventEncoding::UTF8, false);
+			TestTrue(TEXT("Currupt Midi file: Bad Track Header SIZE!"), ReimportedMidiFile->IsEmpty());
+		}
+
+		// Put an error in the Track Header Size...
+		{
+			FutzedMidiFileBytes = StdMidiFileBytes;
+			FutzedMidiFileIn->Seek(0);
+			FutzedMidiFileOut.Seek(18);
+			int32 TrackDataSize = StdMidiFileBytes.Num();
+			FutzedMidiFileOut << TrackDataSize;
+			ReimportedMidiFile->LoadStdMidiFile(FutzedMidiFileIn, TEXT("InMemoryMidiFile"), Harmonix::Midi::Constants::GTicksPerQuarterNoteInt, Harmonix::Midi::Constants::EMidiTextEventEncoding::UTF8, false);
+			TestTrue(TEXT("Currupt Midi file: Bad Track Header SIZE!"), ReimportedMidiFile->IsEmpty());
+		}
+
+		// Put an error in the Track Header Size...
+		{
+			FutzedMidiFileBytes = StdMidiFileBytes;
+			FutzedMidiFileIn->Seek(0);
+			FutzedMidiFileOut.Seek(18);
+			int32 TrackDataSize = 15;
+			FutzedMidiFileOut << TrackDataSize;
+			ReimportedMidiFile->LoadStdMidiFile(FutzedMidiFileIn, TEXT("InMemoryMidiFile"), Harmonix::Midi::Constants::GTicksPerQuarterNoteInt, Harmonix::Midi::Constants::EMidiTextEventEncoding::UTF8, false);
+			TestTrue(TEXT("Currupt Midi file: Bad Track Header SIZE!"), ReimportedMidiFile->IsEmpty());
+		}
+
+		// Put an error in the track name length...
+		{
+			FutzedMidiFileBytes = StdMidiFileBytes;
+			FutzedMidiFileIn->Seek(0);
+			FutzedMidiFileOut.Seek(25);
+			uint32 StringLength = 0;
+			Midi::VarLenNumber::Write(FutzedMidiFileOut, StringLength);
+			ReimportedMidiFile->LoadStdMidiFile(FutzedMidiFileIn, TEXT("InMemoryMidiFile"), Harmonix::Midi::Constants::GTicksPerQuarterNoteInt, Harmonix::Midi::Constants::EMidiTextEventEncoding::UTF8, false);
+			TestTrue(TEXT("Currupt Midi file: Bad string length (zero)!"), ReimportedMidiFile->IsEmpty());
+
+			for (int i = 1; i < 30; ++i)
+			{
+				if (i == 9 || i == 13 || i == 16 || i == 24)
+				{
+					// These lengths result in valid'ish strings and parse-able midi messages to follow, so 
+					// don't bother testing...
+					continue;
+				}
+
+				UE_LOG(LogMIDITest, Log, TEXT("Testing string length %d..."), i);
+				FutzedMidiFileBytes = StdMidiFileBytes;
+				FutzedMidiFileIn->Seek(0);
+				FutzedMidiFileOut.Seek(25);
+				StringLength = i;
+				Midi::VarLenNumber::Write(FutzedMidiFileOut, StringLength);
+				ReimportedMidiFile->LoadStdMidiFile(FutzedMidiFileIn, TEXT("InMemoryMidiFile"), Harmonix::Midi::Constants::GTicksPerQuarterNoteInt, Harmonix::Midi::Constants::EMidiTextEventEncoding::UTF8, false);
+				TestTrue(*FString::Format(TEXT("Currupt Midi file: Bad string length {0}!"), {i}), ReimportedMidiFile->IsEmpty());
+			}
+		}
+
+		// Put an error in the Tempo...
+		{
+			FutzedMidiFileBytes = StdMidiFileBytes;
+			FutzedMidiFileIn->Seek(0);
+			FutzedMidiFileBytes[39] = 0;
+			FutzedMidiFileBytes[40] = 0;
+			FutzedMidiFileBytes[41] = 0;
+			ReimportedMidiFile->LoadStdMidiFile(FutzedMidiFileIn, TEXT("InMemoryMidiFile"), Harmonix::Midi::Constants::GTicksPerQuarterNoteInt, Harmonix::Midi::Constants::EMidiTextEventEncoding::UTF8, false);
+			TestTrue(TEXT("Currupt Midi file: Bad Tempo 0!"), ReimportedMidiFile->IsEmpty());
+
+			FutzedMidiFileBytes = StdMidiFileBytes;
+			FutzedMidiFileIn->Seek(0);
+			FutzedMidiFileBytes[39] = 0;
+			FutzedMidiFileBytes[40] = 0;
+			FutzedMidiFileBytes[41] = 1;
+			ReimportedMidiFile->LoadStdMidiFile(FutzedMidiFileIn, TEXT("InMemoryMidiFile"), Harmonix::Midi::Constants::GTicksPerQuarterNoteInt, Harmonix::Midi::Constants::EMidiTextEventEncoding::UTF8, false);
+			TestTrue(TEXT("Currupt Midi file: Bad Tempo 1!"), ReimportedMidiFile->IsEmpty());
+		}
+
+		// Put an error in the Time Signature...
+		{
+			FutzedMidiFileBytes = StdMidiFileBytes;
+			FutzedMidiFileIn->Seek(0);
+			FutzedMidiFileBytes[46] = 0;
+			FutzedMidiFileBytes[47] = 0;
+			ReimportedMidiFile->LoadStdMidiFile(FutzedMidiFileIn, TEXT("InMemoryMidiFile"), Harmonix::Midi::Constants::GTicksPerQuarterNoteInt, Harmonix::Midi::Constants::EMidiTextEventEncoding::UTF8, false);
+			TestTrue(TEXT("Currupt Midi file: Bad Time Signature 0!"), ReimportedMidiFile->IsEmpty());
+
+			FutzedMidiFileBytes = StdMidiFileBytes;
+			FutzedMidiFileIn->Seek(0);
+			FutzedMidiFileBytes[46] = 127;
+			FutzedMidiFileBytes[47] = 127;
+			ReimportedMidiFile->LoadStdMidiFile(FutzedMidiFileIn, TEXT("InMemoryMidiFile"), Harmonix::Midi::Constants::GTicksPerQuarterNoteInt, Harmonix::Midi::Constants::EMidiTextEventEncoding::UTF8, false);
+			TestTrue(TEXT("Currupt Midi file: Bad Time Signature 1!"), ReimportedMidiFile->IsEmpty());
+		}
+
+		// Put an invalid meta event type...
+		{
+			FutzedMidiFileBytes = StdMidiFileBytes;
+			FutzedMidiFileIn->Seek(0);
+			FutzedMidiFileBytes[24] = 0x10;
+			ReimportedMidiFile->LoadStdMidiFile(FutzedMidiFileIn, TEXT("InMemoryMidiFile"), Harmonix::Midi::Constants::GTicksPerQuarterNoteInt, Harmonix::Midi::Constants::EMidiTextEventEncoding::UTF8, false);
+			TestTrue(TEXT("Currupt Midi file: unsupported meta event type!"), ReimportedMidiFile->IsEmpty());
+		}
+
+		// Put an unsupported system status message...
+		{
+			FutzedMidiFileBytes = StdMidiFileBytes;
+			FutzedMidiFileIn->Seek(0);
+			FutzedMidiFileBytes[23] = 0xf3;
+			ReimportedMidiFile->LoadStdMidiFile(FutzedMidiFileIn, TEXT("InMemoryMidiFile"), Harmonix::Midi::Constants::GTicksPerQuarterNoteInt, Harmonix::Midi::Constants::EMidiTextEventEncoding::UTF8, false);
+			TestTrue(TEXT("Currupt Midi file: unsupported system status!"), ReimportedMidiFile->IsEmpty());
+		}
+
+		return true;
+	}
 }
 #endif
\ No newline at end of file
diff --git a/Engine/Plugins/Runtime/Harmonix/Source/HarmonixMidiTests/Private/MidiTestUtility.cpp b/Engine/Plugins/Runtime/Harmonix/Source/HarmonixMidiTests/Private/MidiTestUtility.cpp
--- a/Engine/Plugins/Runtime/Harmonix/Source/HarmonixMidiTests/Private/MidiTestUtility.cpp
+++ b/Engine/Plugins/Runtime/Harmonix/Source/HarmonixMidiTests/Private/MidiTestUtility.cpp
@@ -21,6 +21,8 @@ namespace Harmonix::Testing::Utility::MidiTestUtility
 		TheMidiFile->GetSongMaps()->GetTempoMap().AddTempoInfoPoint(Midi::Constants::BPMToMidiTempo(InTempo), 0);
 		TheMidiFile->GetSongMaps()->GetBarMap().AddTimeSignatureAtBarIncludingCountIn(0, InTimeSigNum, InTimeSigDenom);
 		TheMidiFile->GetSongMaps()->GetBarMap().SetTicksPerQuarterNote(DefaultTicksPerQuarter);
+		// Whenever you change the song maps you need to build the corresponding midi conductor track...
+		TheMidiFile->BuildConductorTrack();
 
 		//Add tracks to Midi file according to input argument
 		for (int32 Track = 1; Track < InNumTracks; ++Track)
@@ -121,164 +123,3 @@ namespace Harmonix::Testing::Utility::MidiTestUtility
 		CurrentTrack->AddEvent(CurrentPolyPresEvent);
 	}
 }
-
-bool operator==(const FBarMap& Left, const FBarMap& Right)
-{
-	if (Left.StartBar != Right.StartBar || Left.TicksPerQuarterNote != Right.TicksPerQuarterNote || Left.Points.Num() != Right.Points.Num())
-	{
-		return false;
-	}
-	for (int32 PointIndex = 0; PointIndex < Left.Points.Num(); ++PointIndex)
-	{
-		if (Left.Points[PointIndex] != Right.Points[PointIndex])
-		{
-			return false;
-		}
-	}
-	return true;
-}
-
-bool operator==(const FBeatMap& Left, const FBeatMap& Right)
-{
-	if (Left.TicksPerQuarterNote != Right.TicksPerQuarterNote || Left.Points.Num() != Right.Points.Num() || Left.Bars.Num() != Right.Bars.Num())
-	{
-		return false;
-	}
-	for (int32 PointIndex = 0; PointIndex < Left.Points.Num(); ++PointIndex)
-	{
-		if (Left.Points[PointIndex] != Right.Points[PointIndex])
-		{
-			return false;
-		}
-	}
-
-	for (int32 BarIndex = 0; BarIndex < Left.Bars.Num(); ++BarIndex)
-	{
-		if (Left.Bars[BarIndex] != Right.Bars[BarIndex])
-		{
-			return false;
-		}
-	}
-	return true;
-}
-
-bool operator==(const FChordProgressionMap& Left, const FChordProgressionMap& Right)
-{
-	if (Left.TicksPerQuarterNote != Right.TicksPerQuarterNote || Left.Points.Num() != Right.Points.Num() || Left.ChordTrackIndex != Right.ChordTrackIndex)
-	{
-		return false;
-	}
-
-	for (int32 PointIndex = 0; PointIndex < Left.Points.Num(); ++PointIndex)
-	{
-		if (Left.Points[PointIndex] != Right.Points[PointIndex])
-		{
-			return false;
-		}
-	}
-	return true;
-}
-
-bool operator==(const FSectionMap& Left, const FSectionMap& Right)
-{
-	if (Left.TicksPerQuarterNote != Right.TicksPerQuarterNote || Left.Points.Num() != Right.Points.Num())
-	{
-		return false;
-	}
-
-	for (int32 PointIndex = 0; PointIndex < Left.Points.Num(); ++PointIndex)
-	{
-		if (Left.Points[PointIndex] != Right.Points[PointIndex])
-		{
-			return false;
-		}
-	}
-	return true;
-}
-
-bool operator==(const FTempoMap& Left, const FTempoMap& Right)
-{
-	if (Left.TicksPerQuarterNote != Right.TicksPerQuarterNote || Left.Points.Num() != Right.Points.Num())
-	{
-		return false;
-	}
-	for (int32 PointIndex = 0; PointIndex < Left.Points.Num(); ++PointIndex)
-	{
-		if (Left.Points[PointIndex] != Right.Points[PointIndex])
-		{
-			return false;
-		}
-	}
-	return true;
-}
-
-bool operator==(const FSongMaps& Left, const FSongMaps& Right)
-{
-	if (Left.TrackNames.Num() != Right.TrackNames.Num())
-	{
-		return false;
-	}
-	for (int32 TrackIndex = 0; TrackIndex < Left.TrackNames.Num(); ++TrackIndex)
-	{
-		if (Left.TrackNames[TrackIndex] != Right.TrackNames[TrackIndex])
-		{
-			return false;
-		}
-	}
-	return	Left.TicksPerQuarterNote == Right.TicksPerQuarterNote &&
-		Left.TempoMap == Right.TempoMap &&
-		Left.BarMap == Right.BarMap &&
-		Left.BeatMap == Right.BeatMap &&
-		Left.SectionMap == Right.SectionMap &&
-		Left.ChordMap == Right.ChordMap &&
-		Left.LengthData == Right.LengthData;
-}
-
-
-bool operator==(const FMidiTrack& Left, const FMidiTrack& Right)
-{
-	if (Left.Events.Num() != Right.Events.Num() || Left.Strings.Num() != Right.Strings.Num())
-	{
-		return false;
-	}
-	for (int32 EventIndex = 0; EventIndex < Left.Events.Num(); ++EventIndex)
-	{
-		if (Left.Events[EventIndex] != Right.Events[EventIndex])
-		{
-			return false;
-		}
-	}
-	for (int32 StringIndex = 0; StringIndex < Left.Strings.Num(); ++StringIndex)
-	{
-		if (Left.Strings[StringIndex] != Right.Strings[StringIndex])
-		{
-			return false;
-		}
-	}
-	return Left.Sorted == Right.Sorted && Left.PrimaryMidiChannel == Right.PrimaryMidiChannel;
-}
-
-bool operator==(const FMidiFileData& Left, const FMidiFileData& Right)
-{
-	if (Left.Tracks.Num() != Right.Tracks.Num())
-	{
-		return false;
-	}
-	for (int32 TrackIndex = 0; TrackIndex < Left.Tracks.Num(); ++TrackIndex)
-	{
-		if (Left.Tracks[TrackIndex] != Right.Tracks[TrackIndex])
-		{
-			return false;
-		}
-	}
-
-	return	Left.MidiFileName == Right.MidiFileName &&
-			Left.TicksPerQuarterNote == Right.TicksPerQuarterNote &&
-			Left.SongMaps == Right.SongMaps &&
-			Left.LastEventTick == Right.LastEventTick;
-}
-
-bool operator==(const UMidiFile& Left, const UMidiFile& Right)
-{
-	return Left.TheMidiData == Right.TheMidiData;
-}
diff --git a/Engine/Plugins/Runtime/Harmonix/Source/HarmonixMidiTests/Private/MidiTestUtility.h b/Engine/Plugins/Runtime/Harmonix/Source/HarmonixMidiTests/Private/MidiTestUtility.h
--- a/Engine/Plugins/Runtime/Harmonix/Source/HarmonixMidiTests/Private/MidiTestUtility.h
+++ b/Engine/Plugins/Runtime/Harmonix/Source/HarmonixMidiTests/Private/MidiTestUtility.h
@@ -46,5 +46,3 @@ namespace Harmonix::Testing::Utility::MidiTestUtility
 	void AddPolyPresEventsToFile(UMidiFile* InFile, uint8 InNoteNumber, uint8 InPolyPresValue, int32 InTrackIndex, int32 InChannel, float InBarIndex);
 
 }
-
-bool operator==(const UMidiFile& Left, const UMidiFile& Right);
