When an externally clocked MIDI Player's midi file changes, make sure we properly udpate the offset tick for looping and seek our output clock appropriately

Override the "SetupNewMidiFile" method.
- Look at our external clock events and seek to the first event so we're ready to prepare the rest of them en execute
- If there are no clock events yet, simply seek based on the input clock's current tick.

#tests PIE

diff --git a/Engine/Plugins/Runtime/Harmonix/Source/HarmonixMetasound/Private/HarmonixMetasound/Nodes/MidiPlayerNode.cpp b/Engine/Plugins/Runtime/Harmonix/Source/HarmonixMetasound/Private/HarmonixMetasound/Nodes/MidiPlayerNode.cpp
--- a/Engine/Plugins/Runtime/Harmonix/Source/HarmonixMetasound/Private/HarmonixMetasound/Nodes/MidiPlayerNode.cpp
+++ b/Engine/Plugins/Runtime/Harmonix/Source/HarmonixMetasound/Private/HarmonixMetasound/Nodes/MidiPlayerNode.cpp
@@ -121,7 +121,7 @@ namespace HarmonixMetasound
 
 	protected:
 		virtual void InitTransportImpl() override;
-
+		virtual void SetupNewMidiFile(const FMidiFileProxyPtr& NewMidi) override;
 	private:
 		//** INPUTS **********************************
 		FMidiClockReadRef MidiClockIn;
@@ -578,6 +578,52 @@ namespace HarmonixMetasound
 		Init(*TransportInPin, MoveTemp(InitFn));
 	}
 
+	void FExternallyClockedMidiPlayerOperator::SetupNewMidiFile(const FMidiFileProxyPtr& NewMidi)
+	{
+		FMidiPlayerOperator::SetupNewMidiFile(NewMidi);
+		if (*LoopInPin)
+		{
+			const TArray<FMidiClockEvent>& ClockEvents = MidiClockIn->GetMidiClockEventsInBlock();
+			if (ClockEvents.Num() > 0)
+			{
+				for (const FMidiClockEvent& Event : ClockEvents)
+				{
+					// Skip loop events and reset events, just like we do in "Execute"
+					if (Event.Type == FMidiClockEvent::EType::Loop
+					 || Event.Type == FMidiClockEvent::EType::Reset)
+					{
+						continue;
+					}
+
+					// notice that we take the "Tick1" of the event, and not "Tick2".
+					// That's because we want to be synced up for when we process "Tick2" later in execute
+					UpdateLoopOffsetTickFromTick(Event.Tick1);
+					int32 Tick = Event.Tick1 - LoopOffsetTick;
+
+					float Ms = MidiClockOut->GetSongMaps().TickToMs(Tick + 1);
+					FMusicSeekTarget SeekTarget;
+					SeekTarget.Type = ESeekPointType::Millisecond;
+					SeekTarget.Ms = Ms;
+					MidiClockOut->SeekTo(SeekTarget, PrerollBars);
+					break;
+				}
+
+				
+			}
+			else
+			{
+				UpdateLoopOffsetTickFromTick(MidiClockIn->GetCurrentMidiTick());
+				int32 Tick = MidiClockIn->GetCurrentMidiTick() - LoopOffsetTick;
+				float Ms = MidiClockOut->GetSongMaps().TickToMs(Tick + 1);
+				FMusicSeekTarget SeekTarget;
+				SeekTarget.Type = ESeekPointType::Millisecond;
+				SeekTarget.Ms = Ms;
+				MidiClockOut->SeekTo(SeekTarget, PrerollBars);
+			}
+		}
+	}
+
+
 	void FExternallyClockedMidiPlayerOperator::UpdateLoopOffsetTickFromTick(int32 Tick)
 	{
 		// only update our loop offset if we're actually looping
