Clock->PrepareBlock before Clock->WriteAdvance in Harmonix Node tests.

#tests Harmonix Tests

diff --git a/Engine/Plugins/Runtime/Harmonix/Source/HarmonixMetasoundTests/Private/Nodes/MidiPulseGeneratorNodeTests.cpp b/Engine/Plugins/Runtime/Harmonix/Source/HarmonixMetasoundTests/Private/Nodes/MidiPulseGeneratorNodeTests.cpp
--- a/Engine/Plugins/Runtime/Harmonix/Source/HarmonixMetasoundTests/Private/Nodes/MidiPulseGeneratorNodeTests.cpp
+++ b/Engine/Plugins/Runtime/Harmonix/Source/HarmonixMetasoundTests/Private/Nodes/MidiPulseGeneratorNodeTests.cpp
@@ -34,6 +34,7 @@ namespace HarmonixMetasound::Nodes::MidiPulseGeneratorNode::Tests
 		for (int32 BlockIdx = 0; BlockIdx < NumBlocks; ++BlockIdx)
 		{
 			// Advance the clock, which will advance the play cursor in the pulse generators
+			(*Clock)->PrepareBlock();
 			(*Clock)->WriteAdvance(0, Generator->OperatorSettings.GetNumFramesPerBlock());
 
 			// Process
