[Gentech] Add support for dynamic pool growth for pitch shifters. Pools do not shrink back to the original size after extra instances are returned.

#rnx

diff --git a/Engine/Plugins/Runtime/Harmonix/Source/HarmonixDsp/Private/HarmonixDsp/Stretchers/SmbPitchShifterConfig.h b/Engine/Plugins/Runtime/Harmonix/Source/HarmonixDsp/Private/HarmonixDsp/Stretchers/SmbPitchShifterConfig.h
--- a/Engine/Plugins/Runtime/Harmonix/Source/HarmonixDsp/Private/HarmonixDsp/Stretchers/SmbPitchShifterConfig.h
+++ b/Engine/Plugins/Runtime/Harmonix/Source/HarmonixDsp/Private/HarmonixDsp/Stretchers/SmbPitchShifterConfig.h
@@ -14,6 +14,10 @@ class USmbPitchShifterConfig : public UStretcherAndPitchShifterConfig
 
 public:
 
+	/** If true, allows the stretcher pool to grow beyond DefaultAllocatedStretchers in the editor. */
+	UPROPERTY(config, EditAnywhere, Category = "Pool Config")
+	bool bAllowPoolGrowthInEditor = false;
+
 	UPROPERTY(config, EditAnywhere, Category = "Pool Config")
 	uint32 DefaultAllocatedStretchers = 12;
 
@@ -51,6 +55,15 @@ public:
 		return AndroidAllocatedStretchers;
 #else
 		return DefaultAllocatedStretchers;
+#endif
+	}
+
+	bool ShouldAllowPoolGrowth() const
+	{
+#if WITH_EDITOR
+		return bAllowPoolGrowthInEditor;
+#else
+		return false;
 #endif
 	}
 };
diff --git a/Engine/Plugins/Runtime/Harmonix/Source/HarmonixDsp/Private/HarmonixDsp/Stretchers/SmbPitchShifterFactory.cpp b/Engine/Plugins/Runtime/Harmonix/Source/HarmonixDsp/Private/HarmonixDsp/Stretchers/SmbPitchShifterFactory.cpp
--- a/Engine/Plugins/Runtime/Harmonix/Source/HarmonixDsp/Private/HarmonixDsp/Stretchers/SmbPitchShifterFactory.cpp
+++ b/Engine/Plugins/Runtime/Harmonix/Source/HarmonixDsp/Private/HarmonixDsp/Stretchers/SmbPitchShifterFactory.cpp
@@ -40,13 +40,13 @@ void FSmbPitchShifterFactory::OnClientAdded(void* Client, float InSampleRate)
 
 	if (!PitchShifterPool.IsInitialized() && Config->GetNumAllocatedStretchersForPlatform() > 0)
 	{
-		auto FactoryMethod = [this]() -> TSharedPtr<IStretcherAndPitchShifter, ESPMode::ThreadSafe>
+		auto FactoryMethod = [SampleRate = FactorySampleRate]() -> TSharedPtr<IStretcherAndPitchShifter, ESPMode::ThreadSafe>
 		{
-			return CreateStretcherAndPitchShifter(FactorySampleRate);
+			return CreateStretcherAndPitchShifter(SampleRate);
 		};
 
 
-		PitchShifterPool.Allocate(Config->GetNumAllocatedStretchersForPlatform(), FactoryMethod);
+		PitchShifterPool.Allocate(Config->GetNumAllocatedStretchersForPlatform(), FactoryMethod, Config->ShouldAllowPoolGrowth());
 	}
 }
 
@@ -70,7 +70,7 @@ void FSmbPitchShifterFactory::ReleasePitchShifter(TSharedPtr<IStretcherAndPitchS
 
 bool FSmbPitchShifterFactory::HasFreePitchShifters(const FTimeStretchConfig& InConfig) const
 {
-	return PitchShifterPool.GetNumFree() > 0;
+	return (PitchShifterPool.GetNumFree() > 0) || PitchShifterPool.CanGrow();
 }
 
 size_t FSmbPitchShifterFactory::GetMemoryUsedByPool() const
diff --git a/Engine/Plugins/Runtime/Harmonix/Source/HarmonixDsp/Private/HarmonixDsp/Stretchers/SmbPitchShifterFactory.h b/Engine/Plugins/Runtime/Harmonix/Source/HarmonixDsp/Private/HarmonixDsp/Stretchers/SmbPitchShifterFactory.h
--- a/Engine/Plugins/Runtime/Harmonix/Source/HarmonixDsp/Private/HarmonixDsp/Stretchers/SmbPitchShifterFactory.h
+++ b/Engine/Plugins/Runtime/Harmonix/Source/HarmonixDsp/Private/HarmonixDsp/Stretchers/SmbPitchShifterFactory.h
@@ -32,5 +32,5 @@ private:
 
 	float FactorySampleRate = 0;
 
-	TSharedPtr<IStretcherAndPitchShifter, ESPMode::ThreadSafe> CreateStretcherAndPitchShifter(float InSampleRate);
+	static TSharedPtr<IStretcherAndPitchShifter, ESPMode::ThreadSafe> CreateStretcherAndPitchShifter(float InSampleRate);
 };
\ No newline at end of file
diff --git a/Engine/Plugins/Runtime/Harmonix/Source/HarmonixDsp/Public/HarmonixDsp/ObjectPool.h b/Engine/Plugins/Runtime/Harmonix/Source/HarmonixDsp/Public/HarmonixDsp/ObjectPool.h
--- a/Engine/Plugins/Runtime/Harmonix/Source/HarmonixDsp/Public/HarmonixDsp/ObjectPool.h
+++ b/Engine/Plugins/Runtime/Harmonix/Source/HarmonixDsp/Public/HarmonixDsp/ObjectPool.h
@@ -34,7 +34,7 @@ namespace Harmonix
 		 * 
 		 * @return	number of instances successfully created and pooled. Can be fewer than Num requested.
 		 */
-		FORCEINLINE int32 TryAllocate(int32 Num, TFunction<FSharedObjectPtr()> FactoryMethod)
+		FORCEINLINE int32 TryAllocate(int32 Num, TFunction<FSharedObjectPtr()> InFactoryMethod, bool bInAllowGrowth = false)
 		{
 			checkf(Num > 0, TEXT("Num to allocate must be greater than 0"));
 
@@ -42,22 +42,28 @@ namespace Harmonix
 
 			checkf(Objects.Num() == 0, TEXT("Objects already allocated. Check \"IsInitialized() == false\" before allocating."));
 
+			bAllowGrowth = bInAllowGrowth;
 			Objects.Reset(Num);
 			for (int32 Index = 0; Index < Num; ++Index)
 			{
-				if (FSharedObjectPtr NewObject = FactoryMethod())
+				if (FSharedObjectPtr NewObject = InFactoryMethod())
 				{
 					Objects.Add(NewObject);
 				}
 			}
 
+			if (bAllowGrowth)
+			{
+				FactoryMethod = InFactoryMethod;
+			}
+
 			return Objects.Num();
 		}
 
 		/**
 		 * allocates Num number of objects, given a factory method
 		 */
-		void Allocate(int32 Num, TFunction<FSharedObjectPtr()> FactoryMethod)
+		void Allocate(int32 Num, TFunction<FSharedObjectPtr()> InFactoryMethod, bool bInAllowGrowth = false)
 		{
 			checkf(Num > 0, TEXT("Num to allocate must be greater than 0"));
 
@@ -65,12 +71,18 @@ namespace Harmonix
 
 			checkf(Objects.Num() == 0, TEXT("Objects already allocated. Check \"IsInitialized() == false\" before allocating."));
 
+			bAllowGrowth = bInAllowGrowth;
 			Objects.Reset(Num);
 			for (int32 Index = 0; Index < Num; ++Index)
 			{
-				Objects.Add(FactoryMethod());
+				Objects.Add(InFactoryMethod());
 				checkSlow(Objects[Index].IsValid());
 			}
+
+			if (bAllowGrowth)
+			{
+				FactoryMethod = InFactoryMethod;
+			}
 		}
 
 		/**
@@ -83,7 +95,9 @@ namespace Harmonix
 			FScopeLock Lock(&PoolLock);
 
 			NumInUse = 0;
+			bAllowGrowth = false;
 			Objects.Reset();
+			FactoryMethod.Reset();
 		}
 
 		/**
@@ -104,6 +118,15 @@ namespace Harmonix
 				++NumInUse;
 				return Objects[FreeIndex];
 			}
+			else if (bAllowGrowth)
+			{
+				if (FSharedObjectPtr NewObject = FactoryMethod())
+				{
+					Objects.Add(NewObject);
+					++NumInUse;
+					return NewObject;
+				}
+			}
 
 			return nullptr;
 
@@ -151,6 +174,10 @@ namespace Harmonix
 			return Objects.Num() > 0;
 		}
 
+		bool CanGrow() const
+		{
+			return bAllowGrowth;
+		}
 
 		int32 GetNumAllocated() const
 		{
@@ -185,7 +212,9 @@ namespace Harmonix
 
 		mutable FCriticalSection PoolLock;
 
-		int32 NumInUse;
+		int32 NumInUse = 0;
+		bool bAllowGrowth = false;
 		TArray<FSharedObjectPtr> Objects;
+		TFunction<FSharedObjectPtr()> FactoryMethod;
 	};
 };
\ No newline at end of file
