[Harmonix] Remove FMidiStream constructor which takes FOperatorSettings

#tests automated tests

diff --git a/Engine/Plugins/Runtime/Harmonix/Source/HarmonixMetasound/Private/HarmonixMetasound/DataTypes/MidiStream.cpp b/Engine/Plugins/Runtime/Harmonix/Source/HarmonixMetasound/Private/HarmonixMetasound/DataTypes/MidiStream.cpp
--- a/Engine/Plugins/Runtime/Harmonix/Source/HarmonixMetasound/Private/HarmonixMetasound/DataTypes/MidiStream.cpp
+++ b/Engine/Plugins/Runtime/Harmonix/Source/HarmonixMetasound/Private/HarmonixMetasound/DataTypes/MidiStream.cpp
@@ -14,10 +14,6 @@ namespace HarmonixMetasound
 	DEFINE_LOG_CATEGORY(LogMidiStreamDataType);
 
 	using namespace Metasound;
-	
-	FMidiStream::FMidiStream(const FOperatorSettings&)
-	{
-	}
 
 	void FMidiStream::SetClock(const FMidiClock& InClock)
 	{
diff --git a/Engine/Plugins/Runtime/Harmonix/Source/HarmonixMetasound/Private/HarmonixMetasound/Nodes/FusionSamplerNode.cpp b/Engine/Plugins/Runtime/Harmonix/Source/HarmonixMetasound/Private/HarmonixMetasound/Nodes/FusionSamplerNode.cpp
--- a/Engine/Plugins/Runtime/Harmonix/Source/HarmonixMetasound/Private/HarmonixMetasound/Nodes/FusionSamplerNode.cpp
+++ b/Engine/Plugins/Runtime/Harmonix/Source/HarmonixMetasound/Private/HarmonixMetasound/Nodes/FusionSamplerNode.cpp
@@ -74,7 +74,7 @@ namespace HarmonixMetasound
 				: Params(InParams)
 				, InSettings(& InParams.OperatorSettings)
 				, InEnabled(InParams.InputData.GetOrCreateDefaultDataReadReference<bool>(METASOUND_GET_PARAM_NAME(CommonPinNames::Inputs::Enable), InParams.OperatorSettings))
-				, InMidiStream(InParams.InputData.GetOrConstructDataReadReference<FMidiStream>(METASOUND_GET_PARAM_NAME(CommonPinNames::Inputs::MidiStream), InParams.OperatorSettings))
+				, InMidiStream(InParams.InputData.GetOrConstructDataReadReference<FMidiStream>(METASOUND_GET_PARAM_NAME(CommonPinNames::Inputs::MidiStream)))
 				, InTrackNumber(InParams.InputData.GetOrCreateDefaultDataReadReference<int32>(METASOUND_GET_PARAM_NAME(CommonPinNames::Inputs::MidiTrackNumber), InParams.OperatorSettings))
 				, InTransposition(InParams.InputData.GetOrCreateDefaultDataReadReference<int32>(METASOUND_GET_PARAM_NAME(CommonPinNames::Inputs::Transposition), InParams.OperatorSettings))
 				, InLfo0RateOverride(InParams.InputData.GetOrCreateDefaultDataReadReference<float>(METASOUND_GET_PARAM_NAME(FusionSamplerNodePinNames::Lfo0FrequencyOverride), InParams.OperatorSettings))
diff --git a/Engine/Plugins/Runtime/Harmonix/Source/HarmonixMetasound/Private/HarmonixMetasound/Nodes/MidiChannelFilterNode.cpp b/Engine/Plugins/Runtime/Harmonix/Source/HarmonixMetasound/Private/HarmonixMetasound/Nodes/MidiChannelFilterNode.cpp
--- a/Engine/Plugins/Runtime/Harmonix/Source/HarmonixMetasound/Private/HarmonixMetasound/Nodes/MidiChannelFilterNode.cpp
+++ b/Engine/Plugins/Runtime/Harmonix/Source/HarmonixMetasound/Private/HarmonixMetasound/Nodes/MidiChannelFilterNode.cpp
@@ -105,13 +105,13 @@ namespace HarmonixMetasound::Nodes::MidiChannelFilter
 			FInputs Inputs
 			{
 				InputData.GetOrCreateDefaultDataReadReference<bool>(Inputs::EnableName, OperatorSettings),
-				InputData.GetOrConstructDataReadReference<FMidiStream>(Inputs::MidiStreamName, OperatorSettings),
+				InputData.GetOrConstructDataReadReference<FMidiStream>(Inputs::MidiStreamName),
 				InputData.GetOrCreateDefaultDataReadReference<int32>(Inputs::ChannelName, OperatorSettings)
 			};
 
 			FOutputs Outputs
 			{
-				FMidiStreamWriteRef::CreateNew(OperatorSettings)
+				FMidiStreamWriteRef::CreateNew()
 			};
 
 			return MakeUnique<FMidiChannelFilterOperator>(InParams, MoveTemp(Inputs), MoveTemp(Outputs));
diff --git a/Engine/Plugins/Runtime/Harmonix/Source/HarmonixMetasound/Private/HarmonixMetasound/Nodes/MidiNoteFilterNode.cpp b/Engine/Plugins/Runtime/Harmonix/Source/HarmonixMetasound/Private/HarmonixMetasound/Nodes/MidiNoteFilterNode.cpp
--- a/Engine/Plugins/Runtime/Harmonix/Source/HarmonixMetasound/Private/HarmonixMetasound/Nodes/MidiNoteFilterNode.cpp
+++ b/Engine/Plugins/Runtime/Harmonix/Source/HarmonixMetasound/Private/HarmonixMetasound/Nodes/MidiNoteFilterNode.cpp
@@ -119,7 +119,7 @@ namespace HarmonixMetasound::Nodes::MidiNoteFilter
 			FInputs Inputs
 			{
 				InputData.GetOrCreateDefaultDataReadReference<bool>(Inputs::EnableName, OperatorSettings),
-				InputData.GetOrConstructDataReadReference<FMidiStream>(Inputs::MidiStreamName, OperatorSettings),
+				InputData.GetOrConstructDataReadReference<FMidiStream>(Inputs::MidiStreamName),
 				InputData.GetOrCreateDefaultDataReadReference<int32>(Inputs::MinNoteNumberName, OperatorSettings),
 				InputData.GetOrCreateDefaultDataReadReference<int32>(Inputs::MaxNoteNumberName, OperatorSettings),
 				InputData.GetOrCreateDefaultDataReadReference<int32>(Inputs::MinVelocityName, OperatorSettings),
@@ -129,7 +129,7 @@ namespace HarmonixMetasound::Nodes::MidiNoteFilter
 
 			FOutputs Outputs
 			{
-				FMidiStreamWriteRef::CreateNew(OperatorSettings)
+				FMidiStreamWriteRef::CreateNew()
 			};
 
 			return MakeUnique<FMidiNoteFilterOperator>(InParams, MoveTemp(Inputs), MoveTemp(Outputs));
diff --git a/Engine/Plugins/Runtime/Harmonix/Source/HarmonixMetasound/Private/HarmonixMetasound/Nodes/MidiNoteTriggerNode.cpp b/Engine/Plugins/Runtime/Harmonix/Source/HarmonixMetasound/Private/HarmonixMetasound/Nodes/MidiNoteTriggerNode.cpp
--- a/Engine/Plugins/Runtime/Harmonix/Source/HarmonixMetasound/Private/HarmonixMetasound/Nodes/MidiNoteTriggerNode.cpp
+++ b/Engine/Plugins/Runtime/Harmonix/Source/HarmonixMetasound/Private/HarmonixMetasound/Nodes/MidiNoteTriggerNode.cpp
@@ -95,7 +95,7 @@ namespace HarmonixMetasound::Nodes::MidiNoteTriggerNode
 
 			const FInputVertexInterfaceData& InputData = InParams.InputData;
 			FBoolReadRef InEnabled          = InputData.GetOrCreateDefaultDataReadReference<bool>(METASOUND_GET_PARAM_NAME(Inputs::Enable), InParams.OperatorSettings);
-			FMidiStreamReadRef InMidiStream = InputData.GetOrConstructDataReadReference<FMidiStream>(METASOUND_GET_PARAM_NAME(Inputs::MidiStream), InParams.OperatorSettings);
+			FMidiStreamReadRef InMidiStream = InputData.GetOrConstructDataReadReference<FMidiStream>(METASOUND_GET_PARAM_NAME(Inputs::MidiStream));
 
 			return MakeUnique<FMidiNoteTriggerOperator_V1>(InParams, MoveTemp(InEnabled), MoveTemp(InMidiStream));
 		}
diff --git a/Engine/Plugins/Runtime/Harmonix/Source/HarmonixMetasound/Private/HarmonixMetasound/Nodes/MidiPlayerNode.cpp b/Engine/Plugins/Runtime/Harmonix/Source/HarmonixMetasound/Private/HarmonixMetasound/Nodes/MidiPlayerNode.cpp
--- a/Engine/Plugins/Runtime/Harmonix/Source/HarmonixMetasound/Private/HarmonixMetasound/Nodes/MidiPlayerNode.cpp
+++ b/Engine/Plugins/Runtime/Harmonix/Source/HarmonixMetasound/Private/HarmonixMetasound/Nodes/MidiPlayerNode.cpp
@@ -255,7 +255,7 @@ namespace HarmonixMetasound
 		, PrerollBars(InPrerollBars)
 		, bKillVoicesOnSeek(bInKillVoicesOnSeek)
 		, bKillVoicesOnMidiChange(bInKillVoicesOnMidiChange)
-		, MidiOutPin(FMidiStreamWriteRef::CreateNew(InSettings))
+		, MidiOutPin(FMidiStreamWriteRef::CreateNew())
 		, MidiClockOut(FMidiClockWriteRef::CreateNew(InSettings))
 		, BlockSize(InSettings.GetNumFramesPerBlock())
 	{
diff --git a/Engine/Plugins/Runtime/Harmonix/Source/HarmonixMetasound/Private/HarmonixMetasound/Nodes/MidiPulseGeneratorNode.cpp b/Engine/Plugins/Runtime/Harmonix/Source/HarmonixMetasound/Private/HarmonixMetasound/Nodes/MidiPulseGeneratorNode.cpp
--- a/Engine/Plugins/Runtime/Harmonix/Source/HarmonixMetasound/Private/HarmonixMetasound/Nodes/MidiPulseGeneratorNode.cpp
+++ b/Engine/Plugins/Runtime/Harmonix/Source/HarmonixMetasound/Private/HarmonixMetasound/Nodes/MidiPulseGeneratorNode.cpp
@@ -144,7 +144,7 @@ namespace HarmonixMetasound::Nodes::MidiPulseGeneratorNode
 
 			FOutputs Outputs
 			{
-				FMidiStreamWriteRef::CreateNew(InParams.OperatorSettings)
+				FMidiStreamWriteRef::CreateNew()
 			};
 
 			return MakeUnique<FMidiPulseGeneratorOperator>(InParams, MoveTemp(Inputs), MoveTemp(Outputs));
diff --git a/Engine/Plugins/Runtime/Harmonix/Source/HarmonixMetasound/Private/HarmonixMetasound/Nodes/MidiStreamLoggerNode.cpp b/Engine/Plugins/Runtime/Harmonix/Source/HarmonixMetasound/Private/HarmonixMetasound/Nodes/MidiStreamLoggerNode.cpp
--- a/Engine/Plugins/Runtime/Harmonix/Source/HarmonixMetasound/Private/HarmonixMetasound/Nodes/MidiStreamLoggerNode.cpp
+++ b/Engine/Plugins/Runtime/Harmonix/Source/HarmonixMetasound/Private/HarmonixMetasound/Nodes/MidiStreamLoggerNode.cpp
@@ -105,7 +105,7 @@ namespace HarmonixMetasound
 
 		const FInputVertexInterfaceData& InputData = InParams.InputData;
 		FBoolReadRef InEnabled = InputData.GetOrCreateDefaultDataReadReference<bool>(METASOUND_GET_PARAM_NAME(Inputs::Enable), InParams.OperatorSettings);
-		FMidiStreamReadRef InMidiStream = InputData.GetOrConstructDataReadReference<FMidiStream>(METASOUND_GET_PARAM_NAME(Inputs::MidiStream), InParams.OperatorSettings);
+		FMidiStreamReadRef InMidiStream = InputData.GetOrConstructDataReadReference<FMidiStream>(METASOUND_GET_PARAM_NAME(Inputs::MidiStream));
 
 		return MakeUnique<FMidiStreamLoggerOperator>(InParams, InEnabled, InMidiStream);
 	}
diff --git a/Engine/Plugins/Runtime/Harmonix/Source/HarmonixMetasound/Private/HarmonixMetasound/Nodes/MidiStreamMergeNode.cpp b/Engine/Plugins/Runtime/Harmonix/Source/HarmonixMetasound/Private/HarmonixMetasound/Nodes/MidiStreamMergeNode.cpp
--- a/Engine/Plugins/Runtime/Harmonix/Source/HarmonixMetasound/Private/HarmonixMetasound/Nodes/MidiStreamMergeNode.cpp
+++ b/Engine/Plugins/Runtime/Harmonix/Source/HarmonixMetasound/Private/HarmonixMetasound/Nodes/MidiStreamMergeNode.cpp
@@ -109,8 +109,8 @@ namespace HarmonixMetasound
 		const FMidiStreamMergeNode& LoggerNode = static_cast<const FMidiStreamMergeNode&>(InParams.Node);
 
 		const FInputVertexInterfaceData& InputData = InParams.InputData;
-		FMidiStreamReadRef InMidiStreamA = InputData.GetOrConstructDataReadReference<FMidiStream>(METASOUND_GET_PARAM_NAME(InputMidiStreamA), InParams.OperatorSettings);
-		FMidiStreamReadRef InMidiStreamB = InputData.GetOrConstructDataReadReference<FMidiStream>(METASOUND_GET_PARAM_NAME(InputMidiStreamB), InParams.OperatorSettings);
+		FMidiStreamReadRef InMidiStreamA = InputData.GetOrConstructDataReadReference<FMidiStream>(METASOUND_GET_PARAM_NAME(InputMidiStreamA));
+		FMidiStreamReadRef InMidiStreamB = InputData.GetOrConstructDataReadReference<FMidiStream>(METASOUND_GET_PARAM_NAME(InputMidiStreamB));
 
 		return MakeUnique<FMidiStreamMergeOperator>(InParams, InMidiStreamA, InMidiStreamB);
 	}
@@ -120,7 +120,7 @@ namespace HarmonixMetasound
 											   const FMidiStreamReadRef& InMidiStreamB)
 		: MidiStreamAInPin(InMidiStreamA)
 		, MidiStreamBInPin(InMidiStreamB)
-		, MidiStreamOutPin(FMidiStreamWriteRef::CreateNew(InParams.OperatorSettings))
+		, MidiStreamOutPin(FMidiStreamWriteRef::CreateNew())
 	{
 		Reset(InParams);
 	}
diff --git a/Engine/Plugins/Runtime/Harmonix/Source/HarmonixMetasound/Private/HarmonixMetasound/Nodes/MidiStreamSelectNode.cpp b/Engine/Plugins/Runtime/Harmonix/Source/HarmonixMetasound/Private/HarmonixMetasound/Nodes/MidiStreamSelectNode.cpp
--- a/Engine/Plugins/Runtime/Harmonix/Source/HarmonixMetasound/Private/HarmonixMetasound/Nodes/MidiStreamSelectNode.cpp
+++ b/Engine/Plugins/Runtime/Harmonix/Source/HarmonixMetasound/Private/HarmonixMetasound/Nodes/MidiStreamSelectNode.cpp
@@ -132,8 +132,8 @@ namespace HarmonixMetasound
 		const FMidiStreamSelectNode& LoggerNode = static_cast<const FMidiStreamSelectNode&>(InParams.Node);
 
 		const FInputVertexInterfaceData& InputData = InParams.InputData;
-		FMidiStreamReadRef InMidiStreamA = InputData.GetOrConstructDataReadReference<FMidiStream>(METASOUND_GET_PARAM_NAME(InputMidiStreamA), InParams.OperatorSettings);
-		FMidiStreamReadRef InMidiStreamB = InputData.GetOrConstructDataReadReference<FMidiStream>(METASOUND_GET_PARAM_NAME(InputMidiStreamB), InParams.OperatorSettings);
+		FMidiStreamReadRef InMidiStreamA = InputData.GetOrConstructDataReadReference<FMidiStream>(METASOUND_GET_PARAM_NAME(InputMidiStreamA));
+		FMidiStreamReadRef InMidiStreamB = InputData.GetOrConstructDataReadReference<FMidiStream>(METASOUND_GET_PARAM_NAME(InputMidiStreamB));
 		FInt32ReadRef InStreamIndex = InputData.GetOrConstructDataReadReference<int32>(METASOUND_GET_PARAM_NAME(InputStreamIndex));
 		FBoolReadRef InImmediateNoteOff = InputData.GetOrConstructDataReadReference<bool>(METASOUND_GET_PARAM_NAME(InputImmediateNoteOff));
 		FBoolReadRef InCopyInactiveNoteOffs = InputData.GetOrConstructDataReadReference<bool>(METASOUND_GET_PARAM_NAME(InputCopyInactiveNoteOffs));
@@ -152,7 +152,7 @@ namespace HarmonixMetasound
 		, StreamIndexInPin(InStreamIndex)
 		, ImmediateNoteOffInPin(InImmediateNoteOff)
 		, CopyInactiveNoteOffsInPin(InCopyInactiveNoteOffs)
-		, MidiStreamOutPin(FMidiStreamWriteRef::CreateNew(InParams.OperatorSettings))
+		, MidiStreamOutPin(FMidiStreamWriteRef::CreateNew())
 	{
 		Reset(InParams);
 	}
diff --git a/Engine/Plugins/Runtime/Harmonix/Source/HarmonixMetasound/Private/HarmonixMetasound/Nodes/MidiStreamTrackFilterNode.cpp b/Engine/Plugins/Runtime/Harmonix/Source/HarmonixMetasound/Private/HarmonixMetasound/Nodes/MidiStreamTrackFilterNode.cpp
--- a/Engine/Plugins/Runtime/Harmonix/Source/HarmonixMetasound/Private/HarmonixMetasound/Nodes/MidiStreamTrackFilterNode.cpp
+++ b/Engine/Plugins/Runtime/Harmonix/Source/HarmonixMetasound/Private/HarmonixMetasound/Nodes/MidiStreamTrackFilterNode.cpp
@@ -114,7 +114,7 @@ namespace HarmonixMetasound::Nodes::MidiTrackFilter
 			FInputs Inputs
 			{
 				InputData.GetOrCreateDefaultDataReadReference<bool>(Inputs::EnableName, InParams.OperatorSettings),
-				InputData.GetOrConstructDataReadReference<FMidiStream>(Inputs::MidiStreamName, InParams.OperatorSettings),
+				InputData.GetOrConstructDataReadReference<FMidiStream>(Inputs::MidiStreamName),
 				InputData.GetOrCreateDefaultDataReadReference<int32>(Inputs::MinTrackIndexName, InParams.OperatorSettings),
 				InputData.GetOrCreateDefaultDataReadReference<int32>(Inputs::MaxTrackIndexName, InParams.OperatorSettings),
 				InputData.GetOrCreateDefaultDataReadReference<bool>(Inputs::IncludeConductorTrackName, InParams.OperatorSettings)
@@ -122,7 +122,7 @@ namespace HarmonixMetasound::Nodes::MidiTrackFilter
 
 			FOutputs Outputs
 			{
-				FMidiStreamWriteRef::CreateNew(InParams.OperatorSettings)
+				FMidiStreamWriteRef::CreateNew()
 			};
 			
 			return MakeUnique<FMidiStreamTrackFilterOperator_V1>(InParams, MoveTemp(Inputs), MoveTemp(Outputs));
diff --git a/Engine/Plugins/Runtime/Harmonix/Source/HarmonixMetasound/Private/HarmonixMetasound/Nodes/MidiStreamTransposerNode.cpp b/Engine/Plugins/Runtime/Harmonix/Source/HarmonixMetasound/Private/HarmonixMetasound/Nodes/MidiStreamTransposerNode.cpp
--- a/Engine/Plugins/Runtime/Harmonix/Source/HarmonixMetasound/Private/HarmonixMetasound/Nodes/MidiStreamTransposerNode.cpp
+++ b/Engine/Plugins/Runtime/Harmonix/Source/HarmonixMetasound/Private/HarmonixMetasound/Nodes/MidiStreamTransposerNode.cpp
@@ -103,13 +103,13 @@ namespace HarmonixMetasound::Nodes::MidiNoteTranspose
 			FInputs Inputs
 			{
 				InputData.GetOrCreateDefaultDataReadReference<bool>(METASOUND_GET_PARAM_NAME(Inputs::Enable), InParams.OperatorSettings),
-				InputData.GetOrConstructDataReadReference<FMidiStream>(METASOUND_GET_PARAM_NAME(Inputs::MidiStream), InParams.OperatorSettings),
+				InputData.GetOrConstructDataReadReference<FMidiStream>(METASOUND_GET_PARAM_NAME(Inputs::MidiStream)),
 				InputData.GetOrCreateDefaultDataReadReference<int32>(METASOUND_GET_PARAM_NAME(Inputs::Transposition), InParams.OperatorSettings)
 			};
 
 			FOutputs Outputs
 			{
-				FMidiStreamWriteRef::CreateNew(InParams.OperatorSettings)
+				FMidiStreamWriteRef::CreateNew()
 			};
 
 			return MakeUnique<FMidiNoteTransposeOperator>(InParams, MoveTemp(Inputs), MoveTemp(Outputs));
diff --git a/Engine/Plugins/Runtime/Harmonix/Source/HarmonixMetasound/Private/HarmonixMetasound/Nodes/MidiTextTriggerNode.cpp b/Engine/Plugins/Runtime/Harmonix/Source/HarmonixMetasound/Private/HarmonixMetasound/Nodes/MidiTextTriggerNode.cpp
--- a/Engine/Plugins/Runtime/Harmonix/Source/HarmonixMetasound/Private/HarmonixMetasound/Nodes/MidiTextTriggerNode.cpp
+++ b/Engine/Plugins/Runtime/Harmonix/Source/HarmonixMetasound/Private/HarmonixMetasound/Nodes/MidiTextTriggerNode.cpp
@@ -76,7 +76,7 @@ namespace HarmonixMetasound
 
 			const FInputVertexInterfaceData& InputData = InParams.InputData;
 			FBoolReadRef InEnabled          = InputData.GetOrCreateDefaultDataReadReference<bool>(METASOUND_GET_PARAM_NAME(Inputs::Enable), InParams.OperatorSettings);
-			FMidiStreamReadRef InMidiStream = InputData.GetOrConstructDataReadReference<FMidiStream>(METASOUND_GET_PARAM_NAME(Inputs::MidiStream), InParams.OperatorSettings);
+			FMidiStreamReadRef InMidiStream = InputData.GetOrConstructDataReadReference<FMidiStream>(METASOUND_GET_PARAM_NAME(Inputs::MidiStream));
 			FStringReadRef InText           = InputData.GetOrCreateDefaultDataReadReference<FString>(METASOUND_GET_PARAM_NAME(TextInput), InParams.OperatorSettings);
 
 			return MakeUnique<FMidiTextTriggerOperator_V1>(InParams, InEnabled, InMidiStream, InText);
diff --git a/Engine/Plugins/Runtime/Harmonix/Source/HarmonixMetasound/Private/HarmonixMetasound/Nodes/StepSequencePlayerNode.cpp b/Engine/Plugins/Runtime/Harmonix/Source/HarmonixMetasound/Private/HarmonixMetasound/Nodes/StepSequencePlayerNode.cpp
--- a/Engine/Plugins/Runtime/Harmonix/Source/HarmonixMetasound/Private/HarmonixMetasound/Nodes/StepSequencePlayerNode.cpp
+++ b/Engine/Plugins/Runtime/Harmonix/Source/HarmonixMetasound/Private/HarmonixMetasound/Nodes/StepSequencePlayerNode.cpp
@@ -246,7 +246,7 @@ namespace HarmonixMetasound
 		, AutoPagePlaysBlankPagesInPin(InAutoPagePlaysBlankPages)
 		, LoopInPin(InLoop)
 		, EnabledInPin(InEnabled)
-		, MidiOutPin(FMidiStreamWriteRef::CreateNew(InParams.OperatorSettings))
+		, MidiOutPin(FMidiStreamWriteRef::CreateNew())
 	{
 		Reset(InParams);
 		Init();
diff --git a/Engine/Plugins/Runtime/Harmonix/Source/HarmonixMetasound/Public/HarmonixMetasound/DataTypes/MidiStream.h b/Engine/Plugins/Runtime/Harmonix/Source/HarmonixMetasound/Public/HarmonixMetasound/DataTypes/MidiStream.h
--- a/Engine/Plugins/Runtime/Harmonix/Source/HarmonixMetasound/Public/HarmonixMetasound/DataTypes/MidiStream.h
+++ b/Engine/Plugins/Runtime/Harmonix/Source/HarmonixMetasound/Public/HarmonixMetasound/DataTypes/MidiStream.h
@@ -57,8 +57,6 @@ namespace HarmonixMetasound
 	class HARMONIXMETASOUND_API FMidiStream
 	{
 	public:
-		explicit FMidiStream(const Metasound::FOperatorSettings& InSettings);
-
 		void SetClock(const FMidiClock& InClock);
 		void ResetClock();
 		TSharedPtr<const FMidiClock, ESPMode::NotThreadSafe> GetClock() const;
diff --git a/Engine/Plugins/Runtime/Harmonix/Source/HarmonixMetasoundTests/Private/DataTypes/MidiStreamTests.cpp b/Engine/Plugins/Runtime/Harmonix/Source/HarmonixMetasoundTests/Private/DataTypes/MidiStreamTests.cpp
--- a/Engine/Plugins/Runtime/Harmonix/Source/HarmonixMetasoundTests/Private/DataTypes/MidiStreamTests.cpp
+++ b/Engine/Plugins/Runtime/Harmonix/Source/HarmonixMetasoundTests/Private/DataTypes/MidiStreamTests.cpp
@@ -9,7 +9,7 @@ namespace HarmonixMetasound::MidiStream::Tests
 {
 	FMidiStream MakeTestMidiStream(const Metasound::FOperatorSettings& OperatorSettings)
 	{
-		FMidiStream Stream{ OperatorSettings };
+		FMidiStream Stream;
 
 		// Make a bunch of note events
 		constexpr int32 NumNotes = 3;
@@ -60,7 +60,7 @@ namespace HarmonixMetasound::MidiStream::Tests
 	{
 		Metasound::FOperatorSettings OperatorSettings{ 48000, 100 };
 		const FMidiStream FromStream = MakeTestMidiStream(OperatorSettings);
-		FMidiStream ToStream{ OperatorSettings };
+		FMidiStream ToStream;
 
 		// Copy with no filter or transform
 		{
@@ -191,10 +191,10 @@ namespace HarmonixMetasound::MidiStream::Tests
 	bool FMidiStreamMergeTest::RunTest(const FString&)
 	{
 		Metasound::FOperatorSettings OperatorSettings{ 48000, 100 };
-		FMidiStream FromA{OperatorSettings};
-		FMidiStream FromB{OperatorSettings};
-		FMidiStream Expected{OperatorSettings};
-		FMidiStream To{OperatorSettings};
+		FMidiStream FromA;
+		FMidiStream FromB;
+		FMidiStream Expected;
+		FMidiStream To;
 
 		// no filter or transform
 		{
@@ -272,9 +272,9 @@ namespace HarmonixMetasound::MidiStream::Tests
 	{
 		// Make sure merging streams results in the correct clock being set on the output stream
 		Metasound::FOperatorSettings OperatorSettings{ 48000, 100 };
-		FMidiStream FromA{OperatorSettings};
-		FMidiStream FromB{OperatorSettings};
-		FMidiStream To{OperatorSettings};
+		FMidiStream FromA;
+		FMidiStream FromB;
+		FMidiStream To;
 
 		TSharedRef<const FMidiClock, ESPMode::NotThreadSafe> ClockA = MakeShared<FMidiClock, ESPMode::NotThreadSafe>(OperatorSettings);
 		FromA.SetClock(*ClockA);
diff --git a/Engine/Plugins/Runtime/Harmonix/Source/HarmonixMetasoundTests/Private/MidiOps/MidiChannelFilterTests.cpp b/Engine/Plugins/Runtime/Harmonix/Source/HarmonixMetasoundTests/Private/MidiOps/MidiChannelFilterTests.cpp
--- a/Engine/Plugins/Runtime/Harmonix/Source/HarmonixMetasoundTests/Private/MidiOps/MidiChannelFilterTests.cpp
+++ b/Engine/Plugins/Runtime/Harmonix/Source/HarmonixMetasoundTests/Private/MidiOps/MidiChannelFilterTests.cpp
@@ -46,9 +46,8 @@ namespace Harmonix::Midi::Ops::Tests
 	{
 		FMidiChannelFilter Filter;
 		
-		Metasound::FOperatorSettings OperatorSettings{ 48000, 100 };
-		HarmonixMetasound::FMidiStream InputStream{ OperatorSettings };
-		HarmonixMetasound::FMidiStream OutputStream{ OperatorSettings };
+		HarmonixMetasound::FMidiStream InputStream;
+		HarmonixMetasound::FMidiStream OutputStream;
 
 		// Add some events to the input stream
 		constexpr uint8 NoteChannelIdx = 1;
diff --git a/Engine/Plugins/Runtime/Harmonix/Source/HarmonixMetasoundTests/Private/MidiOps/MidiNoteFilterTests.cpp b/Engine/Plugins/Runtime/Harmonix/Source/HarmonixMetasoundTests/Private/MidiOps/MidiNoteFilterTests.cpp
--- a/Engine/Plugins/Runtime/Harmonix/Source/HarmonixMetasoundTests/Private/MidiOps/MidiNoteFilterTests.cpp
+++ b/Engine/Plugins/Runtime/Harmonix/Source/HarmonixMetasoundTests/Private/MidiOps/MidiNoteFilterTests.cpp
@@ -15,9 +15,8 @@ namespace Harmonix::Midi::Ops::Tests
 	{
 		FMidiNoteFilter Filter;
 		
-		Metasound::FOperatorSettings OperatorSettings{ 48000, 100 };
-		HarmonixMetasound::FMidiStream InputStream{ OperatorSettings };
-		HarmonixMetasound::FMidiStream OutputStream{ OperatorSettings };
+		HarmonixMetasound::FMidiStream InputStream;
+		HarmonixMetasound::FMidiStream OutputStream;
 
 		FMidiMsg NoteOnMsg = FMidiMsg::CreateNoteOn(2, 3, 4);
 		HarmonixMetasound::FMidiStreamEvent NoteOnEvent{ static_cast<uint32>(0), NoteOnMsg };
diff --git a/Engine/Plugins/Runtime/Harmonix/Source/HarmonixMetasoundTests/Private/MidiOps/MidiTrackFilterTests.cpp b/Engine/Plugins/Runtime/Harmonix/Source/HarmonixMetasoundTests/Private/MidiOps/MidiTrackFilterTests.cpp
--- a/Engine/Plugins/Runtime/Harmonix/Source/HarmonixMetasoundTests/Private/MidiOps/MidiTrackFilterTests.cpp
+++ b/Engine/Plugins/Runtime/Harmonix/Source/HarmonixMetasoundTests/Private/MidiOps/MidiTrackFilterTests.cpp
@@ -15,9 +15,8 @@ namespace Harmonix::Midi::Ops::Tests
 	{
 		FMidiTrackFilter Filter;
 		
-		Metasound::FOperatorSettings OperatorSettings{ 48000, 100 };
-		HarmonixMetasound::FMidiStream InputStream{ OperatorSettings };
-		HarmonixMetasound::FMidiStream OutputStream{ OperatorSettings };
+		HarmonixMetasound::FMidiStream InputStream;
+		HarmonixMetasound::FMidiStream OutputStream;
 
 		constexpr uint16 MinTrack = 1;
 		constexpr uint16 MaxTrack = 37;
diff --git a/Engine/Plugins/Runtime/Harmonix/Source/HarmonixMetasoundTests/Private/MidiOps/PulseGeneratorTests.cpp b/Engine/Plugins/Runtime/Harmonix/Source/HarmonixMetasoundTests/Private/MidiOps/PulseGeneratorTests.cpp
--- a/Engine/Plugins/Runtime/Harmonix/Source/HarmonixMetasoundTests/Private/MidiOps/PulseGeneratorTests.cpp
+++ b/Engine/Plugins/Runtime/Harmonix/Source/HarmonixMetasoundTests/Private/MidiOps/PulseGeneratorTests.cpp
@@ -23,7 +23,7 @@ namespace Harmonix::Midi::Ops::Tests
 		HarmonixMetasound::FMidiClock Clock{ OperatorSettings };
 		Clock.AttachToMidiResource(HarmonixMetasound::FMidiClock::MakeClockConductorMidiData(Tempo, TimeSigNumerator, TimeSigDenominator));
 		
-		HarmonixMetasound::FMidiStream OutputStream{ OperatorSettings };
+		HarmonixMetasound::FMidiStream OutputStream;
 
 		PulseGenerator.SetClock(&Clock);
 
diff --git a/Engine/Plugins/Runtime/Harmonix/Source/HarmonixMetasoundTests/Private/MidiOps/StuckNoteGuardTests.cpp b/Engine/Plugins/Runtime/Harmonix/Source/HarmonixMetasoundTests/Private/MidiOps/StuckNoteGuardTests.cpp
--- a/Engine/Plugins/Runtime/Harmonix/Source/HarmonixMetasoundTests/Private/MidiOps/StuckNoteGuardTests.cpp
+++ b/Engine/Plugins/Runtime/Harmonix/Source/HarmonixMetasoundTests/Private/MidiOps/StuckNoteGuardTests.cpp
@@ -15,9 +15,8 @@ namespace Harmonix::Midi::Ops::Tests
 	{
 		FStuckNoteGuard Guard;
 
-		Metasound::FOperatorSettings OperatorSettings{ 48000, 100 };
-		HarmonixMetasound::FMidiStream InputStream{ OperatorSettings };
-		HarmonixMetasound::FMidiStream OutputStream{ OperatorSettings };
+		HarmonixMetasound::FMidiStream InputStream;
+		HarmonixMetasound::FMidiStream OutputStream;
 
 		// Use a filter that only includes messages on channel 3
 		const auto Channel3Filter = [](const HarmonixMetasound::FMidiStreamEvent& Event)
diff --git a/Engine/Plugins/Runtime/Harmonix/Source/HarmonixMetasoundTests/Private/Nodes/MidiChannelFilterNodeTests.cpp b/Engine/Plugins/Runtime/Harmonix/Source/HarmonixMetasoundTests/Private/Nodes/MidiChannelFilterNodeTests.cpp
--- a/Engine/Plugins/Runtime/Harmonix/Source/HarmonixMetasoundTests/Private/Nodes/MidiChannelFilterNodeTests.cpp
+++ b/Engine/Plugins/Runtime/Harmonix/Source/HarmonixMetasoundTests/Private/Nodes/MidiChannelFilterNodeTests.cpp
@@ -21,8 +21,8 @@ namespace HarmonixMetasound::Nodes::MidiChannelFilter::Tests
 		UTEST_TRUE("Got node MIDI output", NodeMidiOutput.IsSet());
 		
 		Harmonix::Midi::Ops::FMidiChannelFilter FilterForComparison;
-		FMidiStream InputStream{ Generator->OperatorSettings };
-		FMidiStream OutputStream{ Generator->OperatorSettings };
+		FMidiStream InputStream;
+		FMidiStream OutputStream;
 
 		FMidiStreamEvent NoteOnEvent{ static_cast<uint32>(0), FMidiMsg::CreateNoteOn(2, 43, 45) };
 		const uint8 ChannelOneIndexed = NoteOnEvent.MidiMessage.GetStdChannel() + 1;
diff --git a/Engine/Plugins/Runtime/Harmonix/Source/HarmonixMetasoundTests/Private/Nodes/MidiNoteFilterNodeTests.cpp b/Engine/Plugins/Runtime/Harmonix/Source/HarmonixMetasoundTests/Private/Nodes/MidiNoteFilterNodeTests.cpp
--- a/Engine/Plugins/Runtime/Harmonix/Source/HarmonixMetasoundTests/Private/Nodes/MidiNoteFilterNodeTests.cpp
+++ b/Engine/Plugins/Runtime/Harmonix/Source/HarmonixMetasoundTests/Private/Nodes/MidiNoteFilterNodeTests.cpp
@@ -21,8 +21,8 @@ namespace HarmonixMetasound::Nodes::MidiNoteFilter::Tests
 		UTEST_TRUE("Got node MIDI output", NodeMidiOutput.IsSet());
 		
 		Harmonix::Midi::Ops::FMidiNoteFilter FilterForComparison;
-		FMidiStream InputStream{ Generator->OperatorSettings };
-		FMidiStream OutputStream{ Generator->OperatorSettings };
+		FMidiStream InputStream;
+		FMidiStream OutputStream;
 
 		FMidiStreamEvent NoteOnEvent{ static_cast<uint32>(0), FMidiMsg::CreateNoteOn(2, 43, 45) };
 		FMidiStreamEvent CCEvent{ static_cast<uint32>(0), FMidiMsg::CreateControlChange(3, 4, 5) };
diff --git a/Engine/Plugins/Runtime/Harmonix/Source/HarmonixMetasoundTests/Private/Nodes/MidiPulseGeneratorNodeTests.cpp b/Engine/Plugins/Runtime/Harmonix/Source/HarmonixMetasoundTests/Private/Nodes/MidiPulseGeneratorNodeTests.cpp
--- a/Engine/Plugins/Runtime/Harmonix/Source/HarmonixMetasoundTests/Private/Nodes/MidiPulseGeneratorNodeTests.cpp
+++ b/Engine/Plugins/Runtime/Harmonix/Source/HarmonixMetasoundTests/Private/Nodes/MidiPulseGeneratorNodeTests.cpp
@@ -24,7 +24,7 @@ namespace HarmonixMetasound::Nodes::MidiPulseGeneratorNode::Tests
 		UTEST_TRUE("Got clock", Clock.IsSet());
 
 		Harmonix::Midi::Ops::FPulseGenerator PulseGenerator;
-		FMidiStream PulseGeneratorMidiOutput{ Generator->OperatorSettings };
+		FMidiStream PulseGeneratorMidiOutput;
 		PulseGenerator.SetClock(Clock->Get());
 
 		// Render for a bit and expect the same output from both the node and the raw processor
diff --git a/Engine/Plugins/Runtime/Harmonix/Source/HarmonixMetasoundTests/Private/Nodes/MidiStreamTrackFilterNode.cpp b/Engine/Plugins/Runtime/Harmonix/Source/HarmonixMetasoundTests/Private/Nodes/MidiStreamTrackFilterNode.cpp
--- a/Engine/Plugins/Runtime/Harmonix/Source/HarmonixMetasoundTests/Private/Nodes/MidiStreamTrackFilterNode.cpp
+++ b/Engine/Plugins/Runtime/Harmonix/Source/HarmonixMetasoundTests/Private/Nodes/MidiStreamTrackFilterNode.cpp
@@ -22,8 +22,8 @@ namespace HarmonixMetasound::Nodes::MidiTrackFilter::Tests
 		UTEST_TRUE("Got node MIDI output", NodeMidiOutput.IsSet());
 		
 		Harmonix::Midi::Ops::FMidiTrackFilter FilterForComparison;
-		FMidiStream InputStream{ Generator->OperatorSettings };
-		FMidiStream OutputStream{ Generator->OperatorSettings };
+		FMidiStream InputStream;
+		FMidiStream OutputStream;
 
 		FMidiStreamEvent NoteOnEvent{ static_cast<uint32>(0), FMidiMsg::CreateNoteOn(2, 43, 45) };
 		NoteOnEvent.TrackIndex = 3;
