Delete "IAudioData". Replace with "FSoundWaveProxy".

When FKeyzone initializes it's proxy data, that's the point we "Add" the ForceInline sound wave to the Audio Stream Cache for memory tracking.

"Remove" the ForceInline sound if we re-initialize (which should never happen...) and on destruction.

#tests

diff --git a/Engine/Plugins/Runtime/Harmonix/Source/HarmonixDsp/Private/HarmonixDsp/AudioData.cpp b/Engine/Plugins/Runtime/Harmonix/Source/HarmonixDsp/Private/HarmonixDsp/AudioData.cpp
deleted file mode 100644
diff --git a/Engine/Plugins/Runtime/Harmonix/Source/HarmonixDsp/Private/HarmonixDsp/AudioData/StreamingAudioData.cpp b/Engine/Plugins/Runtime/Harmonix/Source/HarmonixDsp/Private/HarmonixDsp/AudioData/StreamingAudioData.cpp
deleted file mode 100644
diff --git a/Engine/Plugins/Runtime/Harmonix/Source/HarmonixDsp/Private/HarmonixDsp/AudioData/StreamingAudioData.h b/Engine/Plugins/Runtime/Harmonix/Source/HarmonixDsp/Private/HarmonixDsp/AudioData/StreamingAudioData.h
deleted file mode 100644
diff --git a/Engine/Plugins/Runtime/Harmonix/Source/HarmonixDsp/Private/HarmonixDsp/AudioData/StreamingAudioRenderer.cpp b/Engine/Plugins/Runtime/Harmonix/Source/HarmonixDsp/Private/HarmonixDsp/AudioData/StreamingAudioRenderer.cpp
deleted file mode 100644
diff --git a/Engine/Plugins/Runtime/Harmonix/Source/HarmonixDsp/Private/HarmonixDsp/AudioData/StreamingAudioRenderer.h b/Engine/Plugins/Runtime/Harmonix/Source/HarmonixDsp/Private/HarmonixDsp/AudioData/StreamingAudioRenderer.h
deleted file mode 100644
diff --git a/Engine/Plugins/Runtime/Harmonix/Source/HarmonixDsp/Private/HarmonixDsp/AudioData/StreamingAudioRendererV2.cpp b/Engine/Plugins/Runtime/Harmonix/Source/HarmonixDsp/Private/HarmonixDsp/AudioData/StreamingAudioRendererV2.cpp
--- a/Engine/Plugins/Runtime/Harmonix/Source/HarmonixDsp/Private/HarmonixDsp/AudioData/StreamingAudioRendererV2.cpp
+++ b/Engine/Plugins/Runtime/Harmonix/Source/HarmonixDsp/Private/HarmonixDsp/AudioData/StreamingAudioRendererV2.cpp
@@ -1,12 +1,9 @@
 // Copyright Epic Games, Inc. All Rights Reserved.
 #include "HarmonixDsp/AudioData/StreamingAudioRendererV2.h"
-#include "HarmonixDsp/AudioData/StreamingAudioData.h"
 
 #include "HarmonixDsp/FusionSampler/FusionSampler.h"
 
-#include "HarmonixDsp/AudioData.h"
 #include "HarmonixDsp/AudioUtility.h"
-#include "HarmonixDsp/Conversions.h"
 #include "HarmonixDsp/GainMatrix.h"
 #include "HarmonixDsp/PannerDetails.h"
 
@@ -32,26 +29,12 @@ void FStreamingAudioRendererV2::Reset()
 	Shifter = nullptr;
 }
 
-void FStreamingAudioRendererV2::SetAudioData(TSharedRef<HarmonixDsp::IAudioData, ESPMode::ThreadSafe> AudioData, const FSettings& InSettings)
+void FStreamingAudioRendererV2::SetAudioData(TSharedRef<FSoundWaveProxy> InSoundWave, const FSettings& InSettings)
 {
-	if (AudioData->GetFormat() == EAudioEncodedFormat::Float32)
-	{
-		TSharedRef<FStreamingAudioData, ESPMode::ThreadSafe> StreamingAudioDataRef = StaticCastSharedRef<FStreamingAudioData>(AudioData);
-		SetAudioData(StreamingAudioDataRef, InSettings);
-	}
-	else
-	{
-		// this renderer is only set up for StreamingAudioData
-		checkNoEntry();
-	}
-}
-
-void FStreamingAudioRendererV2::SetAudioData(TSharedRef<FStreamingAudioData, ESPMode::ThreadSafe> InStreamingAudioData, const FSettings& InSettings)
-{
-	StreamingAudioData = InStreamingAudioData.ToSharedPtr();
+	SoundWaveProxy = InSoundWave.ToSharedPtr();
 
 	WaveProxyReader.Reset();
-	WaveProxyReader = StreamingAudioData->CreateWaveProxyReader();
+	WaveProxyReader = CreateProxyReader(InSoundWave);
 
 	check(WaveProxyReader.IsValid());
 
@@ -77,14 +60,13 @@ void FStreamingAudioRendererV2::SetAudioData(TSharedRef<FStreamingAudioData, ESP
 
 	if (Shifter)
 	{
-		TSharedPtr<const HarmonixDsp::IAudioData, ESPMode::ThreadSafe> AudioData = StaticCastSharedPtr<const HarmonixDsp::IAudioData>(StreamingAudioData);
-		Shifter->SetSampleSourceReset(AudioData, AsShared());
+		Shifter->SetSampleSourceReset(SoundWaveProxy, AsShared());
 	}
 }
 
-const TSharedPtr<HarmonixDsp::IAudioData, ESPMode::ThreadSafe> FStreamingAudioRendererV2::GetAudioData() const
+const TSharedPtr<FSoundWaveProxy> FStreamingAudioRendererV2::GetAudioData() const
 {
-	return StreamingAudioData;
+	return SoundWaveProxy;
 }
 
 void FStreamingAudioRendererV2::MigrateToSampler(const FFusionSampler* InSampler)
@@ -99,7 +81,7 @@ void FStreamingAudioRendererV2::SetFrame(uint32 InFrameNum)
 
 double FStreamingAudioRendererV2::Render(TAudioBuffer<float>& OutBuffer, double InPos, int32 InMaxFrame, double InResampleInc, double InPitchShift, double InSpeed, bool MaintainPitchWhenSpeedChanges, bool InShouldHonorLoopPoints, const FGainMatrix& InGain)
 {
-	if (!StreamingAudioData || StreamingAudioData->GetNumFrames() == 0)
+	if (!SoundWaveProxy || SoundWaveProxy->GetNumFrames() == 0)
 	{
 		OutBuffer.ZeroValidFrames();
 		return InPos;
@@ -122,7 +104,7 @@ double FStreamingAudioRendererV2::RenderInternal(TAudioBuffer<float>& OutBuffer,
 {
 	check(OutBuffer.GetNumValidFrames() <= AudioRendering::kMicroSliceSize);
 
-	if (!StreamingAudioData || StreamingAudioData->GetNumFrames() < (uint32)FMath::FloorToInt32(InPos))
+	if (!SoundWaveProxy || SoundWaveProxy->GetNumFrames() < FMath::FloorToInt32(InPos))
 	{
 		OutBuffer.ZeroValidFrames();
 		return InPos;
@@ -137,7 +119,7 @@ double FStreamingAudioRendererV2::RenderInternal(TAudioBuffer<float>& OutBuffer,
 	{
 		RenderMultiChannelRoutedUnshifted(OutBuffer, LerpArray, NumOutFrames, InGain, InInc, InShouldHonorLoopPoints);
 	}
-	else if (StreamingAudioData->GetNumChannels() <= 2)
+	else if (SoundWaveProxy->GetNumChannels() <= 2)
 	{
 		RenderSimpleUnshifted(OutBuffer, LerpArray, NumOutFrames, InGain, InInc, InShouldHonorLoopPoints);
 	}
@@ -151,7 +133,7 @@ double FStreamingAudioRendererV2::RenderInternal(TAudioBuffer<float>& OutBuffer,
 
 double FStreamingAudioRendererV2::RenderUnshifted(TAudioBuffer<float>& OutBuffer, double InPos, int32 InMaxFrame, double InInc, bool InShouldHonorLoopPoints, const FGainMatrix& InGain)
 {
-	if (!StreamingAudioData)
+	if (!SoundWaveProxy)
 	{
 		OutBuffer.ZeroValidFrames();
 		return InPos;
@@ -174,17 +156,19 @@ int32 FStreamingAudioRendererV2::CalculateNumFramesNeeded(const FLerpData* LerpD
 	}
 
 	// Fail-safe for if we got what appears to be a loop, but we don't have a loop in the audio data
-	if (!StreamingAudioData->GetHasLoopSection())
+	if (SoundWaveProxy->GetLoopRegions().IsEmpty())
 	{
-		return StreamingAudioData->GetNumFrames() - LerpData[0].PosA;
+		return SoundWaveProxy->GetNumFrames() - LerpData[0].PosA;
 	}
 
+	const FSoundWaveCuePoint& LoopRegion = SoundWaveProxy->GetLoopRegions()[0];
+
 	// loop in lerp. 
 	// we need all the samples from LerpData[0].PosA to the end of the loop, 
 	// and then all the samples from the start of the loop to LerpData[NumPoints-1]PosB.
 	// REMEMBER: Loop start and end frames are INCLUSIVE!
-	uint32 FirstFrameInLoop = StreamingAudioData->GetLoopStartFrame();
-	uint32 LastFrameInLoop = StreamingAudioData->GetLoopEndFrame();
+	uint32 FirstFrameInLoop = LoopRegion.FramePosition;
+	uint32 LastFrameInLoop = LoopRegion.FramePosition + LoopRegion.FrameLength;
 	int32 NumNeeded = 1 + LastFrameInLoop - LerpData[0].PosA;
 	NumNeeded += 1 + LerpData[NumPoints - 1].PosB - FirstFrameInLoop;
 	return NumNeeded;
@@ -198,7 +182,7 @@ void FStreamingAudioRendererV2::RenderSimpleUnshifted(TAudioBuffer<float>& OutBu
 	OutBuffer.ZeroValidFrames();
 
 	int32 NumOutChannels = OutBuffer.GetNumValidChannels();
-	int32 NumInputChannels = StreamingAudioData->GetNumChannels();
+	int32 NumInputChannels = SoundWaveProxy->GetNumChannels();
 
 	uint32 StartFrameIndex = LerpArray[0].PosA;
 	int32 NumSourceFramesNeeded = CalculateNumFramesNeeded(LerpArray, InNumFrames);
@@ -254,7 +238,7 @@ void FStreamingAudioRendererV2::RenderMultiChannelUnshifted(TAudioBuffer<float>&
 	OutBuffer.ZeroValidFrames();
 
 	int32 NumOutChannels = OutBuffer.GetNumValidChannels();
-	int32 NumInputChannels = StreamingAudioData->GetNumChannels();
+	int32 NumInputChannels = SoundWaveProxy->GetNumChannels();
 
 	// interleaved
 	Audio::FAlignedFloatBuffer ResampleBuffer;
@@ -327,7 +311,7 @@ void FStreamingAudioRendererV2::RenderMultiChannelRoutedUnshifted(TAudioBuffer<f
 	OutBuffer.ZeroValidFrames();
 
 	int32 NumOutChannels = OutBuffer.GetNumValidChannels();
-	int32 NumInputChannels = StreamingAudioData->GetNumChannels();
+	int32 NumInputChannels = SoundWaveProxy->GetNumChannels();
 
 	// interleaved
 	Audio::FAlignedFloatBuffer ResampleBuffer;
@@ -443,7 +427,7 @@ void FStreamingAudioRendererV2::SeekSourceAudioToFrame(uint32 FrameIdx)
 	// check if we can advance to the desired frame by popping off samples
 	if (FrameIdx > SourceFrameIndex)
 	{
-		if (bLastLoopFrameCached && StreamingAudioData->GetHasLoopSection() && FrameIdx == StreamingAudioData->GetLoopEndFrame())
+		if (bLastLoopFrameCached && HasLoopSection() && FrameIdx == GetLoopEndFrame())
 		{
 			return;
 		}
@@ -511,15 +495,43 @@ uint32 FStreamingAudioRendererV2::GetSourceAudioFrameIndex()
 
 }
 
+TUniquePtr<FSoundWaveProxyReader> FStreamingAudioRendererV2::CreateProxyReader(TSharedRef<FSoundWaveProxy> WaveProxy)
+{
+	FSoundWaveProxyReader::FSettings WaveReaderSettings;
+	WaveReaderSettings.MaxDecodeSizeInFrames = MaxDecodeSizeInFrames;
+	WaveReaderSettings.StartTimeInSeconds = 0.0f;
+	WaveReaderSettings.bMaintainAudioSync = true;
+	return FSoundWaveProxyReader::Create(WaveProxy, WaveReaderSettings);	
+}
+
+bool FStreamingAudioRendererV2::HasLoopSection() const
+{
+	check(SoundWaveProxy);
+	return SoundWaveProxy->GetLoopRegions().Num() > 0;
+}
+
+uint32 FStreamingAudioRendererV2::GetLoopStartFrame() const
+{
+	check(HasLoopSection());
+	return SoundWaveProxy->GetLoopRegions()[0].FramePosition;
+}
+
+uint32 FStreamingAudioRendererV2::GetLoopEndFrame() const
+{
+	check(HasLoopSection())
+	const FSoundWaveCuePoint& LoopRegion = SoundWaveProxy->GetLoopRegions()[0];
+	return LoopRegion.FramePosition + LoopRegion.FrameLength;
+}
+
 void FStreamingAudioRendererV2::GenerateSourceAudio(uint32 StartFrameIndex, Audio::FAlignedFloatBuffer& OutAudio, bool bHonorLoopRegion)
 {
 	check(OutAudio.Num() % NumDeinterleaveChannels == 0);
 
-	if (bHonorLoopRegion && StreamingAudioData->GetHasLoopSection())
+	if (bHonorLoopRegion && HasLoopSection())
 	{
 		uint32 NumFramesRequested = OutAudio.Num() / NumDeinterleaveChannels;
-		uint32 FirstFrameInLoop = StreamingAudioData->GetLoopStartFrame();
-		uint32 LastFrameInLoop = StreamingAudioData->GetLoopEndFrame();
+		uint32 FirstFrameInLoop = GetLoopStartFrame();
+		uint32 LastFrameInLoop = GetLoopEndFrame();
 		uint32 LoopLengthFrames = 1 + LastFrameInLoop - FirstFrameInLoop;
 		bool bLoopingThisChunk = (StartFrameIndex + NumFramesRequested) > (LastFrameInLoop + 1);
 
@@ -635,6 +647,6 @@ void FStreamingAudioRendererV2::GenerateSourceAudioInternal(uint32 StartFrameInd
 		FMemory::Memzero(&OutAudioData[BufferIdx], sizeof(float) * NumSamplesRequested);
 
 		UE_LOG(LogHarmonixStreamingAudioRendererV2, Verbose, TEXT("%s: Failed to generated samples: StartFrameIndex: %d, NumSamplesRequested: %d"), 
-			*StreamingAudioData->GetName().ToString(), StartFrameIndex, NumSamplesRequested);
+			*SoundWaveProxy->GetFName().ToString(), StartFrameIndex, NumSamplesRequested);
 	}
 }
diff --git a/Engine/Plugins/Runtime/Harmonix/Source/HarmonixDsp/Private/HarmonixDsp/AudioData/StreamingAudioRendererV2.h b/Engine/Plugins/Runtime/Harmonix/Source/HarmonixDsp/Private/HarmonixDsp/AudioData/StreamingAudioRendererV2.h
--- a/Engine/Plugins/Runtime/Harmonix/Source/HarmonixDsp/Private/HarmonixDsp/AudioData/StreamingAudioRendererV2.h
+++ b/Engine/Plugins/Runtime/Harmonix/Source/HarmonixDsp/Private/HarmonixDsp/AudioData/StreamingAudioRendererV2.h
@@ -11,14 +11,8 @@
 #include "DSP/MultichannelBuffer.h"
 #include "DSP/MultichannelLinearResampler.h"
 
-namespace HarmonixDsp
-{
-	class IAudioData;
-}
-
 class FSoundWaveProxy;
 class FSoundWaveProxyReader;
-class FStreamingAudioData;
 class FFusionSampler;
 class IStretcherAndPitchShifter;
 
@@ -34,10 +28,9 @@ public:
 	~FStreamingAudioRendererV2();
 
 	virtual void Reset() override;
-	virtual void SetAudioData(TSharedRef<HarmonixDsp::IAudioData, ESPMode::ThreadSafe> AudioData, const FSettings& InSettings) override;
-	void SetAudioData(TSharedRef<FStreamingAudioData, ESPMode::ThreadSafe> StreamingAudioData, const FSettings& InSettings);
-	virtual const TSharedPtr<HarmonixDsp::IAudioData, ESPMode::ThreadSafe> GetAudioData() const;
-
+	virtual void SetAudioData(TSharedRef<FSoundWaveProxy> SoundWaveProxy, const FSettings& InSettings) override;
+	virtual const TSharedPtr<FSoundWaveProxy> GetAudioData() const override;
+	
 	virtual void MigrateToSampler(const FFusionSampler* InSampler) override;
 
 	virtual void SetFrame(uint32 InFrameNum) override;
@@ -65,7 +58,11 @@ private:
 
 	uint32 GetSourceAudioFrameIndex();
 
-private:
+	static TUniquePtr<FSoundWaveProxyReader> CreateProxyReader(TSharedRef<FSoundWaveProxy> WaveProxy);
+
+	bool HasLoopSection() const;
+	uint32 GetLoopStartFrame() const;
+	uint32 GetLoopEndFrame() const;
 
 	const FFusionSampler* MySampler = nullptr;
 	const TArray<FTrackChannelInfo>* TrackChannelInfo = nullptr;
@@ -74,7 +71,7 @@ private:
 	// Ref to the actual streaming audio data
 	// this data is a shared instance of audio data
 	// it will get loaded on construction 
-	TSharedPtr<FStreamingAudioData, ESPMode::ThreadSafe> StreamingAudioData;
+	TSharedPtr<FSoundWaveProxy> SoundWaveProxy;
 
 	TUniquePtr<FSoundWaveProxyReader> WaveProxyReader;
 
@@ -90,6 +87,7 @@ private:
 	// needs to be small enough to avoid audio artifacts and syncing issues
 	// but also large enough that we're decoding multiple times per block
 	static constexpr int32 DeinterleaveBlockSizeInFrames = 256;
+	static constexpr uint32 MaxDecodeSizeInFrames = 1024;
 
 	int32 CalculateNumFramesNeeded(const FLerpData* LerpData, int32 NumPoints);
 };
\ No newline at end of file
diff --git a/Engine/Plugins/Runtime/Harmonix/Source/HarmonixDsp/Private/HarmonixDsp/AudioDataRenderer.cpp b/Engine/Plugins/Runtime/Harmonix/Source/HarmonixDsp/Private/HarmonixDsp/AudioDataRenderer.cpp
--- a/Engine/Plugins/Runtime/Harmonix/Source/HarmonixDsp/Private/HarmonixDsp/AudioDataRenderer.cpp
+++ b/Engine/Plugins/Runtime/Harmonix/Source/HarmonixDsp/Private/HarmonixDsp/AudioDataRenderer.cpp
@@ -1,31 +1,40 @@
 // Copyright Epic Games, Inc. All Rights Reserved.
 
 #include "HarmonixDsp/AudioDataRenderer.h"
-#include "HarmonixDsp/AudioData.h"
+#include "Sound/SoundWave.h"
 
 double IAudioDataRenderer::CalculateLerpData(FLerpData* LerpArray, int32 InNumPoints, uint32 InNumOutFrames, double InPos, int32 InMaxFrame, bool InHonorLoopPoints, double InInc) const
 {
 	check(InNumPoints >= (int32)InNumOutFrames);
 	FMemory::Memset(LerpArray, 0, sizeof(FLerpData) * InNumPoints);
 
-	TSharedPtr<HarmonixDsp::IAudioData, ESPMode::ThreadSafe> AudioData = GetAudioData();
-	if (!AudioData)
+	const TSharedPtr<FSoundWaveProxy> SoundWaveProxy = GetAudioData();
+	if (!SoundWaveProxy)
 	{
 		return 0.0;
 	}
 
-	if (InMaxFrame < 0 || InMaxFrame > (int32)AudioData->GetNumFrames())
+	if (InMaxFrame < 0 || InMaxFrame > (int32)SoundWaveProxy->GetNumFrames())
 	{
-		InMaxFrame = (int32)AudioData->GetNumFrames();
+		InMaxFrame = (int32)SoundWaveProxy->GetNumFrames();
 	}
 
 	double NumSourceFrames = (double)InMaxFrame;
 
-	bool HasTailSection = AudioData->GetHasTailSection();
-	bool IsLooping = AudioData->GetHasLoopSection() && (!HasTailSection || (HasTailSection && InHonorLoopPoints));
+	uint32 LoopStartFrameIndex = 0;
+	uint32 LoopEndFrameIndex = 0;
+	bool IsLooping = false;
+	bool HasTailSection = false;
+	if (SoundWaveProxy->GetLoopRegions().Num() > 0)
+	{
+		const FSoundWaveCuePoint& LoopRegion = SoundWaveProxy->GetLoopRegions()[0];
+		LoopStartFrameIndex = LoopRegion.FramePosition;
+		LoopEndFrameIndex = LoopRegion.FramePosition + LoopRegion.FrameLength;
+		HasTailSection = LoopEndFrameIndex != SoundWaveProxy->GetNumFrames() - 1;
+		IsLooping = (!HasTailSection || (HasTailSection && InHonorLoopPoints));
+	}
+
 
-	uint32 LoopStartFrameIndex = IsLooping ? AudioData->GetLoopStartFrame() : 0;
-	uint32 LoopEndFrameIndex = IsLooping ? AudioData->GetLoopEndFrame() : 0;
 	double LoopEnd = (double)LoopEndFrameIndex + 1.0;
 	double LoopLength = LoopEnd - (double)LoopStartFrameIndex;
 
diff --git a/Engine/Plugins/Runtime/Harmonix/Source/HarmonixDsp/Private/HarmonixDsp/FusionSampler/FusionSampler.cpp b/Engine/Plugins/Runtime/Harmonix/Source/HarmonixDsp/Private/HarmonixDsp/FusionSampler/FusionSampler.cpp
--- a/Engine/Plugins/Runtime/Harmonix/Source/HarmonixDsp/Private/HarmonixDsp/FusionSampler/FusionSampler.cpp
+++ b/Engine/Plugins/Runtime/Harmonix/Source/HarmonixDsp/Private/HarmonixDsp/FusionSampler/FusionSampler.cpp
@@ -367,7 +367,7 @@ int32 FFusionSampler::GatherMatchingKeyzones(uint8 TransposedNote, uint8 InVeloc
 	for (int32 KeyzoneIdx = 0; KeyzoneIdx < InKeyzones.Num(); ++KeyzoneIdx)
 	{
 		const FKeyzoneSettings& Keyzone = InKeyzones[KeyzoneIdx];
-		if (!Keyzone.AudioSample)
+		if (!Keyzone.SoundWaveProxy)
 		{
 			continue;
 		}
diff --git a/Engine/Plugins/Runtime/Harmonix/Source/HarmonixDsp/Private/HarmonixDsp/FusionSampler/FusionVoice.cpp b/Engine/Plugins/Runtime/Harmonix/Source/HarmonixDsp/Private/HarmonixDsp/FusionSampler/FusionVoice.cpp
--- a/Engine/Plugins/Runtime/Harmonix/Source/HarmonixDsp/Private/HarmonixDsp/FusionSampler/FusionVoice.cpp
+++ b/Engine/Plugins/Runtime/Harmonix/Source/HarmonixDsp/Private/HarmonixDsp/FusionSampler/FusionVoice.cpp
@@ -3,9 +3,7 @@
 #include "HarmonixDsp/FusionSampler/FusionSampler.h"
 #include "HarmonixDsp/FusionSampler/FusionVoicePool.h"
 
-#include "HarmonixDsp/AudioData.h"
 #include "HarmonixDsp/AudioDataRenderer.h"
-#include "HarmonixDsp/AudioData/StreamingAudioRenderer.h"
 #include "HarmonixDsp/AudioData/StreamingAudioRendererV2.h"
 #include "HarmonixDsp/AudioUtility.h"
 #include "HarmonixDsp/StretcherAndPitchShifter.h"
@@ -14,6 +12,8 @@
 
 #include "HAL/IConsoleManager.h"
 
+#include "Sound/SoundWave.h"
+
 DEFINE_LOG_CATEGORY(LogFusionVoice);
 
 extern size_t gAudioSliceCounter;
@@ -23,14 +23,6 @@ namespace FusionVoice
 	static const float kInternalTrimGain = 1.0f; 
 	static const float kVsoAmount = 0.01f;
 	static const float kVsoThresholdMs = 5.0f;
-
-	bool bUseNewStreamingAudioRenderer = true;
-	FAutoConsoleVariableRef CVarMetaSoundProfileAllEnabled(
-		TEXT("au.fusion.usenewstreamingaudiorenderer"),
-		bUseNewStreamingAudioRenderer,
-		TEXT("Use the new StreamingAudioRenderer\n")
-		TEXT("False: Use the old one, True: Use the new one."),
-		ECVF_Default);
 }
 
 const double FFusionVoice::kMaxPitchOffsetCents = 12 * HarmonixDsp::kCentsPerOctave;
@@ -44,14 +36,9 @@ FFusionVoice::FFusionVoice()
 	, MaxAudioLevel(0.0f)
 	, VoicePool(nullptr)
 {
-	if (FusionVoice::bUseNewStreamingAudioRenderer)
-	{
-		ActiveRenderer = MakeShared<FStreamingAudioRendererV2>();
-	}
-	else
-	{
-		ActiveRenderer = MakeShared<FStreamingAudioRenderer>();
-	}
+
+	ActiveRenderer = MakeShared<FStreamingAudioRendererV2>();
+
 	// For now, set the sample rate to a reasonable default.
 	// Sample rate will be set properly when the voice is assigned
 	// to a sampler.
@@ -86,11 +73,7 @@ bool FFusionVoice::AssignIDs(FFusionSampler* InSampler, const FKeyzoneSettings*
 {
 	check(InSampler);
 	check(InKeyZone);
-	check(InKeyZone->AudioSample);
-	if (InKeyZone->AudioSample->Failed())
-	{
-		return false;
-	}
+	check(InKeyZone->SoundWaveProxy);
 
 #if CPUPROFILERTRACE_ENABLED
 	const IConsoleVariable* ProfilingAllGraphsCheat = IConsoleManager::Get().FindConsoleVariable(TEXT("au.MetaSound.ProfileAllGraphs"));
@@ -128,19 +111,21 @@ bool FFusionVoice::AssignIDs(FFusionSampler* InSampler, const FKeyzoneSettings*
 	// so in that case set NumInChannels to 0 here and we'll sort this
 	// out later...
 	// 0 because we need to configure to the output buffer we are asked to fill
-	int32 NumInChannels = InKeyZone->TrackMap.Num() == 0 ? KeyZone->AudioSample->GetNumChannels() : 0; 
 	
+	int32 NumInChannels = InKeyZone->TrackMap.Num() == 0 ? KeyZone->SoundWaveProxy->GetNumChannels() : 0; 
+	EAudioBufferChannelLayout ChannelLayout = HarmonixDsp::FAudioBuffer::GetDefaultChannelLayoutForChannelCount(NumInChannels);
+	ESpeakerMask::Type ChannelMask = HarmonixDsp::FAudioBuffer::GetChannelMaskForNumChannels(NumInChannels);
 	Panner.Setup(KeyZone->Pan,
 		NumInChannels,
 		MySampler->GetNumAudioOutputChannels(),
-		KeyZone->AudioSample->GetChannelLayout(),
-		KeyZone->AudioSample->GetChannelMask(),
+		ChannelLayout,
+		ChannelMask,
 		1.0f,
 		MySampler->GetGainTable());
 
 	// compare the audio file's sample rate to the output sample rate
 	// so we can make them match
-	double FileSamplesPerSecond = InKeyZone->AudioSample->GetSampleRate();
+	double FileSamplesPerSecond = InKeyZone->SoundWaveProxy->GetSampleRate();
 	FileToOutputSampleRatio = FileSamplesPerSecond * SecondsPerSample;
 
 	IAudioDataRenderer::FSettings Settings;
@@ -148,9 +133,9 @@ bool FFusionVoice::AssignIDs(FFusionSampler* InSampler, const FKeyzoneSettings*
 	Settings.TrackChannelInfo = &InKeyZone->TrackMap;
 	Settings.Sampler = MySampler;
 	//TRACE_BOOKMARK(TEXT("Starting Waveform: %s"), *KeyZone->AudioSample->GetName().ToString());
-	ActiveRenderer->SetAudioData(InKeyZone->AudioSample.ToSharedRef(), Settings);
+	ActiveRenderer->SetAudioData(InKeyZone->SoundWaveProxy.ToSharedRef(), Settings);
 
-	OutputBuffer.SetNumValidChannels(InKeyZone->AudioSample->GetNumChannels());
+	OutputBuffer.SetNumValidChannels(InKeyZone->SoundWaveProxy->GetNumChannels());
 
 	VoiceID = InVoiceID;
 	bWaitingForAttack = true;
@@ -392,10 +377,11 @@ void FFusionVoice::AttackWithTargetNote(uint8 InMidiNoteNumber, float InGain, in
 	SamplePos = KeyzoneStartOffset + (OutputSampleOffset * FileToOutputSampleRatio);
 	StartPos = SamplePos;
 	EndOfSampleData = (KeyZone->SampleEndOffset == -1) ?
-		KeyZone->AudioSample->GetNumFrames() :
+		KeyZone->SoundWaveProxy->GetNumFrames() :
 		KeyZone->SampleEndOffset;
 
-	if (SamplePos > EndOfSampleData && !KeyZone->AudioSample->GetHasLoopSection())
+	bool HasLoopSection = KeyZone->SoundWaveProxy->GetLoopRegions().Num() > 0;
+	if (SamplePos > EndOfSampleData && !HasLoopSection)
 	{
 		Kill();
 		return;
@@ -676,7 +662,7 @@ uint32 FFusionVoice::Process(uint32 InSliceIndex, uint32 InSubsliceIndex, float*
 		return 0;
 	}
 
-	if (!KeyZone->AudioSample)
+	if (!KeyZone->SoundWaveProxy)
 	{
 		return 0;
 	}
@@ -765,7 +751,7 @@ uint32 FFusionVoice::Process(uint32 InSliceIndex, uint32 InSubsliceIndex, float*
 					currentSampleFrame -= KeyZone->SampleStartOffset;
 				}
 
-				float ActualElapsedMs = (currentSampleFrame * 1000.0f) / KeyZone->AudioSample->GetSampleRate();
+				float ActualElapsedMs = (currentSampleFrame * 1000.0f) / KeyZone->SoundWaveProxy->GetSampleRate();
 				float ErrorMs = ActualElapsedMs - ExpectedElapsedMs;
 
 				if (!bHasRenderedAnySamples)
@@ -775,11 +761,11 @@ uint32 FFusionVoice::Process(uint32 InSliceIndex, uint32 InSubsliceIndex, float*
 						UE_LOG(LogFusionVoice, Verbose, TEXT("Adjusting start time due to error of %f ms"), ErrorMs);
 					}
 
-					SamplePos = ExpectedElapsedMs / 1000.0f * (float)KeyZone->AudioSample->GetSampleRate();
+					SamplePos = ExpectedElapsedMs / 1000.0f * (float)KeyZone->SoundWaveProxy->GetSampleRate();
 					
 					if (SamplePos > 0.0)
 					{
-						UE_LOG(LogFusionVoice, Verbose, TEXT("SHAVING %f ms from the beginning of %s"), ErrorMs, *KeyZone->AudioSample->GetName().ToString());
+						UE_LOG(LogFusionVoice, Verbose, TEXT("SHAVING %f ms from the beginning of %s"), ErrorMs, *KeyZone->SoundWaveProxy->GetFName().ToString());
 					}
 
 					if (SamplePos < 0.0)
@@ -790,7 +776,7 @@ uint32 FFusionVoice::Process(uint32 InSliceIndex, uint32 InSubsliceIndex, float*
 					{
 						SamplePos += KeyZone->SampleStartOffset;
 					}
-					if (SamplePos > EndOfSampleData && !KeyZone->AudioSample->GetHasLoopSection())
+					if (SamplePos > EndOfSampleData && KeyZone->SoundWaveProxy->GetLoopRegions().IsEmpty())
 					{
 						Kill();
 						return 0;
@@ -802,7 +788,7 @@ uint32 FFusionVoice::Process(uint32 InSliceIndex, uint32 InSubsliceIndex, float*
 				{
 					if (ErrorMs > 20.0f)
 					{
-						UE_LOG(LogFusionVoice, Verbose, TEXT("Fixing %s offset -> %f\n"), *KeyZone->AudioSample->GetName().ToString(), ErrorMs);
+						UE_LOG(LogFusionVoice, Verbose, TEXT("Fixing %s offset -> %f\n"), *KeyZone->SoundWaveProxy->GetFName().ToString(), ErrorMs);
 					}
 					CurrentVso = -(FusionVoice::kVsoAmount * RenderSpeed);
 				}
@@ -810,7 +796,7 @@ uint32 FFusionVoice::Process(uint32 InSliceIndex, uint32 InSubsliceIndex, float*
 				{
 					if (ErrorMs < -20.0f)
 					{
-						UE_LOG(LogFusionVoice, Verbose, TEXT("Fixing %s offset -> %f\n"), *KeyZone->AudioSample->GetName().ToString(), ErrorMs);
+						UE_LOG(LogFusionVoice, Verbose, TEXT("Fixing %s offset -> %f\n"), *KeyZone->SoundWaveProxy->GetFName().ToString(), ErrorMs);
 					}
 					CurrentVso = FusionVoice::kVsoAmount * RenderSpeed;
 				}
@@ -842,7 +828,7 @@ uint32 FFusionVoice::Process(uint32 InSliceIndex, uint32 InSubsliceIndex, float*
 
 		// We probably have a trackmap, and this is the first render call,
 		// so we have to configure the panner for the output buffer setup...
-		int32 NumInChannels = KeyZone->AudioSample->GetNumChannels();
+		int32 NumInChannels = KeyZone->SoundWaveProxy->GetNumChannels();
 		if (NumInChannels != Panner.GetCurrentGainMatrix().GetNumInChannels())
 		{
 			Panner.Configure(NumInChannels, OutputBuffer);
diff --git a/Engine/Plugins/Runtime/Harmonix/Source/HarmonixDsp/Private/HarmonixDsp/FusionSampler/FusionVoicePool.cpp b/Engine/Plugins/Runtime/Harmonix/Source/HarmonixDsp/Private/HarmonixDsp/FusionSampler/FusionVoicePool.cpp
--- a/Engine/Plugins/Runtime/Harmonix/Source/HarmonixDsp/Private/HarmonixDsp/FusionSampler/FusionVoicePool.cpp
+++ b/Engine/Plugins/Runtime/Harmonix/Source/HarmonixDsp/Private/HarmonixDsp/FusionSampler/FusionVoicePool.cpp
@@ -1,12 +1,11 @@
 // Copyright Epic Games, Inc. All Rights Reserved.
 #include "HarmonixDsp/FusionSampler/FusionVoicePool.h"
-#include "HarmonixDsp/FusionSampler/SingletonFusionVoicePool.h"
 #include "HarmonixDsp/FusionSampler/FusionVoice.h"
 #include "HarmonixDsp/FusionSampler/FusionPatch.h"
 #include "HarmonixDsp/FusionSampler/FusionSampler.h"
 #include "HarmonixDsp/FusionSampler/Settings/KeyzoneSettings.h"
-#include "HarmonixDsp/AudioData.h"
 #include "HarmonixDsp/StretcherAndPitchShifterFactory.h"
+#include "Sound/SoundWave.h"
 
 DEFINE_LOG_CATEGORY(LogFusionVoicePool);
 
@@ -311,14 +310,14 @@ FFusionVoice* FFusionVoicePool::GetFreeVoice(
 		if (!Factory)
 		{
 			FName FactoryName = InKeyzone->TimeStretchConfig.PitchShifter.Name;
-			ensureMsgf(false, TEXT("Attempting to play a fusion InKeyzone that is set to \"maintain time\", but no shifter factory was available for assigned Pitch Shifter: %s. Check project configuration! (%s)"), *FactoryName.ToString(), *InKeyzone->AudioSample->GetName().ToString());
+			ensureMsgf(false, TEXT("Attempting to play a fusion InKeyzone that is set to \"maintain time\", but no shifter factory was available for assigned Pitch Shifter: %s. Check project configuration! (%s)"), *FactoryName.ToString(), *InKeyzone->SoundWaveProxy->GetFName().ToString());
 			return nullptr;
 		}
 	}
 
-	if (!InKeyzone->AudioSample || InKeyzone->AudioSample->Failed())
+	if (!InKeyzone->SoundWaveProxy)
 	{
-		UE_LOG(LogFusionVoicePool, Warning, TEXT("Asked to allocate a InSampler voice for a sample resource that failed to load -> %s\n"), *InKeyzone->AudioSample->GetName().ToString());
+		UE_LOG(LogFusionVoicePool, Warning, TEXT("Asked to allocate a InSampler voice for a sample resource that failed to load -> %s\n"), *InKeyzone->SoundWaveProxy->GetFName().ToString());
 		return nullptr;
 	}
 
diff --git a/Engine/Plugins/Runtime/Harmonix/Source/HarmonixDsp/Private/HarmonixDsp/FusionSampler/Settings/KeyzoneSettings.cpp b/Engine/Plugins/Runtime/Harmonix/Source/HarmonixDsp/Private/HarmonixDsp/FusionSampler/Settings/KeyzoneSettings.cpp
--- a/Engine/Plugins/Runtime/Harmonix/Source/HarmonixDsp/Private/HarmonixDsp/FusionSampler/Settings/KeyzoneSettings.cpp
+++ b/Engine/Plugins/Runtime/Harmonix/Source/HarmonixDsp/Private/HarmonixDsp/FusionSampler/Settings/KeyzoneSettings.cpp
@@ -1,17 +1,23 @@
 // Copyright Epic Games, Inc. All Rights Reserved.
 #include "HarmonixDsp/FusionSampler/Settings/KeyzoneSettings.h"
 #include "HarmonixDsp/FusionSampler/SingletonFusionVoicePool.h"
-#include "HarmonixDsp/FusionSampler/FusionPatch.h"
 #include "HarmonixDsp/AudioUtility.h"
-#include "HarmonixDsp/AudioData.h"
-#include "HarmonixDsp/AudioData/StreamingAudioData.h"
 
+#include "AudioStreaming.h"
 #include "Sound/SoundWave.h"
 
 #include "HAL/PlatformTime.h"
 
 DEFINE_LOG_CATEGORY(LogKeyzoneSettings);
 
+FKeyzoneSettings::~FKeyzoneSettings()
+{
+	if (SoundWaveProxy && SoundWaveProxy->GetLoadingBehavior() == ESoundWaveLoadingBehavior::ForceInline)
+	{
+		IStreamingManager::Get().GetAudioStreamingManager().RemoveForceInlineSoundWave(SoundWaveProxy);
+	}
+}
+
 void FKeyzoneSettings::InitProxyData(const Audio::FProxyDataInitParams& InitParams)
 {
 	QUICK_SCOPE_CYCLE_COUNTER(STAT_FKeyzoneSettings_InitProxyData);
@@ -24,9 +30,20 @@ void FKeyzoneSettings::InitProxyData(const Audio::FProxyDataInitParams& InitPara
 		return;
 	}
 
-	AudioSample = CreateStreamingAudioData(InitParams);
-
-	check(AudioSample);
+	if (SoundWaveProxy && SoundWaveProxy->GetLoadingBehavior() == ESoundWaveLoadingBehavior::ForceInline)
+	{
+		// I don't think we'de ever hit this code, but let's just do this to be sure.
+		IStreamingManager::Get().GetAudioStreamingManager().RemoveForceInlineSoundWave(SoundWaveProxy);
+	}
+	
+	TSharedPtr<Audio::IProxyData> AudioProxy = SoundWave->CreateProxyData(InitParams);
+	SoundWaveProxy = StaticCastSharedPtr<FSoundWaveProxy>(AudioProxy);
+	check(SoundWaveProxy);
+		
+	if (SoundWaveProxy->GetLoadingBehavior() == ESoundWaveLoadingBehavior::ForceInline)
+	{
+		IStreamingManager::Get().GetAudioStreamingManager().AddForceInlineSoundWave(SoundWaveProxy);
+	}
 
 	// nulling out sound wave asset for this Keyzone
 	// now that we are turning into ProxyData
@@ -42,36 +59,6 @@ void FKeyzoneSettings::InitProxyData(const Audio::FProxyDataInitParams& InitPara
 	}
 }
 
-TSharedPtr<FStreamingAudioData, ESPMode::ThreadSafe> FKeyzoneSettings::CreateStreamingAudioData(const Audio::FProxyDataInitParams& InitParams) const
-{
-	TSharedPtr<FStreamingAudioData, ESPMode::ThreadSafe> StreamingAudioSample = StaticCastSharedPtr<FStreamingAudioData>(AudioSample);
-	if (StreamingAudioSample.IsValid())
-	{
-		return StreamingAudioSample;
-	}
-
-	if (!ensure(SoundWave))
-	{
-		UE_LOG(LogKeyzoneSettings, Error, TEXT("SoundWave asset is null in keyzone settings when making proxy data!!"));
-		return nullptr;
-	}
-
-	TSharedPtr<Audio::IProxyData, ESPMode::ThreadSafe> AudioProxy = SoundWave->CreateProxyData(InitParams);
-
-	FSoundWaveProxyPtr WaveProxy = StaticCastSharedPtr<FSoundWaveProxy>(AudioProxy);
-	if (!ensure(WaveProxy))
-	{
-		UE_LOG(LogKeyzoneSettings, Error, TEXT("Failed to make sound wave proxy data for SoundWave asset!!"));
-		return nullptr;
-	}
-
-	UE_LOG(LogKeyzoneSettings, Verbose, TEXT("Creating new shared AudioData"));
-
-	FStreamingAudioData::FSettings StreamingAudioSettings;
-	StreamingAudioSample = FStreamingAudioData::GetShared<FStreamingAudioData>(WaveProxy.ToSharedRef(), StreamingAudioSettings);
-	return StreamingAudioSample;
-}
-
 bool FKeyzoneSettings::ContainsNoteAndVelocity(uint8 InNote, uint8 InVelocity) const
 {
 	return (InNote >= MinNote && InNote <= MaxNote) 
diff --git a/Engine/Plugins/Runtime/Harmonix/Source/HarmonixDsp/Private/HarmonixDsp/Stretchers/SmbPitchShifter.cpp b/Engine/Plugins/Runtime/Harmonix/Source/HarmonixDsp/Private/HarmonixDsp/Stretchers/SmbPitchShifter.cpp
--- a/Engine/Plugins/Runtime/Harmonix/Source/HarmonixDsp/Private/HarmonixDsp/Stretchers/SmbPitchShifter.cpp
+++ b/Engine/Plugins/Runtime/Harmonix/Source/HarmonixDsp/Private/HarmonixDsp/Stretchers/SmbPitchShifter.cpp
@@ -52,13 +52,13 @@
 *****************************************************************************/
 
 #include "HarmonixDsp/Stretchers/SmbPitchShifter.h"
-#include "HarmonixDsp/AudioData.h"
 #include "HarmonixDsp/AudioDataRenderer.h"
 
 #include "HAL/PlatformMemory.h"
 #include "HAL/PlatformMath.h"
 #include "Math/UnrealMathUtility.h"
 #include "Misc/AssertionMacros.h"
+#include "Sound/SoundWave.h"
 
 //#include "kiss_fft/kiss_fft.h"
 
diff --git a/Engine/Plugins/Runtime/Harmonix/Source/HarmonixDsp/Public/HarmonixDsp/AudioBufferConstants.h b/Engine/Plugins/Runtime/Harmonix/Source/HarmonixDsp/Public/HarmonixDsp/AudioBufferConstants.h
--- a/Engine/Plugins/Runtime/Harmonix/Source/HarmonixDsp/Public/HarmonixDsp/AudioBufferConstants.h
+++ b/Engine/Plugins/Runtime/Harmonix/Source/HarmonixDsp/Public/HarmonixDsp/AudioBufferConstants.h
@@ -133,6 +133,24 @@ namespace FAudioBuffer
 		}
 	}
 
+	FORCEINLINE constexpr HARMONIXDSP_API ESpeakerMask::Type GetChannelMaskForNumChannels(uint32 NumChannels)
+	{
+		switch (NumChannels)
+		{
+		case 1: return ESpeakerMask::UnspecifiedMono;
+		case 2: return ESpeakerMask::Stereo;
+		case 3: return ESpeakerMask::FiveDotZero;
+		case 6: return ESpeakerMask::FiveDotOne;
+		case 8: return ESpeakerMask::SevenDotOne;
+		case 4:
+			return ESpeakerMask::AmbisonicAssigned;
+		}
+
+		// Invalid num channels
+		checkNoEntry();
+		return ESpeakerMask::AllSpeakers;
+	}
+
 	const uint32 kLayoutHasSpeakerMasks[] = {
 		// Raw,
 		0,                      
diff --git a/Engine/Plugins/Runtime/Harmonix/Source/HarmonixDsp/Public/HarmonixDsp/AudioData.h b/Engine/Plugins/Runtime/Harmonix/Source/HarmonixDsp/Public/HarmonixDsp/AudioData.h
deleted file mode 100644
diff --git a/Engine/Plugins/Runtime/Harmonix/Source/HarmonixDsp/Public/HarmonixDsp/AudioDataRenderer.h b/Engine/Plugins/Runtime/Harmonix/Source/HarmonixDsp/Public/HarmonixDsp/AudioDataRenderer.h
--- a/Engine/Plugins/Runtime/Harmonix/Source/HarmonixDsp/Public/HarmonixDsp/AudioDataRenderer.h
+++ b/Engine/Plugins/Runtime/Harmonix/Source/HarmonixDsp/Public/HarmonixDsp/AudioDataRenderer.h
@@ -5,17 +5,13 @@
 
 #include "HarmonixDsp/Streaming/TrackChannelInfo.h"
 
-namespace HarmonixDsp
-{
-	class IAudioData;
-}
-
 class IStretcherAndPitchShifter;
 struct FTrackChannelInfo;
 template<typename Type> 
 class TAudioBuffer;
 class FGainMatrix;
 class FFusionSampler;
+class FSoundWaveProxy;
 
 class IAudioDataRenderer : public TSharedFromThis<IAudioDataRenderer>
 {
@@ -31,11 +27,11 @@ public:
 
 	virtual ~IAudioDataRenderer() {};
 
-	virtual void SetAudioData(TSharedRef<HarmonixDsp::IAudioData, ESPMode::ThreadSafe> AudioData, const FSettings& InSettings) = 0;
+	virtual void SetAudioData(TSharedRef<FSoundWaveProxy> SoundWaveProxy, const FSettings& InSettings) = 0;
 
 	virtual void Reset() = 0;
-
-	virtual const TSharedPtr<HarmonixDsp::IAudioData, ESPMode::ThreadSafe> GetAudioData() const = 0;
+	
+	virtual const TSharedPtr<FSoundWaveProxy> GetAudioData() const = 0;
 
 	virtual void MigrateToSampler(const FFusionSampler* InSampler) = 0;
 
diff --git a/Engine/Plugins/Runtime/Harmonix/Source/HarmonixDsp/Public/HarmonixDsp/FusionSampler/FusionVoice.h b/Engine/Plugins/Runtime/Harmonix/Source/HarmonixDsp/Public/HarmonixDsp/FusionSampler/FusionVoice.h
--- a/Engine/Plugins/Runtime/Harmonix/Source/HarmonixDsp/Public/HarmonixDsp/FusionSampler/FusionVoice.h
+++ b/Engine/Plugins/Runtime/Harmonix/Source/HarmonixDsp/Public/HarmonixDsp/FusionSampler/FusionVoice.h
@@ -20,10 +20,6 @@
 
 DECLARE_LOG_CATEGORY_EXTERN(LogFusionVoice, Log, All);
 
-namespace HarmonixDsp
-{
-	class IAudioData;
-}
 class IAudioDataRenderer;
 class IStretcherAndPitchShifter;
 class FFusionSampler;
diff --git a/Engine/Plugins/Runtime/Harmonix/Source/HarmonixDsp/Public/HarmonixDsp/FusionSampler/Settings/KeyzoneSettings.h b/Engine/Plugins/Runtime/Harmonix/Source/HarmonixDsp/Public/HarmonixDsp/FusionSampler/Settings/KeyzoneSettings.h
--- a/Engine/Plugins/Runtime/Harmonix/Source/HarmonixDsp/Public/HarmonixDsp/FusionSampler/Settings/KeyzoneSettings.h
+++ b/Engine/Plugins/Runtime/Harmonix/Source/HarmonixDsp/Public/HarmonixDsp/FusionSampler/Settings/KeyzoneSettings.h
@@ -11,12 +11,9 @@
 
 #include "KeyzoneSettings.generated.h"
 
-namespace HarmonixDsp
-{
-	class IAudioData;
-}
 class FSingletonFusionVoicePool;
 class USoundWave;
+class FSoundWaveProxy;
 
 DECLARE_LOG_CATEGORY_EXTERN(LogKeyzoneSettings, Log, All);
 
@@ -53,6 +50,8 @@ struct HARMONIXDSP_API FKeyzoneSettings
 		ApplyArgs(InArgs);
 	}
 
+	~FKeyzoneSettings();
+
 	void ApplyArgs(const FKeyzoneArgs& InArgs)
 	{
 		RootNote = InArgs.RootNote;
@@ -78,8 +77,7 @@ public:
 	UPROPERTY(EditDefaultsOnly, Category = "Audio Sample")
 	TObjectPtr<USoundWave> SoundWave;
 
-	// Not null when this keyzone is being proxied
-	TSharedPtr<HarmonixDsp::IAudioData, ESPMode::ThreadSafe> AudioSample;
+	TSharedPtr<FSoundWaveProxy> SoundWaveProxy;
 
 	UPROPERTY(EditDefaultsOnly, Category = "Settings", Meta = (UIMin=0, UIMax=127, ClampMin=0, ClampMax=127))
 	int8 RootNote = 60;
@@ -158,8 +156,4 @@ public:
 	FString SamplePath;
 #endif
 
-protected:
-
-	TSharedPtr<class FStreamingAudioData, ESPMode::ThreadSafe> CreateStreamingAudioData(const Audio::FProxyDataInitParams& InitParams) const;
-
 };
\ No newline at end of file
diff --git a/Engine/Plugins/Runtime/Harmonix/Source/HarmonixDsp/Public/HarmonixDsp/StretcherAndPitchShifter.h b/Engine/Plugins/Runtime/Harmonix/Source/HarmonixDsp/Public/HarmonixDsp/StretcherAndPitchShifter.h
--- a/Engine/Plugins/Runtime/Harmonix/Source/HarmonixDsp/Public/HarmonixDsp/StretcherAndPitchShifter.h
+++ b/Engine/Plugins/Runtime/Harmonix/Source/HarmonixDsp/Public/HarmonixDsp/StretcherAndPitchShifter.h
@@ -9,13 +9,9 @@
 
 #include "StretcherAndPitchShifter.generated.h"
 
-namespace HarmonixDsp
-{
-	class IAudioData;
-}
-
 class IAudioDataRenderer;
 class FGainMatrix;
+class FSoundWaveProxy;
 
 USTRUCT()
 struct HARMONIXDSP_API FTimeStretchConfig
@@ -57,9 +53,9 @@ public:
 	virtual bool HasCurrentSampleFrame() const { return false; }
 	virtual int32  GetCurrentSampleFrame() const { return 0; }
 
-	void SetSampleSourceReset(TSharedPtr<const HarmonixDsp::IAudioData, ESPMode::ThreadSafe> InAudioData, TSharedPtr<IAudioDataRenderer, ESPMode::ThreadSafe> InAudioRenderer) 
+	void SetSampleSourceReset(TSharedPtr<FSoundWaveProxy> InSoundWave, TSharedPtr<IAudioDataRenderer, ESPMode::ThreadSafe> InAudioRenderer) 
 	{ 
-		MyAudioData = InAudioData;
+		MyAudioData = InSoundWave;
 		MyAudioRenderer = InAudioRenderer;
 
 		SetupAndResetImpl(); 
@@ -78,7 +74,7 @@ public:
 
 	virtual void Cleanup() = 0;
 protected:
-	TSharedPtr<const HarmonixDsp::IAudioData, ESPMode::ThreadSafe> MyAudioData;
+	TSharedPtr<FSoundWaveProxy> MyAudioData;
 	TSharedPtr<IAudioDataRenderer, ESPMode::ThreadSafe> MyAudioRenderer;
 
 	virtual void SetupAndResetImpl() { }
