Fix Incremental Build Minimal: Compile errors in Module.HarmonixDsp.2.cpp and FusionVoice.cpp.

checkSlow doesn't always compile. Should run better preflights to catch this

diff --git a/Engine/Plugins/Runtime/Harmonix/Source/HarmonixDsp/Private/HarmonixDsp/FusionSampler/FusionVoice.cpp b/Engine/Plugins/Runtime/Harmonix/Source/HarmonixDsp/Private/HarmonixDsp/FusionSampler/FusionVoice.cpp
--- a/Engine/Plugins/Runtime/Harmonix/Source/HarmonixDsp/Private/HarmonixDsp/FusionSampler/FusionVoice.cpp
+++ b/Engine/Plugins/Runtime/Harmonix/Source/HarmonixDsp/Private/HarmonixDsp/FusionSampler/FusionVoice.cpp
@@ -153,7 +153,7 @@ void FFusionVoice::Attack()
 {
 	check(MySampler);
 	checkSlow(KeyZone);
-	checkSlow(KeyZone->AudioSample);
+	checkSlow(KeyZone->SoundWaveProxy);
 
 	AdsrVolume.Attack();
 	AdsrAssignable.Attack();
@@ -413,7 +413,7 @@ void FFusionVoice::PrepareWithPitchOffsetAndGain(double InPitchOffsetCents, floa
 
 	checkSlow(MySampler);
 	checkSlow(KeyZone);
-	checkSlow(KeyZone->AudioSample);
+	checkSlow(KeyZone->SoundWaveProxy);
 
 	for (int32 Idx = 0; Idx < kNumLfos; ++Idx)
 	{
