Detect when MidiPlayCursor::GetOwner is called without locking the appropriate mutex.
Fix a case detected in the Patchwork Music Gallery.

diff --git a/Engine/Plugins/Runtime/Harmonix/Source/HarmonixMetasound/Private/HarmonixMetasound/Nodes/StepSequencePlayerNode.cpp b/Engine/Plugins/Runtime/Harmonix/Source/HarmonixMetasound/Private/HarmonixMetasound/Nodes/StepSequencePlayerNode.cpp
--- a/Engine/Plugins/Runtime/Harmonix/Source/HarmonixMetasound/Private/HarmonixMetasound/Nodes/StepSequencePlayerNode.cpp
+++ b/Engine/Plugins/Runtime/Harmonix/Source/HarmonixMetasound/Private/HarmonixMetasound/Nodes/StepSequencePlayerNode.cpp
@@ -51,6 +51,8 @@ namespace HarmonixMetasound
 									const FBoolReadRef& InLoop,
 									const FBoolReadRef& InEnabled);
 
+		virtual ~FStepSequencePlayerOperator() override;
+
 		virtual void BindInputs(FInputVertexInterfaceData& InVertexData) override;
 		virtual void BindOutputs(FOutputVertexInterfaceData& InVertexData) override;
 		virtual FDataReferenceCollection GetInputs() const override;
@@ -252,8 +254,15 @@ namespace HarmonixMetasound
 		Init();
 	}
 
+	FStepSequencePlayerOperator::~FStepSequencePlayerOperator()
+	{
+		MidiClockInPin->UnregisterPlayCursor(this);
+	}
+
 	void FStepSequencePlayerOperator::BindInputs(FInputVertexInterfaceData& InVertexData)
 	{
+		MidiClockInPin->UnregisterPlayCursor(this);
+
 		using namespace StepSequencePlayerPinNames;
 		using namespace CommonPinNames;
 		InVertexData.BindReadVertex(METASOUND_GET_PARAM_NAME(InputSequenceAsset), SequenceAssetInPin);
@@ -469,7 +478,7 @@ namespace HarmonixMetasound
 			bNeedsMidiOutClear = false;
 		}
 
-		int32 BlockFrame = (*MidiClockInPin).GetCurrentBlockFrameIndex();
+		int32 BlockFrame = MidiClockInPin->GetCurrentBlockFrameIndex();
 		ProcessedThruTick = FMath::Max(-1, Tick);
 		bNeedsRebase = true;
 		AllNotesOff(BlockFrame, Tick, true);
diff --git a/Engine/Plugins/Runtime/Harmonix/Source/HarmonixMidi/Private/HarmonixMidi/MidiPlayCursor.cpp b/Engine/Plugins/Runtime/Harmonix/Source/HarmonixMidi/Private/HarmonixMidi/MidiPlayCursor.cpp
--- a/Engine/Plugins/Runtime/Harmonix/Source/HarmonixMidi/Private/HarmonixMidi/MidiPlayCursor.cpp
+++ b/Engine/Plugins/Runtime/Harmonix/Source/HarmonixMidi/Private/HarmonixMidi/MidiPlayCursor.cpp
@@ -24,7 +24,7 @@ FMidiPlayCursor::FMidiPlayCursor()
 
 FMidiPlayCursor::~FMidiPlayCursor()
 {
-	if (Owner)
+	if (GetOwner())
 	{
 		Owner->UnregisterPlayCursor(this);
 		Owner = nullptr;
@@ -34,12 +34,12 @@ FMidiPlayCursor::~FMidiPlayCursor()
 
 void FMidiPlayCursor::SetOwner(FMidiPlayCursorMgr* NewOwner, FMidiPlayCursorTracker* NewTracker, float PreRollMs)
 {
-	check(Owner == nullptr || NewOwner == nullptr);
+	check(GetOwner() == nullptr || NewOwner == nullptr);
 
 	Owner = NewOwner;
 	Tracker = NewTracker;
 
-	if (Owner)
+	if (GetOwner())
 	{
 		TrackNextEventIndexs.SetNumUninitialized(Owner->Tracks().Num());
 		Reset(PreRollMs > 0.0f);
@@ -61,7 +61,7 @@ void FMidiPlayCursor::Reset(bool ForceNoBroadcast)
 {
 	CurrentTick = -1;
 	LoopCount = 0;
-	if (Owner)
+	if (GetOwner())
 	{
 		CurrentMs = Owner->GetTempoMap().TickToMs(CurrentTick);
 		TrackNextEventIndexs.SetNumUninitialized(Owner->Tracks().Num());
@@ -94,12 +94,27 @@ void FMidiPlayCursor::SetupTickLookahead(int32 Ticks, ESyncOptions Opts)
 	SyncOpts = Opts;
 
 	// may be called before or after we are attached to a play cursor manager!
-	if (Owner)
+	if (GetOwner())
 	{
 		PrepareLookAheadTicks();
 	}
 }
 
+#if HARMONIX_MIDIPLAYCURSOR_ENABLE_ENSURE_OWNER
+FMidiPlayCursorMgr* FMidiPlayCursor::GetOwner()
+{
+	if (Tracker)
+	{
+		int32 RecursionCount = Tracker->CursorListCS.GetRecursionCountIfOwned();
+		if (!RecursionCount)
+		{
+			UE_LOG(LogMidi, Warning, TEXT("Play cursor owner accessed with cursor list unlocked."));
+		}
+	}
+	return Owner;
+}
+#endif
+
 void FMidiPlayCursor::RecalcNextEventsDueToMidiChanges(FMidiPlayCursorMgr::EMidiChangePositionCorrectMode PositionMode)
 {
 	// MUST ONLY BE CALLED BY THE PLAYCURSORMGR WITH CURSORS LOCKED!
@@ -115,7 +130,7 @@ void FMidiPlayCursor::RecalcNextEventsDueToMidiChanges(FMidiPlayCursorMgr::EMidi
 	}
 
 	// changes to tracks might actually be more complicated...
-	const UMidiFile::FMidiTrackList& Tracks = Owner->Tracks();
+	const UMidiFile::FMidiTrackList& Tracks = GetOwner()->Tracks();
 	if (Tracks.Num() != TrackNextEventIndexs.Num())
 	{
 		int32 OldNum = TrackNextEventIndexs.Num();
@@ -149,7 +164,7 @@ void FMidiPlayCursor::RecalcNextEventsDueToMidiChanges(FMidiPlayCursorMgr::EMidi
 
 void FMidiPlayCursor::PrepareLookAheadTicks(bool ForceNoBroadcast)
 {
-	check(Owner);
+	check(GetOwner());
 
 	bool Broadcast = !ForceNoBroadcast &&
 		(SyncOpts != ESyncOptions::NoBroadcastNoPreRoll) &&
@@ -164,7 +179,7 @@ void FMidiPlayCursor::PrepareLookAheadTicks(bool ForceNoBroadcast)
 		Tracker->IsAtStart() &&
 		(SyncOpts == ESyncOptions::PreRollIfAtStartOrBroadcast || SyncOpts == ESyncOptions::PreRollIfAtStartOrNoBroadcast))
 	{
-		Owner->RecalculatePreRollDueToCursorPosition(this);
+		GetOwner()->RecalculatePreRollDueToCursorPosition(this);
 		return;
 	}
 
@@ -185,7 +200,7 @@ void FMidiPlayCursor::SetupMsLookahead(float Ms, ESyncOptions Opts)
 	SyncOpts = Opts;
 
 	// may be called before or after we are attached to a play cursor manager!
-	if (Owner)
+	if (GetOwner())
 	{
 		PrepareLookAheadMs();
 	}
@@ -193,7 +208,7 @@ void FMidiPlayCursor::SetupMsLookahead(float Ms, ESyncOptions Opts)
 
 void FMidiPlayCursor::PrepareLookAheadMs(bool ForceNoBroadcast)
 {
-	check(Owner);
+	check(GetOwner());
 
 	bool Broadcast = !ForceNoBroadcast &&
 		(SyncOpts != ESyncOptions::NoBroadcastNoPreRoll) &&
@@ -208,7 +223,7 @@ void FMidiPlayCursor::PrepareLookAheadMs(bool ForceNoBroadcast)
 		Tracker->IsAtStart() &&
 		(SyncOpts == ESyncOptions::PreRollIfAtStartOrBroadcast || SyncOpts == ESyncOptions::PreRollIfAtStartOrNoBroadcast))
 	{
-		Owner->RecalculatePreRollDueToCursorPosition(this);
+		GetOwner()->RecalculatePreRollDueToCursorPosition(this);
 		return;
 	}
 
@@ -221,11 +236,11 @@ void FMidiPlayCursor::PrepareLookAheadMs(bool ForceNoBroadcast)
 		float NewMs = Tracker->CurrentMs + LookaheadMs;
 		if (NewMs > 0.0f)
 		{
-			CurrentTick = (int32)(Owner->GetTempoMap().MsToTick(NewMs) + 0.5f);
+			CurrentTick = (int32)(GetOwner()->GetTempoMap().MsToTick(NewMs) + 0.5f);
 		}
 		else
 		{
-			CurrentTick = (int32)(Owner->GetTempoMap().MsToTick(NewMs) - 0.5f);
+			CurrentTick = (int32)(GetOwner()->GetTempoMap().MsToTick(NewMs) - 0.5f);
 		}
 		SeekThruTick(CurrentTick);
 		CurrentMs = NewMs;
@@ -234,9 +249,9 @@ void FMidiPlayCursor::PrepareLookAheadMs(bool ForceNoBroadcast)
 
 bool FMidiPlayCursor::Advance(bool IsLowRes)
 {
-	check(Owner);
+	check(GetOwner());
 	bool ProcessLoops = IsLowRes ||
-		!(Owner->IsDirectMappedTimeFollower() && LookaheadTicks == 0);
+		!(GetOwner()->IsDirectMappedTimeFollower() && LookaheadTicks == 0);
 	switch (LookaheadType)
 	{
 	case ELookaheadType::Ticks: AdvanceByTicks(ProcessLoops); break;
@@ -248,7 +263,7 @@ bool FMidiPlayCursor::Advance(bool IsLowRes)
 
 bool FMidiPlayCursor::AdvanceAsPreRoll()
 {
-	check(Owner);
+	check(GetOwner());
 	switch (LookaheadType)
 	{
 	case ELookaheadType::Ticks: AdvanceByTicks(true, true, true); break;
@@ -274,7 +289,7 @@ void FMidiPlayCursor::AdvanceByTicks(bool ProcessLoops, bool Broadcast, bool IsP
 void FMidiPlayCursor::DoAdvanceForLaggingTickCursor(bool Broadcast, bool IsPreRoll)
 {
 	int32 NewTick = Tracker->CurrentTick + LookaheadTicks;
-	if (Owner->DoesLoop(Tracker->IsLowRes))
+	if (GetOwner()->DoesLoop(Tracker->IsLowRes))
 	{
 		int32 LoopStartTick = Owner->GetLoopStartTick(Tracker->IsLowRes);
 		int32 LoopEndTick = Owner->GetLoopEndTick(Tracker->IsLowRes);
@@ -402,7 +417,7 @@ void FMidiPlayCursor::DoAdvanceForLaggingTickCursor(bool Broadcast, bool IsPreRo
 void FMidiPlayCursor::DoAdvanceForLeadingTickCursor(bool Broadcast, bool ProcessLoops, bool IsPreRoll)
 {
 	int32 NewTick = Tracker->CurrentTick + LookaheadTicks;
-	if (Owner->DoesLoop(Tracker->IsLowRes) && ProcessLoops) 
+	if (GetOwner()->DoesLoop(Tracker->IsLowRes) && ProcessLoops) 
 	{
 		int32 LoopStartTick = Owner->GetLoopStartTick(Tracker->IsLowRes);
 		int32 LoopEndTick = Owner->GetLoopEndTick(Tracker->IsLowRes);
@@ -528,7 +543,7 @@ void FMidiPlayCursor::AdvanceByMs(bool ProcessLoops, bool Broadcast, bool IsPreR
 void FMidiPlayCursor::DoAdvanceForLaggingMsCursor(bool Broadcast, bool IsPreRoll)
 {
 	float NewMs = Tracker->CurrentMs + LookaheadMs;
-	int32 NewTick = (int32)(Owner->GetTempoMap().MsToTick(NewMs) + 0.5f);
+	int32 NewTick = (int32)(GetOwner()->GetTempoMap().MsToTick(NewMs) + 0.5f);
 	if (Owner->DoesLoop(Tracker->IsLowRes))
 	{
 		float LoopStartMs = Owner->GetLoopStartMs(Tracker->IsLowRes);
@@ -662,7 +677,7 @@ void FMidiPlayCursor::DoAdvanceForLaggingMsCursor(bool Broadcast, bool IsPreRoll
 void FMidiPlayCursor::DoAdvanceForLeadingMsCursor(bool Broadcast, bool ProcessLoops, bool IsPreRoll)
 {
 	float NewMs = Tracker->CurrentMs + LookaheadMs;
-	int32 NewTick = (int32)(Owner->GetTempoMap().MsToTick(NewMs) + 0.5f);
+	int32 NewTick = (int32)(GetOwner()->GetTempoMap().MsToTick(NewMs) + 0.5f);
 	if (Owner->DoesLoop(Tracker->IsLowRes) && ProcessLoops) 
 	{
 		float LoopStartMs = Owner->GetLoopStartMs(Tracker->IsLowRes);
@@ -781,7 +796,7 @@ void FMidiPlayCursor::DoAdvanceForLeadingMsCursor(bool Broadcast, bool ProcessLo
 
 void FMidiPlayCursor::SeekToTick(int32 Tick)
 {
-	const UMidiFile::FMidiTrackList& Tracks = Owner->Tracks();
+	const UMidiFile::FMidiTrackList& Tracks = GetOwner()->Tracks();
 	int32 StartIndex = (WatchTrack < 0) ? 0 : WatchTrack;
 	int32 EndIndex = (WatchTrack < 0) ? Tracks.Num() : (WatchTrack + 1);
 	if (EndIndex > Tracks.Num())
@@ -808,7 +823,7 @@ void FMidiPlayCursor::SeekToTick(int32 Tick)
 
 void FMidiPlayCursor::SeekThruTick(int32 Tick)
 {
-	const UMidiFile::FMidiTrackList& Tracks = Owner->Tracks();
+	const UMidiFile::FMidiTrackList& Tracks = GetOwner()->Tracks();
 	int32 StartIndex = (WatchTrack < 0) ? 0 : WatchTrack;
 	int32 EndIndex = (WatchTrack < 0) ? Tracks.Num() : (WatchTrack + 1);
 	if (EndIndex > Tracks.Num())
@@ -835,7 +850,7 @@ void FMidiPlayCursor::SeekThruTick(int32 Tick)
 
 bool FMidiPlayCursor::IsDone() const
 {
-	const UMidiFile::FMidiTrackList& Tracks = Owner->Tracks();
+	const UMidiFile::FMidiTrackList& Tracks = GetOwner()->Tracks();
 	int32 StartIndex = (WatchTrack < 0) ? 0 : WatchTrack;
 	int32 EndIndex = (WatchTrack < 0) ? Tracks.Num() : (WatchTrack + 1);
 	if (EndIndex > Tracks.Num())
@@ -1058,7 +1073,7 @@ namespace
 
 void FMidiPlayCursor::AdvanceThruTick(int32 Tick, bool IsPreRoll)
 {
-	const UMidiFile::FMidiTrackList& Tracks = Owner->Tracks();
+	const UMidiFile::FMidiTrackList& Tracks = GetOwner()->Tracks();
 	FMidiEventHashTable& HeldNoteOnEvents = ThisThreadsMidiHashTable.Get();
 
 	if (!IsPreRoll)
@@ -1182,7 +1197,7 @@ void FMidiPlayCursor::BroadcastEvent(int32 TrackIndex, const FMidiEvent& Event,
 	case FMidiMsg::EType::Text:
 		if ((FilterPassFlags & EFilterPassFlags::Text) != EFilterPassFlags::None)
 		{
-			const FMidiTrack& Track = Owner->Tracks()[TrackIndex];
+			const FMidiTrack& Track = GetOwner()->Tracks()[TrackIndex];
 			OnText(TrackIndex, Event.GetTick(), Msg.GetTextIndex(), Track.GetTextAtIndex(Msg.GetTextIndex()), Msg.GetTextType(), IsPreroll);
 		}
 		break;
@@ -1193,7 +1208,7 @@ void FMidiPlayCursor::BroadcastEvent(int32 TrackIndex, const FMidiEvent& Event,
 		}
 		break;
 	default:
-		UE_LOG(LogMIDI, Error, TEXT("Unknown midi message type %d on track %d at tick %d, file %s"), int(Msg.MsgType()), TrackIndex, Event.GetTick(), **Owner->GetMidiFileName());
+		UE_LOG(LogMIDI, Error, TEXT("Unknown MIDI message type %d on track %d at tick %d, file %s"), int(Msg.MsgType()), TrackIndex, Event.GetTick(), **GetOwner()->GetMidiFileName());
 		break;
 	}
 }
diff --git a/Engine/Plugins/Runtime/Harmonix/Source/HarmonixMidi/Private/HarmonixMidi/MidiPlayCursorMgr.cpp b/Engine/Plugins/Runtime/Harmonix/Source/HarmonixMidi/Private/HarmonixMidi/MidiPlayCursorMgr.cpp
--- a/Engine/Plugins/Runtime/Harmonix/Source/HarmonixMidi/Private/HarmonixMidi/MidiPlayCursorMgr.cpp
+++ b/Engine/Plugins/Runtime/Harmonix/Source/HarmonixMidi/Private/HarmonixMidi/MidiPlayCursorMgr.cpp
@@ -24,8 +24,8 @@ FMidiPlayCursorMgr::FMidiPlayCursorMgr()
 
 void FMidiPlayCursorMgr::Reset()
 {
-	FScopeLock LowResLock(&GetLowResTracker().CursorListLock);
-	FScopeLock HiResLock(&GetHiResTracker().CursorListLock);
+	FMidiPlayCursorListLock LowResCursorListLock(GetLowResTracker().CursorListCS);
+	FMidiPlayCursorListLock HiResCursorListLock(GetHiResTracker().CursorListCS);
 	UnregisterAllPlayCursors();
 	ResetTrackers();
 	MidiFileData = nullptr;
@@ -120,8 +120,8 @@ FMidiSongPos FMidiPlayCursorMgr::CalculateSongPosWithOffsetMs(float DeltaMs, boo
 
 void FMidiPlayCursorMgr::DetermineLength()
 {
-	FScopeLock LowResLock(&GetLowResTracker().CursorListLock);
-	FScopeLock HiResLock(&GetHiResTracker().CursorListLock);
+	FMidiPlayCursorListLock LowResCursorListLock(GetLowResTracker().CursorListCS);
+	FMidiPlayCursorListLock HiResCursorListLock(GetHiResTracker().CursorListCS);
 	if (!MidiFileData)
 	{
 		LengthMs = 0.0f;
@@ -166,24 +166,24 @@ FMidiPlayCursorMgr::~FMidiPlayCursorMgr()
 
 void FMidiPlayCursorMgr::RegisterHiResPlayCursor(FMidiPlayCursor* PlayCursor, float PreRollMs)
 {
-	if (PlayCursor->Owner)
+	if (PlayCursor->GetOwner())
 	{
 		PlayCursor->Owner->UnregisterPlayCursor(PlayCursor);
 	}
 	PlayCursor->UnregisterASAP = false;
-	FScopeLock HiResLock(&GetHiResTracker().CursorListLock);
+	FMidiPlayCursorListLock HiResCursorListLock(GetHiResTracker().CursorListCS);
 	PlayCursor->SetOwner(this, &GetHiResTracker(), PreRollMs);
 	GetHiResTracker().AddCursor(PlayCursor);
 }
 
 void FMidiPlayCursorMgr::RegisterLowResPlayCursor(FMidiPlayCursor* PlayCursor, float PreRollMs)
 {
-	if (PlayCursor->Owner)
+	if (PlayCursor->GetOwner())
 	{
 		PlayCursor->Owner->UnregisterPlayCursor(PlayCursor);
 	}
 	PlayCursor->UnregisterASAP = false;
-	FScopeLock LowResLock(&GetLowResTracker().CursorListLock);
+	FMidiPlayCursorListLock LowResCursorListLock(GetLowResTracker().CursorListCS);
 	PlayCursor->SetOwner(this, &GetLowResTracker(), PreRollMs);
 	GetLowResTracker().AddCursor(PlayCursor);
 }
@@ -219,7 +219,7 @@ void FMidiPlayCursorMgr::UnregisterPlayCursor(FMidiPlayCursor* PlayCursor, bool
 	// search the hi-res list... then the low res...
 	for (int32 Index = 0; Index != 2; ++Index)
 	{
-		FScopeLock Lock(&Trackers[Index].CursorListLock);
+		FMidiPlayCursorListLock CursorListLock(Trackers[Index].CursorListCS);
 		if (Trackers[Index].TraversingCursors)
 		{
 			if (PlayCursor->Tracker == &Trackers[Index])
@@ -244,8 +244,8 @@ void FMidiPlayCursorMgr::UnregisterPlayCursor(FMidiPlayCursor* PlayCursor, bool
 
 void FMidiPlayCursorMgr::UnregisterAllPlayCursors()
 {
-	FScopeLock LowResLock(&GetLowResTracker().CursorListLock);
-	FScopeLock HiResLock(&GetHiResTracker().CursorListLock);
+	FMidiPlayCursorListLock LowResCursorListLock(GetLowResTracker().CursorListCS);
+	FMidiPlayCursorListLock HiResCursorListLock(GetHiResTracker().CursorListCS);
 
 	if (GetHiResTracker().TraversingCursors || GetLowResTracker().TraversingCursors)
 	{
@@ -278,8 +278,8 @@ bool FMidiPlayCursorMgr::HasLowResCursors() const
 
 void FMidiPlayCursorMgr::ResetTrackers()
 {
-	FScopeLock LowResLock(&GetLowResTracker().CursorListLock);
-	FScopeLock HiResLock(&GetHiResTracker().CursorListLock);
+	FMidiPlayCursorListLock LowResCursorListLock(GetLowResTracker().CursorListCS);
+	FMidiPlayCursorListLock HiResCursorListLock(GetHiResTracker().CursorListCS);
 	float Ms = GetTempoMap().TickToMs(-1);
 	for (FMidiPlayCursorTracker& Tracker : Trackers)
 	{
@@ -350,8 +350,8 @@ void FMidiPlayCursorMgr::GetCursorExtentsTicks(int32& Earliest, int32& Latest) c
 
 bool FMidiPlayCursorMgr::CursorsAllInPhase() const
 {
-	FScopeLock LowResLock(&GetLowResTracker().CursorListLock);
-	FScopeLock HiResLock(&GetHiResTracker().CursorListLock);
+	FMidiPlayCursorListLock LowResCursorListLock(GetLowResTracker().CursorListCS);
+	FMidiPlayCursorListLock HiResCursorListLock(GetHiResTracker().CursorListCS);
 
 	return CursorsAllInPhaseImpl(false) && CursorsAllInPhaseImpl(true);
 }
@@ -410,14 +410,14 @@ void FMidiPlayCursorMgr::SetLoop(int32 StartTick, int32 EndTick, bool IsDirectMa
 	// might fall in the middle of the current span of leading and lagging play cursors,
 	// which would result in ugly behavior!
 
-	FScopeLock HiResLock(&GetHiResTracker().CursorListLock);
+	FMidiPlayCursorListLock HiResCursorListLock(GetHiResTracker().CursorListCS);
 	DirectMappedTimeFollower = IsDirectMappedFollower;
 	SetLoopImpl(StartTick, EndTick, IgnoringLookAhead, false);
 }
 
 void FMidiPlayCursorMgr::ClearLoop(bool IgnoringLookAhead)
 {
-	FScopeLock HiResLock(&GetHiResTracker().CursorListLock);
+	FMidiPlayCursorListLock HiResCursorListLock(GetHiResTracker().CursorListCS);
 	ClearLoopImpl(IgnoringLookAhead, false);
 }
 
@@ -473,8 +473,8 @@ void FMidiPlayCursorMgr::ClearLoopImpl(bool IgnoringLookAhead, bool IsLowRes)
 
 void FMidiPlayCursorMgr::SeekTo(int32 Tick, int32 PreRollBars, bool IsRenderThread, bool IsLoop)
 {
-	FScopeLock LowResLock(&GetLowResTracker().CursorListLock);
-	FScopeLock HiResLock(&GetHiResTracker().CursorListLock);
+	FMidiPlayCursorListLock LowResCursorListLock(GetLowResTracker().CursorListCS);
+	FMidiPlayCursorListLock HiResCursorListLock(GetHiResTracker().CursorListCS);
 
 	// Back up one tick, as we are going to be setting the cursors'
 	// 'played through' position.
@@ -527,7 +527,7 @@ void FMidiPlayCursorMgr::SeekTo(int32 Tick, int32 PreRollBars, bool IsRenderThre
 
 void FMidiPlayCursorMgr::MoveToLoopStart()
 {
-	FScopeLock HiResLock(&GetHiResTracker().CursorListLock);
+	FMidiPlayCursorListLock HiResCursorListLock(GetHiResTracker().CursorListCS);
 	{
 		if (ensureMsgf(GetHiResTracker().Loop, TEXT("That's odd. Asked to move to the beginning of the loop... but there is no loop!")))
 		{
@@ -544,7 +544,7 @@ void FMidiPlayCursorMgr::MoveToLoopStart()
 			// The low res cursor will need to advance to the end of the loop, and then
 			// advance from loop-start to current position at some point in the future
 			// (during a low-res frame poll)
-			FScopeLock LowResLock(&GetLowResTracker().CursorListLock);
+			FMidiPlayCursorListLock LowResCursorListLock(GetLowResTracker().CursorListCS);
 			MsSinceLowResUpdate += GetHiResTracker().LoopEndMs - GetLowResTracker().CurrentMs;
 		}
 	}
@@ -591,8 +591,8 @@ void FMidiPlayCursorMgr::LockForMidiDataChanges()
 	if (ensureAlways(!InMidiChangeLock))
 	{
 		InMidiChangeLock = true;
-		GetLowResTracker().CursorListLock.Lock();
-		GetHiResTracker().CursorListLock.Lock();
+		GetLowResTracker().CursorListCS.Lock();
+		GetHiResTracker().CursorListCS.Lock();
 	}
 }
 
@@ -673,8 +673,8 @@ void FMidiPlayCursorMgr::MidiDataChangeComplete(EMidiChangePositionCorrectMode P
 			int32 SeekTick = PositionMode == EMidiChangePositionCorrectMode::MaintainTick ? GetHiResTracker().CurrentTick : GetTempoMap().MsToTick(GetHiResTracker().CurrentMs);
 			SeekTo(SeekTick, PreRollBars, true, false);
 		}
-		GetHiResTracker().CursorListLock.Unlock();
-		GetLowResTracker().CursorListLock.Unlock();
+		GetHiResTracker().CursorListCS.Unlock();
+		GetLowResTracker().CursorListCS.Unlock();
 		InMidiChangeLock = false;
 	}
 }
@@ -697,14 +697,14 @@ void FMidiPlayCursorMgr::AdvanceHiResToMs(float Ms, bool Broadcast)
 
 void FMidiPlayCursorMgr::AdvanceHiResByDeltaMs(float Ms, bool Broadcast)
 {
-	FScopeLock HiResLock(&GetHiResTracker().CursorListLock);
+	FMidiPlayCursorListLock HiResCursorListLock(GetHiResTracker().CursorListCS);
 	HiResLoopedSinceLastLoResUpdate = AdvanceTrackerByDeltaMs(Ms, GetHiResTracker(), Broadcast) || HiResLoopedSinceLastLoResUpdate;
 	MsSinceLowResUpdate += Ms;
 }
 
 void FMidiPlayCursorMgr::AdvanceHiResByDeltaTick(int32 NumTicks, bool Broadcast)
 {
-	FScopeLock HiResLock(&GetHiResTracker().CursorListLock);
+	FMidiPlayCursorListLock HiResCursorListLock(GetHiResTracker().CursorListCS);
 
 	int32 ThruTick = NumTicks + GetHiResTracker().CurrentTick;
 
@@ -715,7 +715,7 @@ void FMidiPlayCursorMgr::AdvanceHiResByDeltaTick(int32 NumTicks, bool Broadcast)
 
 void FMidiPlayCursorMgr::AdvanceHiResThruTick(int32 ThruTick, bool Broadcast, bool DontAdvancePastLoopEnd)
 {
-	FScopeLock HiResLock(&GetHiResTracker().CursorListLock);
+	FMidiPlayCursorListLock HiResCursorListLock(GetHiResTracker().CursorListCS);
 
 	// early out?...
 	if (ThruTick == GetHiResTracker().CurrentTick)
@@ -741,7 +741,7 @@ void FMidiPlayCursorMgr::AdvanceHiResThruTick(int32 ThruTick, bool Broadcast, bo
 
 bool FMidiPlayCursorMgr::AdvanceLowResCursors()
 {
-	FScopeLock LowResLock(&GetLowResTracker().CursorListLock);
+	FMidiPlayCursorListLock LowResCursorListLock(GetLowResTracker().CursorListCS);
 
 	GetLowResTracker().HandleQueuedReset();
 
@@ -753,7 +753,7 @@ bool FMidiPlayCursorMgr::AdvanceLowResCursors()
 	bool HiResLoop = false;
 	bool HiResLoopIgnoringLookAhead = false;
 	{
-		FScopeLock HiResLock(&GetHiResTracker().CursorListLock);
+		FMidiPlayCursorListLock HiResCursorListLock(GetHiResTracker().CursorListCS);
 
 		MsDiff = MsSinceLowResUpdate;
 		HiResCurrentMs = GetHiResTracker().CurrentMs;
diff --git a/Engine/Plugins/Runtime/Harmonix/Source/HarmonixMidi/Public/HarmonixMidi/MidiPlayCursor.h b/Engine/Plugins/Runtime/Harmonix/Source/HarmonixMidi/Public/HarmonixMidi/MidiPlayCursor.h
--- a/Engine/Plugins/Runtime/Harmonix/Source/HarmonixMidi/Public/HarmonixMidi/MidiPlayCursor.h
+++ b/Engine/Plugins/Runtime/Harmonix/Source/HarmonixMidi/Public/HarmonixMidi/MidiPlayCursor.h
@@ -72,8 +72,12 @@ public:
 	void SetWatchTrack(int32 TrackIndex) { WatchTrack = TrackIndex; }
 	virtual void ManagerIsDetaching() {}
 
-	const FMidiPlayCursorMgr* GetOwner() const { return Owner; }
+#if HARMONIX_MIDIPLAYCURSOR_ENABLE_ENSURE_OWNER
+	FMidiPlayCursorMgr* GetOwner();
+#else
 	FMidiPlayCursorMgr* GetOwner() { return Owner; }
+#endif
+	const FMidiPlayCursorMgr* GetOwner() const { return const_cast<FMidiPlayCursor*>(this)->GetOwner(); }
 
 	static const float kSmallMs;
 
diff --git a/Engine/Plugins/Runtime/Harmonix/Source/HarmonixMidi/Public/HarmonixMidi/MidiPlayCursorTracker.h b/Engine/Plugins/Runtime/Harmonix/Source/HarmonixMidi/Public/HarmonixMidi/MidiPlayCursorTracker.h
--- a/Engine/Plugins/Runtime/Harmonix/Source/HarmonixMidi/Public/HarmonixMidi/MidiPlayCursorTracker.h
+++ b/Engine/Plugins/Runtime/Harmonix/Source/HarmonixMidi/Public/HarmonixMidi/MidiPlayCursorTracker.h
@@ -6,9 +6,62 @@
 
 class FMidiPlayCursor;
 
+#ifndef HARMONIX_MIDIPLAYCURSOR_ENABLE_ENSURE_OWNER
+#define HARMONIX_MIDIPLAYCURSOR_ENABLE_ENSURE_OWNER DO_ENSURE
+#endif
+
+#if HARMONIX_MIDIPLAYCURSOR_ENABLE_ENSURE_OWNER
+class FHarmonixMidiPlayCursorDebugCriticalSection
+{
+public:
+	FORCEINLINE void Lock()
+	{
+		CriticalSection.Lock();
+		++RecursionCount;
+	}
+
+	FORCEINLINE bool TryLock()
+	{
+		if (CriticalSection.TryLock())
+		{
+			++RecursionCount;
+			return true;
+		}
+		return false;
+	}
+
+	FORCEINLINE void Unlock()
+	{
+		--RecursionCount;
+		CriticalSection.Unlock();
+	}
+
+	FORCEINLINE int32 GetRecursionCountIfOwned()
+	{
+		if (CriticalSection.TryLock())
+		{
+			int32 Count = RecursionCount;
+			CriticalSection.Unlock();
+			return Count;
+		}
+		return 0;
+	}
+
+private:
+	FCriticalSection CriticalSection;
+	int32 RecursionCount = 0;
+};
+
+using FMidiPlayCursorListCS = FHarmonixMidiPlayCursorDebugCriticalSection;
+#else
+using FMidiPlayCursorListCS = FCriticalSection;
+#endif
+
+using FMidiPlayCursorListLock = UE::TScopeLock<FMidiPlayCursorListCS>;
+
 struct FMidiPlayCursorTracker
 {
-	mutable FCriticalSection CursorListLock;
+	mutable FMidiPlayCursorListCS CursorListCS;
 
 	int32   CurrentTick; // We've broadcast all events up through this tick
 	float   CurrentMs;   // We've broadcast all events up through this Ms
