Clean up MidiClockEvent and make sure it's setting the values correctly...

#tests HarmonixTests

diff --git a/Engine/Plugins/Runtime/Harmonix/Source/HarmonixMetasound/Private/HarmonixMetasound/DataTypes/MidiClockEvent.cpp b/Engine/Plugins/Runtime/Harmonix/Source/HarmonixMetasound/Private/HarmonixMetasound/DataTypes/MidiClockEvent.cpp
--- a/Engine/Plugins/Runtime/Harmonix/Source/HarmonixMetasound/Private/HarmonixMetasound/DataTypes/MidiClockEvent.cpp
+++ b/Engine/Plugins/Runtime/Harmonix/Source/HarmonixMetasound/Private/HarmonixMetasound/DataTypes/MidiClockEvent.cpp
@@ -13,22 +13,22 @@ namespace HarmonixMetasound
 
 	FMidiClockEvent FMidiClockEvent::MakeLoopEvent(int32 InBlockFrameIndex, int32 InStartTick, int32 InEndTick)
 	{
-		return FMidiClockEvent(EType::Loop, InBlockFrameIndex, InStartTick, InEndTick, false, true);
+		return FMidiClockEvent(EType::Loop, InBlockFrameIndex, InStartTick, InEndTick);
 	}
 
 	FMidiClockEvent FMidiClockEvent::MakeSeekToEvent(int32 InBlockFrameIndex, int32 FromTick, int32 ToTick)
 	{
-		return FMidiClockEvent(EType::SeekTo, InBlockFrameIndex, FromTick, ToTick, false, true);
+		return FMidiClockEvent(EType::SeekTo, InBlockFrameIndex, FromTick, ToTick);
 	}
 
 	FMidiClockEvent FMidiClockEvent::MakeSeekThruEvent(int32 InBlockFrameIndex, int32 FromTick, int32 ThruTick)
 	{
-		return FMidiClockEvent(EType::SeekThru, InBlockFrameIndex, FromTick, ThruTick, false, true);
+		return FMidiClockEvent(EType::SeekThru, InBlockFrameIndex, FromTick, ThruTick);
 	}
 
 	FMidiClockEvent FMidiClockEvent::MakeAdvanceThruEvent(int32 InBlockFrameIndex, int32 FromTick, int32 ThruTick, bool IsPreRoll)
 	{
-		return FMidiClockEvent(EType::AdvanceThru, InBlockFrameIndex, FromTick, ThruTick, IsPreRoll, true);
+		return FMidiClockEvent(EType::AdvanceThru, InBlockFrameIndex, FromTick, ThruTick, IsPreRoll);
 	}
 
 	FMidiClockEvent::FMidiClockEvent(EType InType, int32 InBlockFrameIndex, int32 InTick1, int32 InTick2, bool InIsPreRoll, bool InForceNoBroadcast)
diff --git a/Engine/Plugins/Runtime/Harmonix/Source/HarmonixMetasound/Public/HarmonixMetasound/DataTypes/MidiClockEvent.h b/Engine/Plugins/Runtime/Harmonix/Source/HarmonixMetasound/Public/HarmonixMetasound/DataTypes/MidiClockEvent.h
--- a/Engine/Plugins/Runtime/Harmonix/Source/HarmonixMetasound/Public/HarmonixMetasound/DataTypes/MidiClockEvent.h
+++ b/Engine/Plugins/Runtime/Harmonix/Source/HarmonixMetasound/Public/HarmonixMetasound/DataTypes/MidiClockEvent.h
@@ -18,7 +18,11 @@ namespace HarmonixMetasound
 		};
 
 		const EType Type;
+
+		// for "advance thru" events
 		const bool  IsPreRoll;
+
+		// for "reset" events
 		const bool  ForceNoBroadcast;
 		const int32 BlockFrameIndex;
 
@@ -39,7 +43,7 @@ namespace HarmonixMetasound
 
 	private:
 
-		FMidiClockEvent(EType InType, int32 InBlockFrameIndex, int32 InTick1, int32 InTick2, bool InIsPreRoll, bool InBroadcastEvents);
+		FMidiClockEvent(EType InType, int32 InBlockFrameIndex, int32 InTick1, int32 InTick2, bool InIsPreRoll = false, bool InForceNoBroadcast = false);
 	};
 
 };
