Fix LogMIDI capitalization.

diff --git a/Engine/Plugins/Runtime/Harmonix/Source/HarmonixMidi/Private/HarmonixMidi/MidiPlayCursor.cpp b/Engine/Plugins/Runtime/Harmonix/Source/HarmonixMidi/Private/HarmonixMidi/MidiPlayCursor.cpp
--- a/Engine/Plugins/Runtime/Harmonix/Source/HarmonixMidi/Private/HarmonixMidi/MidiPlayCursor.cpp
+++ b/Engine/Plugins/Runtime/Harmonix/Source/HarmonixMidi/Private/HarmonixMidi/MidiPlayCursor.cpp
@@ -108,7 +108,7 @@ FMidiPlayCursorMgr* FMidiPlayCursor::GetOwner()
 		int32 RecursionCount = Tracker->CursorListCS.GetRecursionCountIfOwned();
 		if (!RecursionCount)
 		{
-			UE_LOG(LogMidi, Warning, TEXT("Play cursor owner accessed with cursor list unlocked."));
+			UE_LOG(LogMIDI, Warning, TEXT("Play cursor owner accessed with cursor list unlocked."));
 		}
 	}
 	return Owner;
