Fix non-unity build.

diff --git a/Engine/Plugins/Runtime/Harmonix/Source/HarmonixMidi/Public/HarmonixMidi/MidiPlayCursorTracker.h b/Engine/Plugins/Runtime/Harmonix/Source/HarmonixMidi/Public/HarmonixMidi/MidiPlayCursorTracker.h
--- a/Engine/Plugins/Runtime/Harmonix/Source/HarmonixMidi/Public/HarmonixMidi/MidiPlayCursorTracker.h
+++ b/Engine/Plugins/Runtime/Harmonix/Source/HarmonixMidi/Public/HarmonixMidi/MidiPlayCursorTracker.h
@@ -2,7 +2,7 @@
 #pragma once
 
 #include "Containers/IntrusiveDoubleLinkedList.h"
-#include "HAL/CriticalSection.h"
+#include "Misc/ScopeLock.h"
 
 class FMidiPlayCursor;
 
