[Harmonix] Handle MIDI stream disconnects (individual voice version)

#tests automated tests, PIE

#changelist validated
#virtualized

diff --git a/Engine/Plugins/Runtime/Harmonix/Source/HarmonixMetasound/Private/HarmonixMetasound/DataTypes/MidiStream.cpp b/Engine/Plugins/Runtime/Harmonix/Source/HarmonixMetasound/Private/HarmonixMetasound/DataTypes/MidiStream.cpp
--- a/Engine/Plugins/Runtime/Harmonix/Source/HarmonixMetasound/Private/HarmonixMetasound/DataTypes/MidiStream.cpp
+++ b/Engine/Plugins/Runtime/Harmonix/Source/HarmonixMetasound/Private/HarmonixMetasound/DataTypes/MidiStream.cpp
@@ -15,6 +15,17 @@ namespace HarmonixMetasound
 
 	using namespace Metasound;
 
+	FMidiStreamEvent::FMidiStreamEvent(const FMidiVoiceGeneratorBase* Owner, const FMidiMsg& Message)
+	: FMidiStreamEvent(Owner != nullptr ? Owner->GetIdBits() : static_cast<uint32>(0), Message)
+	{
+	}
+
+	FMidiStreamEvent::FMidiStreamEvent(const uint32 OwnerId, const FMidiMsg& Message)
+	: MidiMessage(Message)
+	, VoiceId(OwnerId, MidiMessage)
+	{
+	}
+
 	void FMidiStream::SetClock(const FMidiClock& InClock)
 	{
 		Clock = InClock.AsWeak();
@@ -40,16 +51,14 @@ namespace HarmonixMetasound
 	{
 		check (EventsInBlock.IsEmpty() || EventsInBlock.Last().BlockSampleFrameIndex <= Event.BlockSampleFrameIndex);
 		EventsInBlock.Add(Event);
-
-		UpdateActiveVoice(Event);
+		TrackNote(Event);
 	}
 
 	void FMidiStream::InsertMidiEvent(const FMidiStreamEvent& Event)
 	{
 		int32 AtIndex = Algo::UpperBound(EventsInBlock, Event, [](const FMidiStreamEvent& NewEvent, const FMidiStreamEvent& ExistingEvent){ return NewEvent.BlockSampleFrameIndex < ExistingEvent.BlockSampleFrameIndex; });
 		EventsInBlock.Insert(Event, AtIndex);
-
-		UpdateActiveVoice(Event);
+		TrackNote(Event);
 	}
 
 	void FMidiStream::AddNoteOffEventOrCancelPendingNoteOn(const FMidiStreamEvent& Event)
@@ -106,9 +115,11 @@ namespace HarmonixMetasound
 			To.SetClock(*FromClock);
 		}
 
-		// Reset the target
+		// Reset the current events
 		To.EventsInBlock.Reset();
-		To.ActiveVoices = From.ActiveVoices;
+		
+		// Copy the active notes from the other stream
+		To.ActiveNotes = From.ActiveNotes;
 		
 		// Copy the events
 		for (const FMidiStreamEvent& Event : From.GetEventsInBlock())
@@ -148,6 +159,20 @@ namespace HarmonixMetasound
 				To.InsertMidiEvent(Transformer(Event));
 			}
 		}
+
+		// Merge in the active notes from the other stream
+		for (const FMidiStreamEvent& ActiveNote : From.ActiveNotes)
+		{
+			check(ActiveNote.MidiMessage.IsNoteOn());
+			
+			if (!To.ActiveNotes.ContainsByPredicate([&ActiveNote](const FMidiStreamEvent& TrackedEvent)
+			{
+				return ActiveNote.TrackIndex == TrackedEvent.TrackIndex && ActiveNote.VoiceId == TrackedEvent.GetVoiceId();
+			}))
+			{
+				To.ActiveNotes.Add(ActiveNote);
+			}
+		}
 	}
 	
 	void FMidiStream::Merge(
@@ -164,57 +189,57 @@ namespace HarmonixMetasound
 		const auto RemapTransform = [&To, &Transformer](const FMidiStreamEvent& Event)
 		{
 			FMidiStreamEvent TransformedEvent = Transformer(Event);
-			const auto GeneratorId = TransformedEvent.GetVoiceId().GetGeneratorId();
-			TransformedEvent.ReassignOwner(&To.GeneratorMap.FindOrAdd(GeneratorId));
+			const uint32 GeneratorId = TransformedEvent.GetVoiceId().GetGeneratorId();
+			TransformedEvent.VoiceId.ReassignGenerator(To.GeneratorMap.FindOrAdd(GeneratorId).GetIdBits());
 			return TransformedEvent;
 		};
 		Merge(FromB, To, Filter, RemapTransform);
 	}
 
-	void FMidiStream::UpdateActiveVoice(const FMidiStreamEvent& Event)
+	bool FMidiStream::NoteIsActive(const FMidiStreamEvent& Event) const
 	{
-		if (Event.MidiMessage.IsNoteMessage())
+		return ActiveNotes.ContainsByPredicate([&Event](const FMidiStreamEvent& ActiveNote)
 		{
-			if (Event.MidiMessage.IsNoteOn())
-			{
-				if (FMidiVoiceId VoiceId = Event.GetVoiceId(); !ActiveVoices.Contains(VoiceId))
-				{
-					ActiveVoices.Emplace(MoveTemp(VoiceId));
-				}
-			}
-			else if (Event.MidiMessage.IsNoteOff())
-			{
-				ActiveVoices.Remove(Event.GetVoiceId());
-			}
-			// Kill all or all off
-			else
-			{
-				ActiveVoices.Reset();
-			}
-		}
+			return Event.TrackIndex == ActiveNote.TrackIndex && Event.GetVoiceId() == ActiveNote.GetVoiceId();
+		});
 	}
 
-	void FMidiVoiceTracker::Process(const FMidiStream& MidiStream, const FKillVoiceFn& KillVoiceFn)
+	void FMidiStream::TrackNote(const FMidiStreamEvent& Event)
 	{
-		// Remove active voices that have note offs this block
-		for (const auto& MidiEvent : MidiStream.EventsInBlock)
+		if (!Event.MidiMessage.IsNoteMessage())
 		{
-			if (MidiEvent.MidiMessage.IsNoteOff())
-			{
-				ActiveVoices.Remove(MidiEvent.GetVoiceId());
-			}
+			return;
 		}
 
-		// Check for voices that are no longer active and notify the caller
-		for (auto It = ActiveVoices.CreateIterator(); It; ++It)
+		if (Event.MidiMessage.IsNoteOn())
 		{
-			if (!MidiStream.ActiveVoices.Contains(*It))
+			const int32 TrackIndex = Event.TrackIndex;
+			const FMidiVoiceId VoiceId = Event.GetVoiceId();
+
+			if (!ActiveNotes.ContainsByPredicate([TrackIndex, VoiceId](const HarmonixMetasound::FMidiStreamEvent& TrackedEvent)
+			{
+				return TrackIndex == TrackedEvent.TrackIndex && VoiceId == TrackedEvent.GetVoiceId();
+			}))
 			{
-				KillVoiceFn(*It);
+				ActiveNotes.Add(Event);
 			}
 		}
-
-		ActiveVoices = MidiStream.ActiveVoices;
+		else if (Event.MidiMessage.IsNoteOff())
+		{
+			const int32 TrackIndex = Event.TrackIndex;
+			const FMidiVoiceId VoiceId = Event.GetVoiceId();
+				
+			ActiveNotes.RemoveAll([TrackIndex, VoiceId](const HarmonixMetasound::FMidiStreamEvent& TrackedEvent)
+			{
+				return TrackIndex == TrackedEvent.TrackIndex && VoiceId == TrackedEvent.GetVoiceId(); 
+			});
+		}
+		// If this is an all notes off/kill, untrack everything
+		// NOTE: This diverges from the MIDI spec, where "all notes off" has a track and channel associated with it.
+		// If we end up supporting the track- and channel-specific "all notes off" we will need to add support here to avoid stuck notes.
+		else
+		{
+			ActiveNotes.Reset();
+		}
 	}
-
 }
diff --git a/Engine/Plugins/Runtime/Harmonix/Source/HarmonixMetasound/Private/HarmonixMetasound/MidiOps/PulseGenerator.cpp b/Engine/Plugins/Runtime/Harmonix/Source/HarmonixMetasound/Private/HarmonixMetasound/MidiOps/PulseGenerator.cpp
--- a/Engine/Plugins/Runtime/Harmonix/Source/HarmonixMetasound/Private/HarmonixMetasound/MidiOps/PulseGenerator.cpp
+++ b/Engine/Plugins/Runtime/Harmonix/Source/HarmonixMetasound/Private/HarmonixMetasound/MidiOps/PulseGenerator.cpp
@@ -21,6 +21,8 @@ namespace Harmonix::Midi::Ops
 
 	void FPulseGenerator::Process(HarmonixMetasound::FMidiStream& OutStream)
 	{
+		OutStream.PrepareBlock();
+		
 		// Keep draining the queue if disabled, so we get the next note off,
 		// and so we stay in phase if we toggle off and back on
 		FCursor::FPulseTime NextPulse;
diff --git a/Engine/Plugins/Runtime/Harmonix/Source/HarmonixMetasound/Private/HarmonixMetasound/MidiOps/StuckNoteGuard.cpp b/Engine/Plugins/Runtime/Harmonix/Source/HarmonixMetasound/Private/HarmonixMetasound/MidiOps/StuckNoteGuard.cpp
--- a/Engine/Plugins/Runtime/Harmonix/Source/HarmonixMetasound/Private/HarmonixMetasound/MidiOps/StuckNoteGuard.cpp
+++ b/Engine/Plugins/Runtime/Harmonix/Source/HarmonixMetasound/Private/HarmonixMetasound/MidiOps/StuckNoteGuard.cpp
@@ -8,6 +8,42 @@ namespace Harmonix::Midi::Ops
 		const HarmonixMetasound::FMidiStream& InStream,
 		HarmonixMetasound::FMidiStream& OutStream,
 		const FIncludeNotePredicate& Predicate)
+	{
+		TrackNotes(InStream, Predicate);
+		
+		// For all of the tracked notes, unstick ones that no longer pass the filter
+		for (auto It = ActiveNotes.CreateIterator(); It; ++It)
+		{
+			if (!Predicate(*It))
+			{
+				// Send a note off
+				check(It->MidiMessage.IsNoteOn()); // shouldn't happen, yell for help
+				const uint8 Channel = It->MidiMessage.GetStdChannel();
+				const uint8 NoteNumber = It->MidiMessage.GetStdData1();
+				HarmonixMetasound::FMidiStreamEvent NoteOffEvent{ It->GetVoiceId().GetGeneratorId(), FMidiMsg::CreateNoteOff(Channel, NoteNumber) };
+				OutStream.InsertMidiEvent(NoteOffEvent);
+
+				// Stop tracking this note on this track
+				It.RemoveCurrent();
+			}
+		}
+	}
+
+	void FStuckNoteGuard::UnstickNotes(const HarmonixMetasound::FMidiStream& StreamToCompare, const FUnstickNoteFn& UnstickNoteFn)
+	{
+		TrackNotes(StreamToCompare, [](const HarmonixMetasound::FMidiStreamEvent&) { return true; });
+		
+		for (auto It = ActiveNotes.CreateIterator(); It; ++It)
+		{
+			if (!StreamToCompare.NoteIsActive(*It))
+			{
+				UnstickNoteFn(*It);
+				It.RemoveCurrent();
+			}
+		}
+	}
+
+	void FStuckNoteGuard::TrackNotes(const HarmonixMetasound::FMidiStream& InStream, const FIncludeNotePredicate& Predicate)
 	{
 		for (const HarmonixMetasound::FMidiStreamEvent& Event : InStream.GetEventsInBlock())
 		{
@@ -44,23 +80,5 @@ namespace Harmonix::Midi::Ops
 				ActiveNotes.Reset();
 			}
 		}
-		
-		// For all of the tracked notes, unstick ones that no longer pass the filter
-		for (auto It = ActiveNotes.CreateIterator(); It; ++It)
-		{
-			if (!Predicate(*It))
-			{
-				// Send a note off
-				check(It->MidiMessage.IsNoteOn()); // shouldn't happen, yell for help
-				const uint8 Channel = It->MidiMessage.GetStdChannel();
-				const uint8 NoteNumber = It->MidiMessage.GetStdData1();
-				HarmonixMetasound::FMidiStreamEvent NoteOffEvent{ static_cast<uint32>(0), FMidiMsg::CreateNoteOff(Channel, NoteNumber) };
-				NoteOffEvent.SetVoiceId(It->GetVoiceId());
-				OutStream.InsertMidiEvent(NoteOffEvent);
-
-				// Stop tracking this note on this track
-				It.RemoveCurrent();
-			}
-		}
 	}
 }
diff --git a/Engine/Plugins/Runtime/Harmonix/Source/HarmonixMetasound/Private/HarmonixMetasound/Nodes/FusionSamplerNode.cpp b/Engine/Plugins/Runtime/Harmonix/Source/HarmonixMetasound/Private/HarmonixMetasound/Nodes/FusionSamplerNode.cpp
--- a/Engine/Plugins/Runtime/Harmonix/Source/HarmonixMetasound/Private/HarmonixMetasound/Nodes/FusionSamplerNode.cpp
+++ b/Engine/Plugins/Runtime/Harmonix/Source/HarmonixMetasound/Private/HarmonixMetasound/Nodes/FusionSamplerNode.cpp
@@ -1,4 +1,5 @@
 // Copyright Epic Games, Inc. All Rights Reserved.
+
 #include "MetasoundExecutableOperator.h"
 #include "MetasoundNodeInterface.h"
 #include "MetasoundParamHelper.h"
@@ -17,6 +18,8 @@
 #include "HarmonixDsp/FusionSampler/FusionVoicePool.h"
 #include "HarmonixDsp/AudioUtility.h"
 
+#include "HarmonixMetasound/MidiOps/StuckNoteGuard.h"
+
 #define LOCTEXT_NAMESPACE "HarmonixMetaSound"
 
 DEFINE_LOG_CATEGORY_STATIC(LogFusionSamplerPlayer, Log, All);
@@ -49,7 +52,8 @@ namespace HarmonixMetasound
 	* FFusionSamplerOperatorBase
 	***********************************************************************************************************
 	*********************************************************************************************************/
-	class FFusionSamplerOperatorBase : public TExecutableOperator<FFusionSamplerOperatorBase>, protected FFusionSampler
+	class FFusionSamplerOperatorBase
+		: public TExecutableOperator<FFusionSamplerOperatorBase>, protected FFusionSampler
 	{
 	public:
 		struct FConstructionArgs
@@ -213,6 +217,7 @@ namespace HarmonixMetasound
 		FFusionPatchDataProxy::NodePtr FusionPatchDataPtr = nullptr;
 		int32 SliceIndex = 0;
 		EAudioBufferChannelLayout OutputChannelLayout = EAudioBufferChannelLayout::Mono;
+		Harmonix::Midi::Ops::FStuckNoteGuard StuckNoteGuard;
 
 		void SetPatchAndOverrides(FFusionPatchDataProxy::NodePtr NewPatch)
 		{
@@ -539,6 +544,11 @@ namespace HarmonixMetasound
 	*********************************************************************************************************/
 	void FFusionSamplerOperatorBase::Execute()
 	{
+		StuckNoteGuard.UnstickNotes(*MidiStreamInPin, [this](const FMidiStreamEvent& Event)
+		{
+			NoteOff(Event.GetVoiceId(), Event.MidiMessage.GetStdData1(), Event.MidiMessage.GetStdChannel());
+		});
+		
 		if (EnableMTRenderingInPin)
 		{
 			if (!SyncLinkOutPin->GetTask().IsCompleted())
diff --git a/Engine/Plugins/Runtime/Harmonix/Source/HarmonixMetasound/Private/HarmonixMetasound/Nodes/MidiNoteTriggerNode.cpp b/Engine/Plugins/Runtime/Harmonix/Source/HarmonixMetasound/Private/HarmonixMetasound/Nodes/MidiNoteTriggerNode.cpp
--- a/Engine/Plugins/Runtime/Harmonix/Source/HarmonixMetasound/Private/HarmonixMetasound/Nodes/MidiNoteTriggerNode.cpp
+++ b/Engine/Plugins/Runtime/Harmonix/Source/HarmonixMetasound/Private/HarmonixMetasound/Nodes/MidiNoteTriggerNode.cpp
@@ -14,6 +14,7 @@
 #include "HarmonixMetasound/Common.h"
 #include "HarmonixMetasound/DataTypes/MidiStream.h"
 #include "HarmonixMetasound/DataTypes/MusicTransport.h"
+#include "HarmonixMetasound/MidiOps/StuckNoteGuard.h"
 
 #define LOCTEXT_NAMESPACE "HarmonixMetaSound"
 
@@ -144,6 +145,14 @@ namespace HarmonixMetasound::Nodes::MidiNoteTriggerNode
 			NoteOnOutPin->AdvanceBlock();
 			NoteOffOutPin->AdvanceBlock();
 
+			int32 NoteOffTriggerFrame = -1;
+
+			StuckNoteGuard.UnstickNotes(*MidiStreamInPin, [this, &NoteOffTriggerFrame](const FMidiStreamEvent& Event)
+			{
+				TriggerNoteOff(0, Event.MidiMessage.GetStdData1());
+				NoteOffTriggerFrame = 0;
+			});
+
 			if (!*EnableInPin)
 			{
 				if (SoundingNote >= 0)
@@ -163,8 +172,6 @@ namespace HarmonixMetasound::Nodes::MidiNoteTriggerNode
 					return;
 				}
 			}
-
-			int32 NoteOffTriggerFrame = -1;
 			
 			for (const FMidiStreamEvent& Event : MidiStreamInPin->GetEventsInBlock())
 			{
@@ -214,6 +221,7 @@ namespace HarmonixMetasound::Nodes::MidiNoteTriggerNode
 
 		int8 SoundingNote = -1;
 		FMidiVoiceId PlayingId = FMidiVoiceId::None();
+		Harmonix::Midi::Ops::FStuckNoteGuard StuckNoteGuard;
 	};
 
 	class FMidiNoteTriggerNode_V1 final : public FNodeFacade
diff --git a/Engine/Plugins/Runtime/Harmonix/Source/HarmonixMetasound/Private/HarmonixMetasound/Nodes/MidiPulseGeneratorNode.cpp b/Engine/Plugins/Runtime/Harmonix/Source/HarmonixMetasound/Private/HarmonixMetasound/Nodes/MidiPulseGeneratorNode.cpp
--- a/Engine/Plugins/Runtime/Harmonix/Source/HarmonixMetasound/Private/HarmonixMetasound/Nodes/MidiPulseGeneratorNode.cpp
+++ b/Engine/Plugins/Runtime/Harmonix/Source/HarmonixMetasound/Private/HarmonixMetasound/Nodes/MidiPulseGeneratorNode.cpp
@@ -187,8 +187,7 @@ namespace HarmonixMetasound::Nodes::MidiPulseGeneratorNode
 		void Execute()
 		{
 			ApplyParameters();
-
-			Outputs.MidiStream->PrepareBlock();
+			
 			PulseGenerator.Process(*Outputs.MidiStream);
 		}
 	private:
diff --git a/Engine/Plugins/Runtime/Harmonix/Source/HarmonixMetasound/Private/HarmonixMetasound/Nodes/MidiStreamSelectNode.cpp b/Engine/Plugins/Runtime/Harmonix/Source/HarmonixMetasound/Private/HarmonixMetasound/Nodes/MidiStreamSelectNode.cpp
--- a/Engine/Plugins/Runtime/Harmonix/Source/HarmonixMetasound/Private/HarmonixMetasound/Nodes/MidiStreamSelectNode.cpp
+++ b/Engine/Plugins/Runtime/Harmonix/Source/HarmonixMetasound/Private/HarmonixMetasound/Nodes/MidiStreamSelectNode.cpp
@@ -37,8 +37,6 @@ namespace HarmonixMetasound
 
 		virtual void BindInputs(FInputVertexInterfaceData& InVertexData) override;
 		virtual void BindOutputs(FOutputVertexInterfaceData& InVertexData) override;
-		virtual FDataReferenceCollection GetInputs() const override;
-		virtual FDataReferenceCollection GetOutputs() const override;
 
 		void Execute();
 		void Reset(const FResetParams& ResetParams);
@@ -175,22 +173,6 @@ namespace HarmonixMetasound
 		InVertexData.BindReadVertex(METASOUND_GET_PARAM_NAME(Outputs::MidiStream), MidiStreamOutPin);
 	}
 
-	FDataReferenceCollection FMidiStreamSelectOperator::GetInputs() const
-	{
-		// This should never be called. Bind(...) is called instead. This method
-		// exists as a stop-gap until the API can be deprecated and removed.
-		checkNoEntry();
-		return {};
-	}
-
-	FDataReferenceCollection FMidiStreamSelectOperator::GetOutputs() const
-	{
-		// This should never be called. Bind(...) is called instead. This method
-		// exists as a stop-gap until the API can be deprecated and removed.
-		checkNoEntry();
-		return {};
-	}
-
 	void FMidiStreamSelectOperator::Execute()
 	{
 		MidiStreamOutPin->PrepareBlock();
@@ -245,12 +227,11 @@ namespace HarmonixMetasound
 					uint8 MidiCh;
 					uint8 MidiNote;
 					VoiceId.GetChannelAndNote(MidiCh, MidiNote);
-					FMidiStreamEvent MidiEvent(0u, FMidiMsg::CreateNoteOff(MidiCh, MidiNote)); // Do we need a note? VoiceId should handle the note off
+					FMidiStreamEvent MidiEvent(VoiceId, FMidiMsg::CreateNoteOff(MidiCh, MidiNote)); // Do we need a note? VoiceId should handle the note off
 					MidiEvent.BlockSampleFrameIndex = 0;
 					MidiEvent.AuthoredMidiTick = 0;
 					MidiEvent.CurrentMidiTick = 0;
 					MidiEvent.TrackIndex = 1;
-					MidiEvent.SetVoiceId(VoiceId);
 					MidiStreamOutPin->InsertNoteOffEventOrCancelPendingNoteOn(MidiEvent);
 				}
 				PlayingVoices.Empty();
diff --git a/Engine/Plugins/Runtime/Harmonix/Source/HarmonixMetasound/Private/HarmonixMetasound/Nodes/MidiStreamTransposerNode.cpp b/Engine/Plugins/Runtime/Harmonix/Source/HarmonixMetasound/Private/HarmonixMetasound/Nodes/MidiStreamTransposerNode.cpp
--- a/Engine/Plugins/Runtime/Harmonix/Source/HarmonixMetasound/Private/HarmonixMetasound/Nodes/MidiStreamTransposerNode.cpp
+++ b/Engine/Plugins/Runtime/Harmonix/Source/HarmonixMetasound/Private/HarmonixMetasound/Nodes/MidiStreamTransposerNode.cpp
@@ -140,6 +140,8 @@ namespace HarmonixMetasound::Nodes::MidiNoteTranspose
 
 		void Execute()
 		{
+			Outputs.MidiStream->PrepareBlock();
+			
 			if (!*Inputs.Enabled)
 			{
 				return;
diff --git a/Engine/Plugins/Runtime/Harmonix/Source/HarmonixMetasound/Private/HarmonixMetasound/Nodes/StepSequencePlayerNode.cpp b/Engine/Plugins/Runtime/Harmonix/Source/HarmonixMetasound/Private/HarmonixMetasound/Nodes/StepSequencePlayerNode.cpp
--- a/Engine/Plugins/Runtime/Harmonix/Source/HarmonixMetasound/Private/HarmonixMetasound/Nodes/StepSequencePlayerNode.cpp
+++ b/Engine/Plugins/Runtime/Harmonix/Source/HarmonixMetasound/Private/HarmonixMetasound/Nodes/StepSequencePlayerNode.cpp
@@ -62,6 +62,15 @@ namespace HarmonixMetasound
 
 	private:
 		void Init();
+
+		void PrepareMidiOutIfNeeded()
+		{
+			if (bNeedsMidiOutClear)
+			{
+				MidiOutPin->PrepareBlock();
+				bNeedsMidiOutClear = false;
+			}
+		}
 		
 		//** INPUTS
 		FMidiStepSequenceAssetReadRef SequenceAssetInPin;
@@ -327,11 +336,7 @@ namespace HarmonixMetasound
 	{
 		// If no midi cursor callbacks happened we still have to 
 		// remove old midi events from last block...
-		if (bNeedsMidiOutClear)
-		{
-			MidiOutPin->PrepareBlock();
-			bNeedsMidiOutClear = false;
-		}
+		PrepareMidiOutIfNeeded();
 
 		// if we have no sequence table there is nothing to do. 
 		// Make sure the notes are all off and return.
@@ -472,11 +477,7 @@ namespace HarmonixMetasound
 
 	void FStepSequencePlayerOperator::SeekThruTick(int32 Tick)
 	{
-		if (bNeedsMidiOutClear)
-		{
-			MidiOutPin->PrepareBlock();
-			bNeedsMidiOutClear = false;
-		}
+		PrepareMidiOutIfNeeded();
 
 		int32 BlockFrame = MidiClockInPin->GetCurrentBlockFrameIndex();
 		ProcessedThruTick = FMath::Max(-1, Tick);
@@ -570,11 +571,7 @@ namespace HarmonixMetasound
 			return;
 		}
 
-		if (bNeedsMidiOutClear)
-		{
-			MidiOutPin->PrepareBlock();
-			bNeedsMidiOutClear = false;
-		}
+		PrepareMidiOutIfNeeded();
 
 		// Read input pins
 		const int32 CurrentMaxColumns = *MaxColumnsInPin;
@@ -773,12 +770,11 @@ namespace HarmonixMetasound
 						}
 
 						// note off!
-						FMidiStreamEvent MidiEvent(this, FMidiMsg::CreateNoteOff(MidiCh, SequenceTable->Notes[i].NoteNumber + AdditionalOctaveNotes));
+						FMidiStreamEvent MidiEvent(CurrentCellNotes[i].GetGeneratorId(), FMidiMsg::CreateNoteOff(MidiCh, SequenceTable->Notes[i].NoteNumber + AdditionalOctaveNotes));
 						MidiEvent.BlockSampleFrameIndex = BlockFrameIndex;
 						MidiEvent.AuthoredMidiTick = ProcessedThruTick;
 						MidiEvent.CurrentMidiTick = ProcessedThruTick;
 						MidiEvent.TrackIndex = 1;
-						MidiEvent.SetVoiceId(CurrentCellNotes[i]);
 						MidiOutPin->AddNoteOffEventOrCancelPendingNoteOn(MidiEvent);
 						UE_LOG(LogStepSequencePlayer, Verbose, TEXT("0x%x Note-Off %d at %d"), (uint32)(size_t)this, SequenceTable->Notes[i].NoteNumber + AdditionalOctaveNotes, BlockFrameIndex);
 						CurrentCellNotes[i] = FMidiVoiceId::None();
@@ -853,12 +849,11 @@ namespace HarmonixMetasound
 					uint8 MidiCh;
 					uint8 MidiNote;
 					CurrentCellNotes[i].GetChannelAndNote(MidiCh, MidiNote);
-					FMidiStreamEvent MidiEvent(this, FMidiMsg::CreateNoteOff(MidiCh, MidiNote));
+					FMidiStreamEvent MidiEvent(CurrentCellNotes[i].GetGeneratorId(), FMidiMsg::CreateNoteOff(MidiCh, MidiNote));
 					MidiEvent.BlockSampleFrameIndex  = CurrentBlockSpanStart;
 					MidiEvent.AuthoredMidiTick       = 0;
 					MidiEvent.CurrentMidiTick        = 0;
 					MidiEvent.TrackIndex             = 1;
-					MidiEvent.SetVoiceId(CurrentCellNotes[i]);
 					UE_LOG(LogStepSequencePlayer, Verbose, TEXT("0x%x Note-Off %d (during resize)"), (uint32)(size_t)this, MidiNote);
 					MidiOutPin->AddMidiEvent(MidiEvent);
 				}
@@ -883,12 +878,11 @@ namespace HarmonixMetasound
 				uint8 MidiCh;
 				uint8 MidiNote;
 				CurrentCellNotes[i].GetChannelAndNote(MidiCh, MidiNote);
-				FMidiStreamEvent MidiEvent(this, FMidiMsg::CreateNoteOff(MidiCh, MidiNote));
+				FMidiStreamEvent MidiEvent(CurrentCellNotes[i].GetGeneratorId(), FMidiMsg::CreateNoteOff(MidiCh, MidiNote));
 				MidiEvent.BlockSampleFrameIndex  = AtFrameIndex;
 				MidiEvent.AuthoredMidiTick       = AbsMidiTick;
 				MidiEvent.CurrentMidiTick        = AbsMidiTick;
 				MidiEvent.TrackIndex             = 1;
-				MidiEvent.SetVoiceId(CurrentCellNotes[i]);
 				MidiOutPin->AddMidiEvent(MidiEvent);
 				UE_LOG(LogStepSequencePlayer, Verbose, TEXT("0x%x Note-Off %d (during all notes off)"), (uint32)(size_t)this, MidiNote);
 				CurrentCellNotes[i] = FMidiVoiceId::None();
diff --git a/Engine/Plugins/Runtime/Harmonix/Source/HarmonixMetasound/Public/HarmonixMetasound/DataTypes/MidiStream.h b/Engine/Plugins/Runtime/Harmonix/Source/HarmonixMetasound/Public/HarmonixMetasound/DataTypes/MidiStream.h
--- a/Engine/Plugins/Runtime/Harmonix/Source/HarmonixMetasound/Public/HarmonixMetasound/DataTypes/MidiStream.h
+++ b/Engine/Plugins/Runtime/Harmonix/Source/HarmonixMetasound/Public/HarmonixMetasound/DataTypes/MidiStream.h
@@ -23,34 +23,15 @@ namespace HarmonixMetasound
 		int32        CurrentMidiTick        = 0;
 		FMidiMsg     MidiMessage;
 
-		FMidiStreamEvent(const FMidiVoiceGeneratorBase* Owner, const FMidiMsg& Message)
-			: MidiMessage(Message)
-		{
-			VoiceId = FMidiVoiceId(Owner, MidiMessage);
-		}
+		FMidiStreamEvent(const FMidiVoiceGeneratorBase* Owner, const FMidiMsg& Message);
 
-		FMidiStreamEvent(uint32 OwnerId, const FMidiMsg& Message)
-			: MidiMessage(Message)
-		{
-			VoiceId = FMidiVoiceId(OwnerId, MidiMessage);
-		}
-
-		// reassign owner with a new generator
-		void ReassignOwner(const FMidiVoiceGeneratorBase* Owner)
-		{
-			VoiceId.ReassignGenerator(Owner);
-		}
-
-		// reassign owner with a new generator id
-		void ReassignOwner(uint32 OwnerId)
-		{
-			VoiceId.ReassignGenerator(OwnerId);
-		}
+		FMidiStreamEvent(const uint32 OwnerId, const FMidiMsg& Message);
 
 		FMidiVoiceId GetVoiceId() const { return VoiceId; }
-		void SetVoiceId(FMidiVoiceId InVoiceId) { VoiceId = InVoiceId; }
 
 	private:
+		friend class FMidiStream;
+		
 		FMidiVoiceId VoiceId;
 	};
 
@@ -105,35 +86,25 @@ namespace HarmonixMetasound
 			const FEventFilter& Filter = NoOpFilter,
 			const FEventTransformer& Transformer = NoOpTransformer);
 
+		bool NoteIsActive(const FMidiStreamEvent& Event) const;
+
 	private:
 		FMidiFileProxyPtr MidiFileSourceOfEvents;
 		int32 TicksPerQuarterNote = Harmonix::Midi::Constants::GTicksPerQuarterNoteInt;
 
 		TArray<FMidiStreamEvent> EventsInBlock;
 
-		friend class FMidiVoiceTracker;
-		void UpdateActiveVoice(const FMidiStreamEvent& Event);
-		TSet<FMidiVoiceId> ActiveVoices;
-
 		TWeakPtr<const FMidiClock, ESPMode::NotThreadSafe> Clock;
 
 		// Map to handle re-mapping merged MIDI events, which helps to disambiguate split/transposed notes
 		TMap<uint32, FMidiVoiceGeneratorBase> GeneratorMap;
+
+		void TrackNote(const FMidiStreamEvent& Event);
+		TArray<FMidiStreamEvent> ActiveNotes;
 	};
 
 	// Declare aliases IN the namespace...
 	DECLARE_METASOUND_DATA_REFERENCE_ALIAS_TYPES(FMidiStream, FMidiStreamTypeInfo, FMidiStreamReadRef, FMidiStreamWriteRef)
-
-	class HARMONIXMETASOUND_API FMidiVoiceTracker
-	{
-	public:
-		using FKillVoiceFn = TFunctionRef<void(const FMidiVoiceId&)>;
-
-		void Process(const FMidiStream& MidiStream, const FKillVoiceFn& KillVoiceFn);
-
-	private:
-		TSet<FMidiVoiceId> ActiveVoices;
-	};
 }
 
 // Declare reference types OUT of the namespace...
diff --git a/Engine/Plugins/Runtime/Harmonix/Source/HarmonixMetasound/Public/HarmonixMetasound/MidiOps/StuckNoteGuard.h b/Engine/Plugins/Runtime/Harmonix/Source/HarmonixMetasound/Public/HarmonixMetasound/MidiOps/StuckNoteGuard.h
--- a/Engine/Plugins/Runtime/Harmonix/Source/HarmonixMetasound/Public/HarmonixMetasound/MidiOps/StuckNoteGuard.h
+++ b/Engine/Plugins/Runtime/Harmonix/Source/HarmonixMetasound/Public/HarmonixMetasound/MidiOps/StuckNoteGuard.h
@@ -22,7 +22,18 @@ namespace Harmonix::Midi::Ops
 		 */
 		void Process(const HarmonixMetasound::FMidiStream& InStream, HarmonixMetasound::FMidiStream& OutStream, const FIncludeNotePredicate& Predicate);
 
+		using FUnstickNoteFn = TFunction<void(const HarmonixMetasound::FMidiStreamEvent&)>;
+
+		/**
+		 * If any active notes we're tracking are no longer present in the input stream, trigger a function so the caller can turn them off.
+		 * @param StreamToCompare The stream to use to check whether an active note has disappeared
+		 * @param UnstickNoteFn The function to call when an active note has disappeared from the incoming stream
+		 */
+		void UnstickNotes(const HarmonixMetasound::FMidiStream& StreamToCompare, const FUnstickNoteFn& UnstickNoteFn);
+
 	private:
+		void TrackNotes(const HarmonixMetasound::FMidiStream& InStream, const FIncludeNotePredicate& Predicate);
+		
 		TArray<HarmonixMetasound::FMidiStreamEvent> ActiveNotes;
 	};
 }
diff --git a/Engine/Plugins/Runtime/Harmonix/Source/HarmonixMidi/Private/HarmonixMidi/MidiVoiceId.cpp b/Engine/Plugins/Runtime/Harmonix/Source/HarmonixMidi/Private/HarmonixMidi/MidiVoiceId.cpp
--- a/Engine/Plugins/Runtime/Harmonix/Source/HarmonixMidi/Private/HarmonixMidi/MidiVoiceId.cpp
+++ b/Engine/Plugins/Runtime/Harmonix/Source/HarmonixMidi/Private/HarmonixMidi/MidiVoiceId.cpp
@@ -9,10 +9,10 @@ FMidiVoiceGeneratorBase::FMidiVoiceGeneratorBase()
 {
 	FScopeLock Lock(&GeneratorIdLock);
 	IdBits = NextGeneratorId++;
-	if (NextGeneratorId >= (1L << kIdWidth))
+	if (NextGeneratorId >= (1L << IdWidth))
 	{
 		NextGeneratorId = 1;
 	}
-	IdBits <<= (32 - kIdWidth);
+	IdBits <<= (32 - IdWidth);
 }
 
diff --git a/Engine/Plugins/Runtime/Harmonix/Source/HarmonixMidi/Public/HarmonixMidi/MidiVoiceId.h b/Engine/Plugins/Runtime/Harmonix/Source/HarmonixMidi/Public/HarmonixMidi/MidiVoiceId.h
--- a/Engine/Plugins/Runtime/Harmonix/Source/HarmonixMidi/Public/HarmonixMidi/MidiVoiceId.h
+++ b/Engine/Plugins/Runtime/Harmonix/Source/HarmonixMidi/Public/HarmonixMidi/MidiVoiceId.h
@@ -9,7 +9,7 @@ public:
 	FMidiVoiceGeneratorBase();
 	uint32 GetIdBits() const { return IdBits; }
 
-	static const uint32 kIdWidth = 20;
+	static constexpr uint32 IdWidth = 20;
 
 private:
 	uint32 IdBits;
@@ -24,17 +24,9 @@ public:
 		: Id(0)
 	{}
 
-	FMidiVoiceId(const FMidiVoiceGeneratorBase* Generator, const FMidiMsg& FromMessage)
+	FMidiVoiceId(const uint32 GeneratorId, const FMidiMsg& FromMessage)
 	{
-		static_assert(FMidiVoiceGeneratorBase::kIdWidth <= 20);
-		uint8 NotePart = (FromMessage.IsNoteOn() || FromMessage.IsNoteOff()) ? FromMessage.GetStdData1() : 0;
-		uint8 ChPart   = (FromMessage.IsStd()) ? FromMessage.GetStdChannel() : 0;
-		Id = (Generator ? Generator->GetIdBits() : 0) | ((uint32)(ChPart & 0xF) << 8) | (uint32)NotePart;
-	}
-
-	FMidiVoiceId(uint32 GeneratorId, const FMidiMsg& FromMessage)
-	{
-		static_assert(FMidiVoiceGeneratorBase::kIdWidth <= 20);
+		static_assert(FMidiVoiceGeneratorBase::IdWidth <= 20);
 		uint8 NotePart = (FromMessage.IsNoteOn() || FromMessage.IsNoteOff()) ? FromMessage.GetStdData1() : 0;
 		uint8 ChPart = (FromMessage.IsStd()) ? FromMessage.GetStdChannel() : 0;
 		Id = GeneratorId | ((uint32)(ChPart & 0xF) << 8) | (uint32)NotePart;
@@ -44,17 +36,12 @@ public:
 		: Id (Other.Id)
 	{}
 
-	FMidiVoiceId(FMidiVoiceId&& Other)
+	FMidiVoiceId(FMidiVoiceId&& Other) noexcept
 		: Id(Other.Id)
 	{
 		Other.Id = 0;
 	}
 
-	void ReassignGenerator(const FMidiVoiceGeneratorBase* Generator)
-	{
-		ReassignGenerator(Generator ? Generator->GetIdBits() : 0);
-	}
-
 	void ReassignGenerator(uint32 GeneratorId)
 	{
 		uint8 Ch;
