MidiClock
- Add method to SeekTo specific tick.
- Add method to calculate "Mapped" loop tick based on loop region.

Metronome and MidiPlayer Node
- Use "SeekTo" tick to simplify seeking based on MidiClockEvents
- Use the "MappedTick" to seek based on looping
- AdvanceThru now seeks based on "DeltaTicks" instead, so the clock doesn't need to keep track of the "OffsetTick"
- We _have_ to respond to Reset now since AdvanceThru is a relative operation and we can't risk getting out of sink with owr clock authority

There's still a lot of repetative code. I plan on making another change that moves even more logic to the MidiClock (specifically, all the functionality handling the ClockEvents).

#tests HarmonixTests

diff --git a/Engine/Plugins/Runtime/Harmonix/Source/HarmonixMetasound/Private/HarmonixMetasound/DataTypes/MidiClock.cpp b/Engine/Plugins/Runtime/Harmonix/Source/HarmonixMetasound/Private/HarmonixMetasound/DataTypes/MidiClock.cpp
--- a/Engine/Plugins/Runtime/Harmonix/Source/HarmonixMetasound/Private/HarmonixMetasound/DataTypes/MidiClock.cpp
+++ b/Engine/Plugins/Runtime/Harmonix/Source/HarmonixMetasound/Private/HarmonixMetasound/DataTypes/MidiClock.cpp
@@ -308,6 +308,11 @@ namespace HarmonixMetasound
 		DrivingMidiPlayCursorMgr->AttachToTimeAuthority(MidiClockRef.DrivingMidiPlayCursorMgr);
 	}
 
+	void FMidiClock::DetachFromTimeAuthority()
+	{
+		DrivingMidiPlayCursorMgr->DetachFromTimeAuthority();
+	}
+	
 	float FMidiClock::GetQuarterNoteIncludingCountIn() const
 	{
 		int32 Tick = DrivingMidiPlayCursorMgr->GetCurrentHiResTick();
@@ -360,6 +365,11 @@ namespace HarmonixMetasound
 		DrivingMidiPlayCursorMgr->SeekTo(Tick, PreRollBars,true,false);
 	}
 
+	void FMidiClock::SeekTo(int32 Tick, int32 PreRollBars)
+	{
+		DrivingMidiPlayCursorMgr->SeekTo(Tick, PreRollBars, true, false);
+	}
+
 	void FMidiClock::InformOfCurrentAdvanceRate(float AdvanceRate)
 	{
 		DrivingMidiPlayCursorMgr->InformOfHiResAdvanceRate(AdvanceRate);
@@ -418,6 +428,31 @@ namespace HarmonixMetasound
 		FramesUntilNextProcess = 0;
 	}
 
+	void FMidiClock::SeekTo(int32 BlockFrameIndex, int32 Tick, int32 InPrerollBars)
+	{
+		CurrentBlockFrameIndex = BlockFrameIndex;
+		DrivingMidiPlayCursorMgr->SeekTo(Tick, InPrerollBars, true, false);
+		SampleCount = FMath::Max<FSampleCount>(FSampleCount(GetCurrentHiResMs() / 1000.0f * SampleRate), 0);
+		FramesUntilNextProcess = 0;
+	}
+
+	int32 FMidiClock::CalculateMappedTick(int32 Tick) const
+	{
+		if (DoesLoop())
+		{
+			const int32 LoopStartTick = GetLoopStartTick();
+			const int32 LoopEndTick = GetLoopEndTick();
+			const int32 LoopLengthTicks = LoopEndTick - LoopStartTick;
+			// only wrap the tick if it we're passed the loop end tick
+			if (LoopLengthTicks > 0 && Tick >= LoopEndTick)
+			{
+				return LoopStartTick + (Tick - LoopStartTick) % LoopLengthTicks;
+			}
+		}
+		return Tick;
+	}
+
+
 	TSharedPtr<FMidiFileData> FMidiClock::MakeClockConductorMidiData(float InTempoBPM, int32 InTimeSigNum, int32 InTimeSigDen)
 	{
 		TSharedPtr<FMidiFileData> OutMidiData = MakeShared<FMidiFileData>();
diff --git a/Engine/Plugins/Runtime/Harmonix/Source/HarmonixMetasound/Private/HarmonixMetasound/Nodes/MetronomeNode.cpp b/Engine/Plugins/Runtime/Harmonix/Source/HarmonixMetasound/Private/HarmonixMetasound/Nodes/MetronomeNode.cpp
--- a/Engine/Plugins/Runtime/Harmonix/Source/HarmonixMetasound/Private/HarmonixMetasound/Nodes/MetronomeNode.cpp
+++ b/Engine/Plugins/Runtime/Harmonix/Source/HarmonixMetasound/Private/HarmonixMetasound/Nodes/MetronomeNode.cpp
@@ -47,7 +47,6 @@ namespace HarmonixMetasound
 
 		void Reset(const FResetParams& Params);
 		void Execute();
-		void UpdateLoopOffsetTickFromTick(int32 Tick);
 
 		//~ BEGIN FMidiPlayCursor Overrides
 		virtual void Reset(bool ForceNoBroadcast = false) override;
@@ -82,7 +81,6 @@ namespace HarmonixMetasound
 		int32        CurrentTimeSigNum;
 		int32        CurrentTimeSigDenom;
 		int32		 LastClockTickUpdate = -1;
-		int32		 LoopOffsetTick = 0;
 
 		void BuildMidiData(bool ResetToStart = true);
 		void UpdateMidi();
@@ -256,15 +254,8 @@ namespace HarmonixMetasound
 	void FMetronomeOperator::Reset(bool ForceNoBroadcast /*= false*/)
 	{
 		FMidiPlayCursor::Reset(ForceNoBroadcast);
-		int32 Tick = MetronomeClock.GetCurrentHiResTick();
-		UpdateLoopOffsetTickFromTick(Tick);
-		int32 MappedTick = Tick - LoopOffsetTick;
-
-		float Ms = MidiClockOutPin->GetSongMaps().TickToMs(Tick + 1);
-		FMusicSeekTarget SeekTarget;
-		SeekTarget.Type = ESeekPointType::Millisecond;
-		SeekTarget.Ms = Ms;
-		MidiClockOutPin->SeekTo(MetronomeClock.GetCurrentBlockFrameIndex(), SeekTarget, SeekPreRollBarsInPin);
+		int32 Tick = MidiClockOutPin->CalculateMappedTick(MetronomeClock.GetCurrentHiResTick());
+		MidiClockOutPin->SeekTo(MetronomeClock.GetCurrentBlockFrameIndex(), Tick + 1, SeekPreRollBarsInPin);
 	}
 
 	void FMetronomeOperator::OnLoop(int32 LoopStartTick, int32 LoopEndTick)
@@ -275,63 +266,38 @@ namespace HarmonixMetasound
 	void FMetronomeOperator::SeekToTick(int32 Tick)
 	{
 		FMidiPlayCursor::SeekToTick(Tick);
-		UpdateLoopOffsetTickFromTick(Tick);
-		int32 MappedTick = Tick - LoopOffsetTick;
-
-		float Ms = MidiClockOutPin->GetSongMaps().TickToMs(Tick + 1);
-		FMusicSeekTarget SeekTarget;
-		SeekTarget.Type = ESeekPointType::Millisecond;
-		SeekTarget.Ms = Ms;
-		MidiClockOutPin->SeekTo(MetronomeClock.GetCurrentBlockFrameIndex(), SeekTarget, SeekPreRollBarsInPin);
+		int32 MappedTick = MidiClockOutPin->CalculateMappedTick(Tick);
+		MidiClockOutPin->SeekTo(MetronomeClock.GetCurrentBlockFrameIndex(), MappedTick, SeekPreRollBarsInPin);
 	}
 
 	void FMetronomeOperator::SeekThruTick(int32 Tick)
 	{
 		FMidiPlayCursor::SeekThruTick(Tick);
-
-		UpdateLoopOffsetTickFromTick(Tick);
-		int32 MappedTick = Tick - LoopOffsetTick;
-
-		float Ms = MidiClockOutPin->GetSongMaps().TickToMs(Tick + 1);
-		FMusicSeekTarget SeekTarget;
-		SeekTarget.Type = ESeekPointType::Millisecond;
-		SeekTarget.Ms = Ms;
-		MidiClockOutPin->SeekTo(MetronomeClock.GetCurrentBlockFrameIndex(), SeekTarget, SeekPreRollBarsInPin);
+		int32 MappedTick = MidiClockOutPin->CalculateMappedTick(Tick);
+		MidiClockOutPin->SeekTo(MetronomeClock.GetCurrentBlockFrameIndex(), MappedTick + 1, SeekPreRollBarsInPin);
 	}
 
 	void FMetronomeOperator::AdvanceThruTick(int32 InTick, bool IsPreRoll)
 	{
+		int32 PrevTick = CurrentTick;
 		FMidiPlayCursor::AdvanceThruTick(InTick, IsPreRoll);
 
 		// if this is a pre-roll, perform a seek instead
 		// NOTE: This code is/should be identical to SeekThruTick above
 		if (IsPreRoll)
 		{
-			UpdateLoopOffsetTickFromTick(InTick);
-			int32 MappedTick = InTick - LoopOffsetTick;
-
-			float Ms = MidiClockOutPin->GetSongMaps().TickToMs(InTick + 1);
-			FMusicSeekTarget SeekTarget;
-			SeekTarget.Type = ESeekPointType::Millisecond;
-			SeekTarget.Ms = Ms;
-			MidiClockOutPin->SeekTo(MetronomeClock.GetCurrentBlockFrameIndex(), SeekTarget, SeekPreRollBarsInPin);
+			int32 MappedTick = MidiClockOutPin->CalculateMappedTick(InTick);
+			MidiClockOutPin->SeekTo(MetronomeClock.GetCurrentBlockFrameIndex(), MappedTick + 1, SeekPreRollBarsInPin);
 			return;
 		}
-
-		int32 MappedTick = InTick - LoopOffsetTick;
+		
+		int32 NewTick = MidiClockOutPin->GetCurrentMidiTick() + (InTick - PrevTick);
+		float Ms = MidiClockOutPin->GetSongMaps().TickToMs(NewTick);
 		float CurrentSpeed = MetronomeClock.GetSpeedAtBlockSampleFrame(MetronomeClock.GetCurrentBlockFrameIndex());
-		float Ms = MidiClockOutPin->GetSongMaps().TickToMs(MappedTick);
-		float AdvanceRatio = Owner->GetSongMaps().GetTempoAtTick(InTick) / MidiClockOutPin->GetSongMaps().GetTempoAtTick(MappedTick);
+		float AdvanceRatio = Owner->GetSongMaps().GetTempoAtTick(PrevTick)
+		                   / MidiClockOutPin->GetSongMaps().GetTempoAtTick(MidiClockOutPin->GetCurrentMidiTick());
 		MidiClockOutPin->InformOfCurrentAdvanceRate(CurrentSpeed * AdvanceRatio);
 		MidiClockOutPin->AdvanceHiResToMs(MetronomeClock.GetCurrentBlockFrameIndex(), Ms, true);
-
-		if (MidiClockOutPin->DoesLoop())
-		{
-			if (MidiClockOutPin->GetCurrentHiResTick() < MappedTick)
-			{
-				UpdateLoopOffsetTickFromTick(InTick);
-			}
-		}
 	}
 
 	void FMetronomeOperator::Init()
@@ -481,28 +447,6 @@ namespace HarmonixMetasound
 		}
 	}
 
-	void FMetronomeOperator::UpdateLoopOffsetTickFromTick(int32 Tick)
-	{
-		// only update our loop offset if we're actually looping
-		// if our external clock is seeking us
-		// we need to update our loop offset ticks
-		if (MidiClockOutPin->DoesLoop())
-		{
-			int32 LoopStartTick = MidiClockOutPin->GetLoopStartTick();
-			int32 LoopEndTick = MidiClockOutPin->GetLoopEndTick();
-			int32 LoopLengthTicks = LoopEndTick - LoopStartTick;
-			if (LoopLengthTicks > 0)
-			{
-				// the whole number part of the division tells us 
-				// how many times we "should have looped" based on the incoming tick
-				int32 LoopNum = (Tick - LoopStartTick) / LoopLengthTicks;
-
-				// So we can multiply it back to get the LoopOffsetTick
-				LoopOffsetTick = LoopLengthTicks * LoopNum + LoopStartTick;
-			}
-		}
-	}
-
 	void FMetronomeOperator::BuildMidiData(bool ResetToStart)
 	{
 		// make sure we have valid values
diff --git a/Engine/Plugins/Runtime/Harmonix/Source/HarmonixMetasound/Private/HarmonixMetasound/Nodes/MidiPlayerNode.cpp b/Engine/Plugins/Runtime/Harmonix/Source/HarmonixMetasound/Private/HarmonixMetasound/Nodes/MidiPlayerNode.cpp
--- a/Engine/Plugins/Runtime/Harmonix/Source/HarmonixMetasound/Private/HarmonixMetasound/Nodes/MidiPlayerNode.cpp
+++ b/Engine/Plugins/Runtime/Harmonix/Source/HarmonixMetasound/Private/HarmonixMetasound/Nodes/MidiPlayerNode.cpp
@@ -121,15 +121,9 @@ namespace HarmonixMetasound
 
 	protected:
 		virtual void InitTransportImpl() override;
-		virtual void SetupNewMidiFile(const FMidiFileProxyPtr& NewMidi) override;
 	private:
 		//** INPUTS **********************************
 		FMidiClockReadRef MidiClockIn;
-
-		//** DATA   **********************************
-		int32 LoopOffsetTick = 0;
-
-		void UpdateLoopOffsetTickFromTick(int32 Tick);
 	};
 
 	class FSelfClockedMidiPlayerOperator : public FMidiPlayerOperator
@@ -310,7 +304,6 @@ namespace HarmonixMetasound
 	void FExternallyClockedMidiPlayerOperator::Reset(const FResetParams& Params)
 	{
 		FMidiPlayerOperator::Reset(Params);
-		LoopOffsetTick = 0;
 	}
 
 	void FSelfClockedMidiPlayerOperator::BindInputs(FInputVertexInterfaceData& InVertexData)
@@ -411,9 +404,8 @@ namespace HarmonixMetasound
 					{
 					case FMidiClockEvent::EType::Reset:
 					{
-						// Ignore resets, because every reset event is preceded by a SeekThru event
-						// SeekThru triggers a new reset in the following clock,
-						// resulting in an exponentially growing list of reset and seek thru events
+						int32 Tick = MidiClockOut->CalculateMappedTick(Event.Tick2);
+						MidiClockOut->SeekTo(Event.BlockFrameIndex, Tick + 1, PrerollBars);
 						break;
 					}
 					case FMidiClockEvent::EType::Loop:
@@ -423,26 +415,14 @@ namespace HarmonixMetasound
 					}
 					case FMidiClockEvent::EType::SeekTo:
 					{
-						UpdateLoopOffsetTickFromTick(Event.Tick2);
-						int32 Tick = Event.Tick2 - LoopOffsetTick;
-
-						float Ms = MidiClockOut->GetSongMaps().TickToMs(Tick);
-						FMusicSeekTarget SeekTarget;
-						SeekTarget.Type = ESeekPointType::Millisecond;
-						SeekTarget.Ms = Ms;
-						MidiClockOut->SeekTo(Event.BlockFrameIndex, SeekTarget, PrerollBars);
+						int32 Tick = MidiClockOut->CalculateMappedTick(Event.Tick2);
+						MidiClockOut->SeekTo(Event.BlockFrameIndex, Tick, PrerollBars);
 						break;
 					}
 					case FMidiClockEvent::EType::SeekThru:
 					{
-						UpdateLoopOffsetTickFromTick(Event.Tick2);
-						int32 Tick = Event.Tick2 - LoopOffsetTick;
-
-						float Ms = MidiClockOut->GetSongMaps().TickToMs(Tick + 1);
-						FMusicSeekTarget SeekTarget;
-						SeekTarget.Type = ESeekPointType::Millisecond;
-						SeekTarget.Ms = Ms;
-						MidiClockOut->SeekTo(Event.BlockFrameIndex, SeekTarget, PrerollBars);
+						int32 Tick = MidiClockOut->CalculateMappedTick(Event.Tick2);
+						MidiClockOut->SeekTo(Event.BlockFrameIndex, Tick + 1, PrerollBars);
 						break;
 					}
 					case FMidiClockEvent::EType::AdvanceThru:
@@ -451,38 +431,21 @@ namespace HarmonixMetasound
 						// NOTE: This code is/should be identical to the SeekThru event above
 						if (Event.IsPreRoll)
 						{
-							UpdateLoopOffsetTickFromTick(Event.Tick2);
-							int32 Tick = Event.Tick2 - LoopOffsetTick;
-
-							float Ms = MidiClockOut->GetSongMaps().TickToMs(Tick + 1);
-							FMusicSeekTarget SeekTarget;
-							SeekTarget.Type = ESeekPointType::Millisecond;
-							SeekTarget.Ms = Ms;
-							MidiClockOut->SeekTo(Event.BlockFrameIndex, SeekTarget, PrerollBars);
+							int32 Tick = MidiClockOut->CalculateMappedTick(Event.Tick2);
+							MidiClockOut->SeekTo(Event.BlockFrameIndex, Tick + 1, PrerollBars);
 							break;
 						}
 						// The MidiClock handles looping on its own, so we can conveniently advance it
-						int32 Tick = Event.Tick2 - LoopOffsetTick;
+						int32 Tick = MidiClockOut->GetCurrentMidiTick() + (Event.Tick2 - Event.Tick1);
 						float Ms = MidiClockOut->GetSongMaps().TickToMs(Tick);
 
 						float ClockInSpeed = MidiClockIn->GetSpeedAtBlockSampleFrame(StartFrameIndex);
-						float AdvanceRatio = MidiClockIn->GetSongMaps().GetTempoAtTick(Event.Tick2) / MidiClockOut->GetSongMaps().GetTempoAtTick(Tick);
+						float AdvanceRatio = MidiClockIn->GetSongMaps().GetTempoAtTick(Event.Tick1)
+							               / MidiClockOut->GetSongMaps().GetTempoAtTick(MidiClockOut->GetCurrentMidiTick());
 						// midi clock needs to know how fast its advancing based on their authority
 						MidiClockOut->InformOfCurrentAdvanceRate(ClockInSpeed * *SpeedMultInPin * AdvanceRatio);
 						MidiClockOut->AdvanceHiResToMs(Event.BlockFrameIndex, Ms, true);
-
-						// we have to update our loop offset index _after_ the advance
-						// which we can do a little bit differently here
-						
-						// check if the clock looped back
-						if (MidiClockOut->DoesLoop())
-						{
-							int32 NewTick = MidiClockOut->GetCurrentHiResTick();
-							if (NewTick < Tick)
-							{
-								UpdateLoopOffsetTickFromTick(Event.Tick2);
-							}
-						}
+							
 						break;
 					}
 					}
@@ -579,74 +542,6 @@ namespace HarmonixMetasound
 		Init(*TransportInPin, MoveTemp(InitFn));
 	}
 
-	void FExternallyClockedMidiPlayerOperator::SetupNewMidiFile(const FMidiFileProxyPtr& NewMidi)
-	{
-		FMidiPlayerOperator::SetupNewMidiFile(NewMidi);
-		if (*LoopInPin)
-		{
-			const TArray<FMidiClockEvent>& ClockEvents = MidiClockIn->GetMidiClockEventsInBlock();
-			if (ClockEvents.Num() > 0)
-			{
-				for (const FMidiClockEvent& Event : ClockEvents)
-				{
-					// Skip loop events and reset events, just like we do in "Execute"
-					if (Event.Type == FMidiClockEvent::EType::Loop
-					 || Event.Type == FMidiClockEvent::EType::Reset)
-					{
-						continue;
-					}
-
-					// notice that we take the "Tick1" of the event, and not "Tick2".
-					// That's because we want to be synced up for when we process "Tick2" later in execute
-					UpdateLoopOffsetTickFromTick(Event.Tick1);
-					int32 Tick = Event.Tick1 - LoopOffsetTick;
-
-					float Ms = MidiClockOut->GetSongMaps().TickToMs(Tick + 1);
-					FMusicSeekTarget SeekTarget;
-					SeekTarget.Type = ESeekPointType::Millisecond;
-					SeekTarget.Ms = Ms;
-					MidiClockOut->SeekTo(SeekTarget, PrerollBars);
-					break;
-				}
-
-				
-			}
-			else
-			{
-				UpdateLoopOffsetTickFromTick(MidiClockIn->GetCurrentMidiTick());
-				int32 Tick = MidiClockIn->GetCurrentMidiTick() - LoopOffsetTick;
-				float Ms = MidiClockOut->GetSongMaps().TickToMs(Tick + 1);
-				FMusicSeekTarget SeekTarget;
-				SeekTarget.Type = ESeekPointType::Millisecond;
-				SeekTarget.Ms = Ms;
-				MidiClockOut->SeekTo(SeekTarget, PrerollBars);
-			}
-		}
-	}
-
-
-	void FExternallyClockedMidiPlayerOperator::UpdateLoopOffsetTickFromTick(int32 Tick)
-	{
-		// only update our loop offset if we're actually looping
-		// if our external clock is seeking us
-		// we need to update our loop offset ticks
-		if (MidiClockOut->DoesLoop())
-		{
-			int32 LoopStartTick = MidiClockOut->GetLoopStartTick();
-			int32 LoopEndTick = MidiClockOut->GetLoopEndTick();
-			int32 LoopLengthTicks = LoopEndTick - LoopStartTick;
-			if (LoopLengthTicks > 0)
-			{
-				// the whole number part of the division tells us 
-				// how many times we "should have looped" based on the incoming tick
-				int32 LoopNum = (Tick - LoopStartTick) / LoopLengthTicks;
-
-				// So we can multiply it back to get the LoopOffsetTick
-				LoopOffsetTick = LoopLengthTicks * LoopNum + LoopStartTick;
-			}
-		}
-	}
-
 	void FSelfClockedMidiPlayerOperator::Execute()
 	{
 		FMidiPlayerOperator::Execute();
@@ -913,6 +808,11 @@ namespace HarmonixMetasound
 				LoopEndTick = BarMap.BarIncludingCountInToTick(LoopEndBarIndex);
 				
 				MidiClockOut->SetLoop(LoopStartTick, LoopEndTick);
+				
+				// remap our current tick based on the looping behavior
+				// maybe this should happen automatically when resetting a loop or changing midi files?
+				int32 NewTick = MidiClockOut->CalculateMappedTick(CurrentTick);
+				MidiClockOut->SeekTo(NewTick, PrerollBars);
 			}
 			else
 			{
diff --git a/Engine/Plugins/Runtime/Harmonix/Source/HarmonixMetasound/Public/HarmonixMetasound/DataTypes/MidiClock.h b/Engine/Plugins/Runtime/Harmonix/Source/HarmonixMetasound/Public/HarmonixMetasound/DataTypes/MidiClock.h
--- a/Engine/Plugins/Runtime/Harmonix/Source/HarmonixMetasound/Public/HarmonixMetasound/DataTypes/MidiClock.h
+++ b/Engine/Plugins/Runtime/Harmonix/Source/HarmonixMetasound/Public/HarmonixMetasound/DataTypes/MidiClock.h
@@ -145,6 +145,7 @@ namespace HarmonixMetasound
 		void AttachToMidiResource(TSharedPtr<FMidiFileData> MidiDataProxy, bool ResetCursorsToStart = true, int32 PreRollBars = 0) { DrivingMidiPlayCursorMgr->AttachToMidiResource(MidiDataProxy, ResetCursorsToStart, PreRollBars); }
 		void DetachFromMidiResource() { DrivingMidiPlayCursorMgr->DetachFromMidiResource(); }
 		void AttachToTimeAuthority(const FMidiClock& MidiClockRef);
+		void DetachFromTimeAuthority();
 		void InformOfCurrentAdvanceRate(float AdvanceRate);
 		
 		void LockForMidiDataChanges();
@@ -154,7 +155,7 @@ namespace HarmonixMetasound
 		const FBarMap& GetBarMap() const         { return DrivingMidiPlayCursorMgr->GetBarMap();     }
 		void AdvanceHiResToMs(int32 BlockFrameIndex, float Ms, bool Broadcast);
 		void SeekTo(const FMusicSeekTarget& Timestamp, int32 PreRollBars);
-
+		void SeekTo(int32 Tick, int32 PrerollBars);
 
 		void SetLoop(int32 StartTick, int32 EndTick) { DrivingMidiPlayCursorMgr->SetLoop(StartTick, EndTick, false, true); }
 		void ClearLoop() { DrivingMidiPlayCursorMgr->ClearLoop(true); }
@@ -181,6 +182,33 @@ namespace HarmonixMetasound
 
 		void WriteAdvance(int32 StartFrameIndex, int32 EndFrameIndex, float InSpeed = 1.0f);
 		void SeekTo(int32 BlockFrameIndex, const FMusicSeekTarget& InTarget, int32 InPrerollBars);
+		void SeekTo(int32 BlockFrameIndex, int32 Tick, int32 InPrerollBars);
+		
+		/**
+		 * Given an input tick, outputs a looped tick if the input tick is > the StartTick of the Loop Region
+		 * If the clock is not looping, or loop region length is 0, then the output will be unchanged.
+		 *
+		 * The output tick will be in range [Min(Tick, LoopStartTick), LoopEndTick).
+		 * Example:
+		 * LoopRegion: (0, 100):
+		 * 10 -> 10
+		 * 100 -> 0
+		 * 110 -> 10
+		 * -10 -> 90
+		 *
+		 * LoopRegion: (40, 100):
+		 * 0 -> 0
+		 * 10 -> 10
+		 * -10 -> -10
+		 * 99  -> 99
+		 * 100 -> 40
+		 * 110 -> 50
+		 *
+		 * @param		Tick - Absolute Tick
+		 * @return		Looped Tick if Tick > LoopEnd: LoopedTick = LoopStart + (Tick - LoopStart) % (LoopEnd - LoopStart) 
+		 */
+		int32 CalculateMappedTick(int32 Tick) const;
+		
 		// copy speed and tempo changes from in clock to this clock
 		// with an optional Speed multiplier to adjust the out going speed on this clock
 		void CopySpeedAndTempoChanges(const FMidiClock* InClock, float InSpeedMult = 1.0f);
