Allow changing the LoopInPin on MidiPlayer node while it's playing.

This actually works perfectly. If you turn looping off, it'll continue playing until the end of the midi file no matter how many times it's looped already. And if you turn looping on, it'll pick up from where it should have been looping.

SetupNewMidiFile does the work of updating the loop data and seek if necessary

#tests PIE, HarmonixTests

diff --git a/Engine/Plugins/Runtime/Harmonix/Source/HarmonixMetasound/Private/HarmonixMetasound/Nodes/MidiPlayerNode.cpp b/Engine/Plugins/Runtime/Harmonix/Source/HarmonixMetasound/Private/HarmonixMetasound/Nodes/MidiPlayerNode.cpp
--- a/Engine/Plugins/Runtime/Harmonix/Source/HarmonixMetasound/Private/HarmonixMetasound/Nodes/MidiPlayerNode.cpp
+++ b/Engine/Plugins/Runtime/Harmonix/Source/HarmonixMetasound/Private/HarmonixMetasound/Nodes/MidiPlayerNode.cpp
@@ -360,7 +360,7 @@ namespace HarmonixMetasound
 		
 		MidiClockOut->PrepareBlock();
 
-		if (CurrentMidiFile != MidiAssetInPin->GetMidiProxy())
+		if (MidiClockOut->DoesLoop() != *LoopInPin || CurrentMidiFile != MidiAssetInPin->GetMidiProxy())
 		{
 			SetupNewMidiFile(MidiAssetInPin->GetMidiProxy());
 		}
