[Gentech] Add support for named fusion voice pools. Clean up FFusionVoiceConfig and turn it into a USTRUCT to allow configs to be specified via ini

#rnx
#tests Editor

#changelist validated
#virtualized

diff --git a/Engine/Plugins/Runtime/Harmonix/Source/HarmonixDsp/Private/HarmonixDsp/FusionSampler/FusionSamplerConfig.cpp b/Engine/Plugins/Runtime/Harmonix/Source/HarmonixDsp/Private/HarmonixDsp/FusionSampler/FusionSamplerConfig.cpp
new file mode 100644
--- /dev/null
+++ b/Engine/Plugins/Runtime/Harmonix/Source/HarmonixDsp/Private/HarmonixDsp/FusionSampler/FusionSamplerConfig.cpp
@@ -0,0 +1,19 @@
+// Copyright Epic Games, Inc. All Rights Reserved.
+
+#include "HarmonixDsp/FusionSampler/FusionSamplerConfig.h"
+
+FFusionVoiceConfig UFusionSamplerConfig::GetVoiceConfigForPoolName(FName PoolName) const
+{
+	if (PoolName == NAME_None)
+	{
+		return DefaultVoicePoolConfig;
+	}
+
+	const FFusionVoiceConfig* FoundConfig = VoicePoolConfigs.Find(PoolName);
+	if (!FoundConfig)
+	{
+		return DefaultVoicePoolConfig;
+	}
+
+	return *FoundConfig;
+}
\ No newline at end of file
diff --git a/Engine/Plugins/Runtime/Harmonix/Source/HarmonixDsp/Private/HarmonixDsp/FusionSampler/FusionVoice.cpp b/Engine/Plugins/Runtime/Harmonix/Source/HarmonixDsp/Private/HarmonixDsp/FusionSampler/FusionVoice.cpp
--- a/Engine/Plugins/Runtime/Harmonix/Source/HarmonixDsp/Private/HarmonixDsp/FusionSampler/FusionVoice.cpp
+++ b/Engine/Plugins/Runtime/Harmonix/Source/HarmonixDsp/Private/HarmonixDsp/FusionSampler/FusionVoice.cpp
@@ -218,7 +218,7 @@ void FFusionVoice::Kill()
 	}
 }
 
-void FFusionVoice::Init(FFusionVoicePool* InOwner, uint32 InDebugID, bool bDecompressSamplesOnLoad)
+void FFusionVoice::Init(FFusionVoicePool* InOwner, uint32 InDebugID)
 {
 	VoicePool = InOwner;
 	DebugID = InDebugID;
diff --git a/Engine/Plugins/Runtime/Harmonix/Source/HarmonixDsp/Private/HarmonixDsp/FusionSampler/FusionVoicePool.cpp b/Engine/Plugins/Runtime/Harmonix/Source/HarmonixDsp/Private/HarmonixDsp/FusionSampler/FusionVoicePool.cpp
--- a/Engine/Plugins/Runtime/Harmonix/Source/HarmonixDsp/Private/HarmonixDsp/FusionSampler/FusionVoicePool.cpp
+++ b/Engine/Plugins/Runtime/Harmonix/Source/HarmonixDsp/Private/HarmonixDsp/FusionSampler/FusionVoicePool.cpp
@@ -3,14 +3,14 @@
 #include "HarmonixDsp/FusionSampler/FusionVoice.h"
 #include "HarmonixDsp/FusionSampler/FusionPatch.h"
 #include "HarmonixDsp/FusionSampler/FusionSampler.h"
+#include "HarmonixDsp/FusionSampler/FusionSamplerConfig.h"
 #include "HarmonixDsp/FusionSampler/Settings/KeyzoneSettings.h"
 #include "HarmonixDsp/StretcherAndPitchShifterFactory.h"
 #include "Sound/SoundWave.h"
 
 DEFINE_LOG_CATEGORY(LogFusionVoicePool);
 
-FFusionVoiceConfig FFusionVoiceConfig::DefaultConfig(32, 24, false);
-FFusionVoicePool::FPoolMap FFusionVoicePool::gDefaultVoicePools;
+FFusionVoicePool::FPoolMap FFusionVoicePool::GVoicePools;
 
 void FFusionVoicePool::SetSampleRate(float InSampleRate)
 {
@@ -41,10 +41,15 @@ void FFusionVoicePool::Unlock()
 
 FSharedFusionVoicePoolPtr FFusionVoicePool::GetDefault(float InSampleRate)
 {
-	static FCriticalSection sDefaultVoicePoolLock;
-	FScopeLock Lock(&sDefaultVoicePoolLock);
+	return GetNamedPool(NAME_None, InSampleRate);
+}
+
+FSharedFusionVoicePoolPtr FFusionVoicePool::GetNamedPool(FName InPoolName, float InSampleRate)
+{
+	static FCriticalSection sVoicePoolLock;
+	FScopeLock Lock(&sVoicePoolLock);
 
-	gDefaultVoicePools = gDefaultVoicePools.FilterByPredicate([](const FPoolMap::ElementType& Pair)
+	GVoicePools = GVoicePools.FilterByPredicate([](const FPoolMap::ElementType& Pair)
 		{
 			if (FSharedFusionVoicePoolPtr SharedPool = (Pair.Value).Pin())
 			{
@@ -53,31 +58,26 @@ FSharedFusionVoicePoolPtr FFusionVoicePool::GetDefault(float InSampleRate)
 			return false;
 		});
 
-	int32 SampleRateInt = FMath::FloorToInt32(InSampleRate);
-	if (TWeakPtr<FFusionVoicePool, ESPMode::ThreadSafe>* WeakPoolPtr = gDefaultVoicePools.Find(SampleRateInt))
+	FPoolMapKey Key = MakeTuple(InPoolName, FMath::FloorToInt32(InSampleRate));
+	if (TWeakPtr<FFusionVoicePool, ESPMode::ThreadSafe>* WeakPoolPtr = GVoicePools.Find(Key))
 	{
 		if (FSharedFusionVoicePoolPtr SharedPool = (*WeakPoolPtr).Pin())
 		{
 			return SharedPool;
 		}
-		else 
-		{
-			gDefaultVoicePools.Remove(SampleRateInt);
-		}
 	}
 
-	FFusionVoiceConfig Config = FFusionVoiceConfig::DefaultConfig;
-	Config.SampleRate = InSampleRate;
-	FSharedFusionVoicePoolPtr NewPool = Create(Config);
-	gDefaultVoicePools.Add(SampleRateInt, NewPool.ToWeakPtr());
+	const UFusionSamplerConfig* FusionConfig = ::GetDefault<UFusionSamplerConfig>();
+
+	FSharedFusionVoicePoolPtr NewPool = Create(FusionConfig->GetVoiceConfigForPoolName(InPoolName), InSampleRate);
+	GVoicePools.Add(Key, NewPool.ToWeakPtr());
 	return NewPool;
 }
 
-FSharedFusionVoicePoolPtr FFusionVoicePool::Create(const FFusionVoiceConfig& InConfig)
+FSharedFusionVoicePoolPtr FFusionVoicePool::Create(const FFusionVoiceConfig& InConfig, float InSampleRate)
 {
-	FSharedFusionVoicePoolPtr NewVoicePool = MakeShared<FFusionVoicePool, ESPMode::ThreadSafe>(InConfig.SampleRate);
+	FSharedFusionVoicePoolPtr NewVoicePool = MakeShared<FFusionVoicePool, ESPMode::ThreadSafe>(InSampleRate);
 
-	NewVoicePool->SetDecompressSamplesOnLoad(InConfig.DecompressSamplesOnLoad);
 	NewVoicePool->SetHardVoiceLimit(InConfig.NumTotalVoices);
 	NewVoicePool->SetSoftVoiceLimit(InConfig.SoftVoiceLimit);
 	NewVoicePool->SetFormantVolumeCorrection(
@@ -246,15 +246,6 @@ void FFusionVoicePool::SetHardVoiceLimit(uint32 NewPolyphony)
 	}
 }
 
-
-void FFusionVoicePool::SetDecompressSamplesOnLoad(bool InDecompress)
-{
-	// NOTE: we don't recreate voices here to save time. SetHardVoiceLimit()
-	// must be called after this to take effect.
-	DecompressSamplesOnLoad = InDecompress;
-}
-
-
 void FFusionVoicePool::SetFormantVolumeCorrection(float DBPerHalfStepUp, float DBPerHalfStepDown, float DBMaxUp, float DBMaxDown)
 {
 	for (IStretcherAndPitchShifterFactory* Factory : IStretcherAndPitchShifterFactory::GetAllRegisteredFactories())
@@ -671,7 +662,7 @@ void FFusionVoicePool::CreateVoices(uint16 InNumToAllocate)
 
 	for (uint16 VoiceIdx = 0; VoiceIdx < InNumToAllocate; ++VoiceIdx)
 	{
-		Voices[VoiceIdx].Init(this, VoiceIdx, DecompressSamplesOnLoad);
+		Voices[VoiceIdx].Init(this, VoiceIdx);
 	}
 
 	PeakVoiceUsage = 0;
diff --git a/Engine/Plugins/Runtime/Harmonix/Source/HarmonixDsp/Public/HarmonixDsp/FusionSampler/FusionSamplerConfig.h b/Engine/Plugins/Runtime/Harmonix/Source/HarmonixDsp/Public/HarmonixDsp/FusionSampler/FusionSamplerConfig.h
new file mode 100644
--- /dev/null
+++ b/Engine/Plugins/Runtime/Harmonix/Source/HarmonixDsp/Public/HarmonixDsp/FusionSampler/FusionSamplerConfig.h
@@ -0,0 +1,24 @@
+// Copyright Epic Games, Inc. All Rights Reserved.
+
+#pragma once
+
+#include "Harmonix/HarmonixDeveloperSettings.h"
+#include "HarmonixDsp/FusionSampler/FusionVoicePool.h"
+
+#include "FusionSamplerConfig.generated.h"
+
+UCLASS(Config = Engine, defaultconfig, meta = (DisplayName = "Fusion Sampler Settings"))
+class HARMONIXDSP_API UFusionSamplerConfig : public UHarmonixDeveloperSettings
+{
+	GENERATED_BODY()
+
+public:
+
+	UPROPERTY(config, EditDefaultsOnly, Category = "Voice Pools")
+	FFusionVoiceConfig DefaultVoicePoolConfig;
+
+	UPROPERTY(config, EditDefaultsOnly, Category = "Voice Pools")
+	TMap<FName, FFusionVoiceConfig> VoicePoolConfigs;
+
+	FFusionVoiceConfig GetVoiceConfigForPoolName(FName PoolName) const;
+};
\ No newline at end of file
diff --git a/Engine/Plugins/Runtime/Harmonix/Source/HarmonixDsp/Public/HarmonixDsp/FusionSampler/FusionVoice.h b/Engine/Plugins/Runtime/Harmonix/Source/HarmonixDsp/Public/HarmonixDsp/FusionSampler/FusionVoice.h
--- a/Engine/Plugins/Runtime/Harmonix/Source/HarmonixDsp/Public/HarmonixDsp/FusionSampler/FusionVoice.h
+++ b/Engine/Plugins/Runtime/Harmonix/Source/HarmonixDsp/Public/HarmonixDsp/FusionSampler/FusionVoice.h
@@ -34,7 +34,7 @@ public:
 	FFusionVoice();
 	virtual ~FFusionVoice() = default;
 
-	void Init(FFusionVoicePool* InOwner, uint32 InDebugID, bool bDecompressSamplesOnLoad);
+	void Init(FFusionVoicePool* InOwner, uint32 InDebugID);
 
 	static const double kMaxPitchOffsetCents;
 
diff --git a/Engine/Plugins/Runtime/Harmonix/Source/HarmonixDsp/Public/HarmonixDsp/FusionSampler/FusionVoicePool.h b/Engine/Plugins/Runtime/Harmonix/Source/HarmonixDsp/Public/HarmonixDsp/FusionSampler/FusionVoicePool.h
--- a/Engine/Plugins/Runtime/Harmonix/Source/HarmonixDsp/Public/HarmonixDsp/FusionSampler/FusionVoicePool.h
+++ b/Engine/Plugins/Runtime/Harmonix/Source/HarmonixDsp/Public/HarmonixDsp/FusionSampler/FusionVoicePool.h
@@ -3,9 +3,12 @@
 
 #include "HAL/Platform.h"
 #include "Templates/SharedPointer.h"
+#include "Templates/Tuple.h"
 #include "HarmonixDsp/StretcherAndPitchShifterFactory.h"
 #include "HarmonixMidi/MidiVoiceId.h"
 
+#include "FusionVoicePool.generated.h"
+
 struct FKeyzoneSettings;
 class FFusionSampler;
 class FFusionVoice;
@@ -15,31 +18,37 @@ using FSharedFusionVoicePoolPtr = TSharedPtr<FFusionVoicePool, ESPMode::ThreadSa
 
 DECLARE_LOG_CATEGORY_EXTERN(LogFusionVoicePool, Log, All);
 
+USTRUCT()
 struct HARMONIXDSP_API FFusionVoiceConfig
 {
-	float SampleRate = 48000.0f;
+	GENERATED_BODY()
+
+	UPROPERTY(EditAnywhere, Category = "Voice Pool")
+	uint32 NumTotalVoices = 32;
 
-	uint32 NumTotalVoices = 0;
+	UPROPERTY(EditAnywhere, Category = "Voice Pool")
 	uint32 SoftVoiceLimit = 24;
-	bool   DecompressSamplesOnLoad = true;
-	bool   DecompressAsync = true;
-	int32  MaxActiveDecoders = -1;
-	int32  AsynchronousDecodePrebufferFrames = 48000;
-	bool   PreserveCompressedData = false;
+
+	UPROPERTY(EditAnywhere, Category = "Pitch Shifter")
 	float  FormantDbCorrectionPerHalfStepUp = 0.0f;
+
+	UPROPERTY(EditAnywhere, Category = "Pitch Shifter")
 	float  FormantDbCorrectionPerHalfStepDown = 0.0f;
+
+	UPROPERTY(EditAnywhere, Category = "Pitch Shifter")
 	float  FormantDbCorrectionMaxUp = 12.0f;
+
+	UPROPERTY(EditAnywhere, Category = "Pitch Shifter")
 	float  FormantDbCorrectionMaxDown = -12.0f;
 
+	FFusionVoiceConfig() {}
+
 	FFusionVoiceConfig(
 		uint32 InNumTotalVoices,
-		uint32 InSoftVoiceLimit,
-		bool InDecompressSamplesOnLoad)
+		uint32 InSoftVoiceLimit)
 		: NumTotalVoices(InNumTotalVoices)
 		, SoftVoiceLimit(InSoftVoiceLimit)
 	{}
-
-	static FFusionVoiceConfig DefaultConfig;
 };
 
 class HARMONIXDSP_API FFusionVoicePool
@@ -47,7 +56,8 @@ class HARMONIXDSP_API FFusionVoicePool
 public:
 
 	static FSharedFusionVoicePoolPtr GetDefault(float InSampleRate);
-	static FSharedFusionVoicePoolPtr Create(const FFusionVoiceConfig& InConfig);
+	static FSharedFusionVoicePoolPtr GetNamedPool(FName InPoolName, float InSampleRate);
+	static FSharedFusionVoicePoolPtr Create(const FFusionVoiceConfig& InConfig, float InSampleRate);
 
 	FFusionVoicePool(float InSampleRate) 
 		: Voices(nullptr)
@@ -55,7 +65,6 @@ public:
 		, NumVoicesSetting(kDefaultPoolSize)
 		, SoftVoiceLimit(kDefaultPoolSize)
 		, PeakVoiceUsage(0)
-		, DecompressSamplesOnLoad(false)
 		, SampleRate(InSampleRate)
 	{};
 
@@ -112,10 +121,6 @@ public:
 	 */
 	uint32 GetSoftVoiceLimit() const { return SoftVoiceLimit; }
 
-	// if we're decompressing samples at load time, we don't need to set up per-voice mogg decoders 
-	//  (which are RAM-intensive)
-	void SetDecompressSamplesOnLoad(bool Decompress);
-
 	void SetFormantVolumeCorrection(float dBperHalfStepUp, float dBperHalfStepDown, float dBMaxUp, float  dBMaxDown);
 
 	uint32 GetPeakVoiceUsage() 
@@ -172,10 +177,11 @@ public:
 
 private:
 
-	// keep a weak ptr, and pass around shared ptrs
-	// so it gets automatically destroyed when nothing references it anymore
-	using FPoolMap = TMap<int32, TWeakPtr<FFusionVoicePool, ESPMode::ThreadSafe>>;
-	static FPoolMap gDefaultVoicePools;
+	// Key is a tuple of PoolName and SampleRate
+	// Value is a weak ptr, but shared ptrs are returned so it gets automatically destroyed when nothing references it anymore
+	using FPoolMapKey = TTuple<FName, int32>;
+	using FPoolMap = TMap<FPoolMapKey, TWeakPtr<FFusionVoicePool, ESPMode::ThreadSafe>>;
+	static FPoolMap GVoicePools;
 
 	FFusionVoice* Voices;
 
@@ -186,8 +192,6 @@ private:
 
 	uint32     PeakVoiceUsage = 0;
 
-	bool       DecompressSamplesOnLoad = false;
-
 	void AllocVoicesAndShifters();
 	void CreateVoices(uint16 InMaxPolyphony);
 	void CreateShifters();
diff --git a/Engine/Plugins/Runtime/Harmonix/Source/HarmonixMetasound/Private/HarmonixMetasound/Nodes/FusionSamplerNode.cpp b/Engine/Plugins/Runtime/Harmonix/Source/HarmonixMetasound/Private/HarmonixMetasound/Nodes/FusionSamplerNode.cpp
--- a/Engine/Plugins/Runtime/Harmonix/Source/HarmonixMetasound/Private/HarmonixMetasound/Nodes/FusionSamplerNode.cpp
+++ b/Engine/Plugins/Runtime/Harmonix/Source/HarmonixMetasound/Private/HarmonixMetasound/Nodes/FusionSamplerNode.cpp
@@ -37,6 +37,7 @@ namespace HarmonixMetasound
 		METASOUND_PARAM(FineTuneCents, "Fine Tune Cents", "Adds to the patch's fine tune cents setting if connected AND if the vaue is > -1200 and < 1200.")
 		METASOUND_PARAM(EnableMTFusion, "Multithreaded Rendering", "Turn on to allow Fusion rendering to be done across multiple tasks. NOTE: YOU MUST then connect this node's "
 		                                 "Audio Output AND Render Sync Output to a Fusion Synchronizer Node, and reference THAT NODE'S audio outputs elsewhere in the graph!")
+		METASOUND_PARAM(VoicePoolName, "Voice Pool Name", "The name of a specific voice pool for this node to use. Uses the default pool if the named pool does not exist, or if the name is \"None\"")
 		METASOUND_PARAM(RenderSync, "Render Sync", "YOU MUST connect this AND this node's audio output to a Fusion Synchronizer Node if you have enabled Multithreaded Rendering for this node!")
 	}
 
@@ -72,6 +73,7 @@ namespace HarmonixMetasound
 			FFusionPatchAssetReadRef InPatch;
 			FBoolReadRef             InClockSpeedToPitch;
 			bool		             InMTRenderingEnable;
+			FString					 InVoicePoolName;
 			EAudioBufferChannelLayout InOutputChannelLayout;
 
 			FConstructionArgs(const FBuildOperatorParams& InParams, EAudioBufferChannelLayout OutputChannelLayout)
@@ -89,6 +91,7 @@ namespace HarmonixMetasound
 				, InPatch(InParams.InputData.GetOrConstructDataReadReference<FFusionPatchAsset>(METASOUND_GET_PARAM_NAME(CommonPinNames::Inputs::SynthPatch)))
 				, InClockSpeedToPitch(InParams.InputData.GetOrCreateDefaultDataReadReference<bool>(METASOUND_GET_PARAM_NAME(CommonPinNames::Inputs::ClockSpeedToPitch), InParams.OperatorSettings))
 				, InMTRenderingEnable(InParams.InputData.GetOrCreateDefaultValue<bool>(METASOUND_GET_PARAM_NAME(FusionSamplerNodePinNames::EnableMTFusion), InParams.OperatorSettings))
+				, InVoicePoolName(InParams.InputData.GetOrCreateDefaultValue<FString>(METASOUND_GET_PARAM_NAME(FusionSamplerNodePinNames::VoicePoolName), InParams.OperatorSettings))
 				, InOutputChannelLayout(OutputChannelLayout)
 			{
 			}
@@ -110,7 +113,8 @@ namespace HarmonixMetasound
 				TInputDataVertex<float>(METASOUND_GET_PARAM_NAME_AND_METADATA(FusionSamplerNodePinNames::Lfo1DepthOverride), -1.0f),
 				TInputDataVertex<float>(METASOUND_GET_PARAM_NAME_AND_METADATA(FusionSamplerNodePinNames::FineTuneCents), -1200.0f),
 				TInputDataVertex<bool>(METASOUND_GET_PARAM_NAME_AND_METADATA(Inputs::ClockSpeedToPitch), true),
-				TInputConstructorVertex<bool>(METASOUND_GET_PARAM_NAME_AND_METADATA_ADVANCED(FusionSamplerNodePinNames::EnableMTFusion), false)
+				TInputConstructorVertex<bool>(METASOUND_GET_PARAM_NAME_AND_METADATA_ADVANCED(FusionSamplerNodePinNames::EnableMTFusion), false),
+				TInputConstructorVertex<FString>(METASOUND_GET_PARAM_NAME_AND_METADATA_ADVANCED(FusionSamplerNodePinNames::VoicePoolName), FString(TEXT("")))
 			);
 
 			return InputInterface;
@@ -130,6 +134,7 @@ namespace HarmonixMetasound
 			, PatchInPin(Args.InPatch)
 			, ClockSpeedAffectsPitchInPin(Args.InClockSpeedToPitch)
 			, EnableMTRenderingInPin(Args.InMTRenderingEnable)
+			, VoicePoolName(Args.InVoicePoolName)
 			, OutputChannelLayout(Args.InOutputChannelLayout)
 		{
 			Reset(Args.Params);
@@ -163,6 +168,7 @@ namespace HarmonixMetasound
 			InVertexData.BindReadVertex(METASOUND_GET_PARAM_NAME(Inputs::SynthPatch), PatchInPin);
 			InVertexData.BindReadVertex(METASOUND_GET_PARAM_NAME(Inputs::ClockSpeedToPitch), ClockSpeedAffectsPitchInPin);
 			InVertexData.SetValue(METASOUND_GET_PARAM_NAME(FusionSamplerNodePinNames::EnableMTFusion), EnableMTRenderingInPin);
+			InVertexData.SetValue(METASOUND_GET_PARAM_NAME(FusionSamplerNodePinNames::VoicePoolName), VoicePoolName);
 
 			Init();
 		}
@@ -177,7 +183,7 @@ namespace HarmonixMetasound
 			ResetInstrumentState();
 			SetSampleRate(Params.OperatorSettings.GetSampleRate());
 			
-			SetVoicePool(FFusionVoicePool::GetDefault(Params.OperatorSettings.GetSampleRate()), false);
+			SetVoicePool(FFusionVoicePool::GetNamedPool(FName(VoicePoolName), Params.OperatorSettings.GetSampleRate()), false);
 		}
 
 		void Execute();
@@ -209,6 +215,7 @@ namespace HarmonixMetasound
 		FFusionPatchAssetReadRef PatchInPin;
 		FBoolReadRef ClockSpeedAffectsPitchInPin;
 		bool EnableMTRenderingInPin;
+		FString VoicePoolName;
 
 		//** DATA
 		int32 FramesPerBlock = 0;
