StepSequncer: Handle clock events instead of being a midi play cursor

#tests PIE, Harmonix Tests

diff --git a/Engine/Plugins/Runtime/Harmonix/Source/HarmonixMetasound/Private/HarmonixMetasound/Nodes/StepSequencePlayerNode.cpp b/Engine/Plugins/Runtime/Harmonix/Source/HarmonixMetasound/Private/HarmonixMetasound/Nodes/StepSequencePlayerNode.cpp
--- a/Engine/Plugins/Runtime/Harmonix/Source/HarmonixMetasound/Private/HarmonixMetasound/Nodes/StepSequencePlayerNode.cpp
+++ b/Engine/Plugins/Runtime/Harmonix/Source/HarmonixMetasound/Private/HarmonixMetasound/Nodes/StepSequencePlayerNode.cpp
@@ -29,7 +29,7 @@ namespace HarmonixMetasound
 	using namespace Metasound;
 	using namespace Harmonix;
 
-	class FStepSequencePlayerOperator : public TExecutableOperator<FStepSequencePlayerOperator>, public FMusicTransportControllable, public FMidiPlayCursor, public FMidiVoiceGeneratorBase
+	class FStepSequencePlayerOperator : public TExecutableOperator<FStepSequencePlayerOperator>, public FMusicTransportControllable, public FMidiVoiceGeneratorBase
 	{
 	public:
 		static const FNodeClassMetadata& GetNodeInfo();
@@ -62,15 +62,6 @@ namespace HarmonixMetasound
 
 	private:
 		void Init();
-
-		void PrepareMidiOutIfNeeded()
-		{
-			if (bNeedsMidiOutClear)
-			{
-				MidiOutPin->PrepareBlock();
-				bNeedsMidiOutClear = false;
-			}
-		}
 		
 		//** INPUTS
 		FMidiStepSequenceAssetReadRef SequenceAssetInPin;
@@ -99,7 +90,6 @@ namespace HarmonixMetasound
 		int32 ProcessedThruTick			= -1;
 		int32 SequenceStartTick			= -1;
 		int32 CurrentStepSkipIndex		= 0;
-		bool  bNeedsMidiOutClear		= true;
 		bool  bAutoPage					= false;
 		bool  bPreviousAutoPage			= false;
 		bool  bAutoPagePlaysBlankPages	= false;
@@ -119,13 +109,11 @@ namespace HarmonixMetasound
 		void EnsureCurrentPageIndexIsValid();
 
 	protected:
-		//** BEGIN FMidiPlayCursor
-		virtual void SeekToTick(int32 Tick) override;
-		virtual void SeekThruTick(int32 Tick) override;
-		virtual void AdvanceThruTick(int32 Tick, bool IsPreRoll) override;
-		// We have to override this to disambiguate the FMidiPlayCursor method and the MS operator Reset above
-		virtual void Reset(const bool ForceNoBroadcast) override { FMidiPlayCursor::Reset(ForceNoBroadcast); }
-		//** END FMidiPlayCursor
+
+		void SeekToTick(int32 BlockFrameIndex, int32 Tick);
+		void SeekThruTick(int32 BlockFrameIndex, int32 Tick);
+		void AdvanceThruTick(int32 BlockFrameIndex, int32 Tick, bool IsPreRoll);
+
 	};
 
 	class FStepSequencePlayerNode : public FNodeFacade
@@ -265,13 +253,10 @@ namespace HarmonixMetasound
 
 	FStepSequencePlayerOperator::~FStepSequencePlayerOperator()
 	{
-		MidiClockInPin->UnregisterPlayCursor(this);
 	}
 
 	void FStepSequencePlayerOperator::BindInputs(FInputVertexInterfaceData& InVertexData)
 	{
-		MidiClockInPin->UnregisterPlayCursor(this);
-
 		using namespace StepSequencePlayerPinNames;
 		using namespace CommonPinNames;
 		InVertexData.BindReadVertex(METASOUND_GET_PARAM_NAME(InputSequenceAsset), SequenceAssetInPin);
@@ -321,7 +306,6 @@ namespace HarmonixMetasound
 		CurrentCellIndex = -1;
 		ProcessedThruTick = -1;
 		SequenceStartTick = -1;
-		bNeedsMidiOutClear = true;
 		bAutoPage = false;
 		bPreviousAutoPage = false;
 		bAutoPagePlaysBlankPages = false;
@@ -334,9 +318,7 @@ namespace HarmonixMetasound
 
 	void FStepSequencePlayerOperator::Execute()
 	{
-		// If no midi cursor callbacks happened we still have to 
-		// remove old midi events from last block...
-		PrepareMidiOutIfNeeded();
+		MidiOutPin->PrepareBlock();
 
 		// if we have no sequence table there is nothing to do. 
 		// Make sure the notes are all off and return.
@@ -356,6 +338,38 @@ namespace HarmonixMetasound
 		// sequence table asset changes while we are in the middle of rendering. 
 		CurrentStepSkipIndex = SequenceTable->StepSkipIndex;
 
+		TransportSpanPostProcessor ClockEventHandler = [&](int32 StartFrameIndex, int32 EndFrameIndex, EMusicPlayerTransportState State)
+			{
+				switch (State)
+				{
+				case EMusicPlayerTransportState::Playing:
+				case EMusicPlayerTransportState::Continuing:
+					for (const FMidiClockEvent& Event : MidiClockInPin->GetMidiClockEventsInBlock())
+					{
+						if (Event.BlockFrameIndex >= EndFrameIndex)
+						{
+							break;
+						}
+
+						if (Event.BlockFrameIndex >= StartFrameIndex)
+						{
+							switch (Event.Type)
+							{
+							case FMidiClockEvent::EType::AdvanceThru:
+								AdvanceThruTick(Event.BlockFrameIndex, Event.Tick2, Event.IsPreRoll);
+								break;
+							case FMidiClockEvent::EType::SeekThru:
+								SeekThruTick(Event.BlockFrameIndex, Event.Tick2);
+								break;
+							case FMidiClockEvent::EType::SeekTo:
+								SeekToTick(Event.BlockFrameIndex, Event.Tick2);
+								break;
+							}
+						}
+					}
+				}
+			};
+
 		TransportSpanProcessor TransportHandler = [&](int32 StartFrameIndex, int32 EndFrameIndex, EMusicPlayerTransportState CurrentState)
 			{
 				CurrentBlockSpanStart = StartFrameIndex;
@@ -399,15 +413,12 @@ namespace HarmonixMetasound
 					return EMusicPlayerTransportState::Invalid;
 				}
 			};
-		ExecuteTransportSpans(TransportInPin, BlockSize, TransportHandler);
-
-		bNeedsMidiOutClear = true;
+		ExecuteTransportSpans(TransportInPin, BlockSize, TransportHandler, ClockEventHandler);
 	}
 
 	void FStepSequencePlayerOperator::Init()
 	{
 		MidiOutPin->SetClock(*MidiClockInPin);
-		MidiClockInPin->RegisterHiResPlayCursor(this);
 		MidiOutPin->PrepareBlock();
 
 		SequenceTable = SequenceAssetInPin->GetRenderable();
@@ -470,19 +481,16 @@ namespace HarmonixMetasound
 		FMusicTransportControllable::Init(*TransportInPin, MoveTemp(InitFn));
 	}
 
-	void FStepSequencePlayerOperator::SeekToTick(int32 Tick)
+	void FStepSequencePlayerOperator::SeekToTick(int32 BlockFrameIndex, int32 Tick)
 	{
-		SeekThruTick(Tick-1);
+		SeekThruTick(BlockFrameIndex, Tick-1);
 	}
 
-	void FStepSequencePlayerOperator::SeekThruTick(int32 Tick)
+	void FStepSequencePlayerOperator::SeekThruTick(int32 BlockFrameIndex ,int32 Tick)
 	{
-		PrepareMidiOutIfNeeded();
-
-		int32 BlockFrame = MidiClockInPin->GetCurrentBlockFrameIndex();
 		ProcessedThruTick = FMath::Max(-1, Tick);
 		bNeedsRebase = true;
-		AllNotesOff(BlockFrame, Tick, true);
+		AllNotesOff(BlockFrameIndex, Tick, true);
 	}
 
 	int32 FStepSequencePlayerOperator::CalculatePagesProgressed(const int FromTick, const int ToTick, int32 TableTickLength, const bool bRound) const
@@ -564,21 +572,18 @@ namespace HarmonixMetasound
 		}
 	}
 
-	void FStepSequencePlayerOperator::AdvanceThruTick(int32 Tick, bool IsPreRoll)
+	void FStepSequencePlayerOperator::AdvanceThruTick(int32 BlockFrameIndex, int32 Tick, bool IsPreRoll)
 	{
 		if (!SequenceTable || SequenceTable->Pages.IsEmpty())
 		{
 			return;
 		}
 
-		PrepareMidiOutIfNeeded();
-
 		// Read input pins
 		const int32 CurrentMaxColumns = *MaxColumnsInPin;
 		const int32 AdditionalOctaveNotes = (int32)*AdditionalOctavesInPin * 12;
 		const float CurrentStepSizeQuarterNotes = *StepSizeQuarterNotesInPin;
 		const float CurrentVelocityMultiplierValue = *VelocityMultInPin;
-		const int32 BlockFrameIndex = MidiClockInPin->GetCurrentBlockFrameIndex();
 
 		if (CurrentMaxColumns < 1)
 		{
