Make the StepSequencer generate midi events using the original note and then transposing the midi message after.

#tests HarmonixTests,PIE,Creative

diff --git a/Engine/Plugins/Runtime/Harmonix/Source/HarmonixMetasound/Private/HarmonixMetasound/Nodes/StepSequencePlayerNode.cpp b/Engine/Plugins/Runtime/Harmonix/Source/HarmonixMetasound/Private/HarmonixMetasound/Nodes/StepSequencePlayerNode.cpp
--- a/Engine/Plugins/Runtime/Harmonix/Source/HarmonixMetasound/Private/HarmonixMetasound/Nodes/StepSequencePlayerNode.cpp
+++ b/Engine/Plugins/Runtime/Harmonix/Source/HarmonixMetasound/Private/HarmonixMetasound/Nodes/StepSequencePlayerNode.cpp
@@ -775,7 +775,12 @@ namespace HarmonixMetasound
 						}
 
 						// note off!
-						FMidiStreamEvent MidiEvent(CurrentCellNotes[i].GetGeneratorId(), FMidiMsg::CreateNoteOff(MidiCh, SequenceTable->Notes[i].NoteNumber + AdditionalOctaveNotes));
+						int32 OriginalNote = SequenceTable->Notes[i].NoteNumber;
+						int32 TransposedNote = SequenceTable->Notes[i].NoteNumber + AdditionalOctaveNotes;
+						// create the midi event with the original note to maintain voice ids.
+						FMidiStreamEvent MidiEvent(CurrentCellNotes[i].GetGeneratorId(), FMidiMsg::CreateNoteOff(MidiCh, OriginalNote));
+						// and then assign the note directly to the midi message after
+						MidiEvent.MidiMessage.Data1 = FMath::Clamp(TransposedNote, 0, 127);
 						MidiEvent.BlockSampleFrameIndex = BlockFrameIndex;
 						MidiEvent.AuthoredMidiTick = ProcessedThruTick;
 						MidiEvent.CurrentMidiTick = ProcessedThruTick;
@@ -796,7 +801,12 @@ namespace HarmonixMetasound
 						if (!CurrentCellNotes[i])
 						{
 							// note on!
-							FMidiStreamEvent MidiEvent(this, FMidiMsg::CreateNoteOn(0, SequenceTable->Notes[i].NoteNumber + AdditionalOctaveNotes, SequenceTable->Notes[i].Velocity));
+							int32 OriginalNote = SequenceTable->Notes[i].NoteNumber;
+							int32 TransposedNote = SequenceTable->Notes[i].NoteNumber + AdditionalOctaveNotes;
+							// create the midi event with the original note to maintain voice ids.
+							FMidiStreamEvent MidiEvent(this, FMidiMsg::CreateNoteOn(0, OriginalNote, SequenceTable->Notes[i].Velocity));
+							// and then assign the note directly to the midi message after
+							MidiEvent.MidiMessage.Data1 = FMath::Clamp(TransposedNote, 0, 127);
 							float NoteOnVelocity = FMath::Clamp(static_cast<float>(SequenceTable->Notes[i].Velocity) * CurrentVelocityMultiplierValue, 0.0f, 127.0f);
 							MidiEvent.MidiMessage.SetNoteOnVelocity(static_cast<uint8>(NoteOnVelocity));
 							MidiEvent.BlockSampleFrameIndex = BlockFrameIndex;
