- Allow midi files that are not an even multiple of whole bars
- Allow more sophisticated "conforming" of midi files lengths when they are imported and/or from the details panel.
- Fix bugs in midi file length conforming logic

diff --git a/Engine/Plugins/Runtime/Harmonix/Config/BaseHarmonix.ini b/Engine/Plugins/Runtime/Harmonix/Config/BaseHarmonix.ini
--- a/Engine/Plugins/Runtime/Harmonix/Config/BaseHarmonix.ini
+++ b/Engine/Plugins/Runtime/Harmonix/Config/BaseHarmonix.ini
@@ -21,3 +21,5 @@
 
 +ClassRedirects=(OldName="/Script/HarmonixAudioReactivity.HarmonixPeakTamer",NewName="/Script/HarmonixMetasound.HarmonixPeakTamer")
 +StructRedirects=(OldName="/Script/HarmonixAudioReactivity.HarmonixPeakTamerSettings",NewName="/Script/HarmonixMetasound.HarmonixPeakTamerSettings")
+
++PropertyRedirects=(OldName="SongLengthData.LengthBars", NewName="LengthFractionalBars")
diff --git a/Engine/Plugins/Runtime/Harmonix/Source/HarmonixMetasound/Private/HarmonixMetasound/Components/MusicClockComponent.cpp b/Engine/Plugins/Runtime/Harmonix/Source/HarmonixMetasound/Private/HarmonixMetasound/Components/MusicClockComponent.cpp
--- a/Engine/Plugins/Runtime/Harmonix/Source/HarmonixMetasound/Private/HarmonixMetasound/Components/MusicClockComponent.cpp
+++ b/Engine/Plugins/Runtime/Harmonix/Source/HarmonixMetasound/Private/HarmonixMetasound/Components/MusicClockComponent.cpp
@@ -521,9 +521,9 @@ float UMusicClockComponent::GetSongLengthBeats() const
 	return GetSongMaps().GetSongLengthBeats();
 }
 
-int32 UMusicClockComponent::GetSongLengthBars() const
+float UMusicClockComponent::GetSongLengthBars() const
 {
-	return GetSongMaps().GetSongLengthBars();
+	return GetSongMaps().GetSongLengthFractionalBars();
 }
 
 float UMusicClockComponent::GetSongRemainingMs(ECalibratedMusicTimebase Timebase) const
diff --git a/Engine/Plugins/Runtime/Harmonix/Source/HarmonixMetasound/Private/HarmonixMetasound/DataTypes/MidiClock.cpp b/Engine/Plugins/Runtime/Harmonix/Source/HarmonixMetasound/Private/HarmonixMetasound/DataTypes/MidiClock.cpp
--- a/Engine/Plugins/Runtime/Harmonix/Source/HarmonixMetasound/Private/HarmonixMetasound/DataTypes/MidiClock.cpp
+++ b/Engine/Plugins/Runtime/Harmonix/Source/HarmonixMetasound/Private/HarmonixMetasound/DataTypes/MidiClock.cpp
@@ -466,7 +466,8 @@ namespace HarmonixMetasound
 
 		// create conductor track
 		FMidiTrack& Track = OutMidiData->Tracks.Add_GetRef(FMidiTrack(TEXT("conductor")));
-		OutMidiData->LastEventTick = std::numeric_limits<int32>::max();
+		// max out song length data so the midi can play indefinitely
+		OutMidiData->ConformToLength(std::numeric_limits<int32>::max());
 
 		// add time sig info
 		int32 TimeSigNum = FMath::Clamp(InTimeSigNum, 1, 64);
@@ -482,11 +483,6 @@ namespace HarmonixMetasound
 
 		Track.Sort();
 
-		// max out song length data so the midi can play indefinitely
-		OutMidiData->SongMaps.GetSongLengthData().LastTick = std::numeric_limits<int32>::max();
-		OutMidiData->SongMaps.GetSongLengthData().LengthTicks = std::numeric_limits<int32>::max();
-		OutMidiData->SongMaps.GetSongLengthData().LengthBars = std::numeric_limits<int32>::max();
-
 		return OutMidiData;
 	}
 
diff --git a/Engine/Plugins/Runtime/Harmonix/Source/HarmonixMetasound/Public/HarmonixMetasound/Components/MusicClockComponent.h b/Engine/Plugins/Runtime/Harmonix/Source/HarmonixMetasound/Public/HarmonixMetasound/Components/MusicClockComponent.h
--- a/Engine/Plugins/Runtime/Harmonix/Source/HarmonixMetasound/Public/HarmonixMetasound/Components/MusicClockComponent.h
+++ b/Engine/Plugins/Runtime/Harmonix/Source/HarmonixMetasound/Public/HarmonixMetasound/Components/MusicClockComponent.h
@@ -250,7 +250,7 @@ public:
 	float GetSongLengthBeats() const;
 
 	UFUNCTION(BlueprintPure, Category = "Song Data")
-	int32 GetSongLengthBars() const;
+	float GetSongLengthBars() const;
 
 	const FSongMaps& GetSongMaps() const;
 
diff --git a/Engine/Plugins/Runtime/Harmonix/Source/HarmonixMetasoundTests/Private/Analysis/MidiClockVertexAnalyzerTests.cpp b/Engine/Plugins/Runtime/Harmonix/Source/HarmonixMetasoundTests/Private/Analysis/MidiClockVertexAnalyzerTests.cpp
--- a/Engine/Plugins/Runtime/Harmonix/Source/HarmonixMetasoundTests/Private/Analysis/MidiClockVertexAnalyzerTests.cpp
+++ b/Engine/Plugins/Runtime/Harmonix/Source/HarmonixMetasoundTests/Private/Analysis/MidiClockVertexAnalyzerTests.cpp
@@ -51,16 +51,13 @@ namespace HarmonixMetasoundTests::MidiClockVertexAnalyzer
 		MidiData->Tracks.Empty();
 
 		MidiData->Tracks.Add(FMidiTrack(TEXT("conductor")));
-		MidiData->LastEventTick = std::numeric_limits<int32>::max();
 		MidiData->Tracks[0].AddEvent(FMidiEvent(0, FMidiMsg(static_cast<uint8>(TimeSigNumerator), static_cast<uint8>(TimeSigDenominator))));
 		BarMap.AddTimeSignatureAtBarIncludingCountIn(0, TimeSigNumerator, TimeSigNumerator);
 		const int32 MidiTempo = Harmonix::Midi::Constants::BPMToMidiTempo(Tempo);
 		MidiData->Tracks[0].AddEvent(FMidiEvent(0, FMidiMsg(MidiTempo)));
 		TempoMap.AddTempoInfoPoint(MidiTempo, 0);
 		MidiData->Tracks[0].Sort();
-		MidiData->SongMaps.GetSongLengthData().LastTick = std::numeric_limits<int32>::max();
-		MidiData->SongMaps.GetSongLengthData().LengthTicks = std::numeric_limits<int32>::max();
-		MidiData->SongMaps.GetSongLengthData().LengthBars = std::numeric_limits<int32>::max();
+		MidiData->ConformToLength(std::numeric_limits<int32>::max());
 
 		ClockInput->AttachToMidiResource(MidiData);
 		ClockInput->ResetAndStart(0);
diff --git a/Engine/Plugins/Runtime/Harmonix/Source/HarmonixMetasoundTests/Private/Analysis/MidiStreamVertexAnalyzerTests.cpp b/Engine/Plugins/Runtime/Harmonix/Source/HarmonixMetasoundTests/Private/Analysis/MidiStreamVertexAnalyzerTests.cpp
--- a/Engine/Plugins/Runtime/Harmonix/Source/HarmonixMetasoundTests/Private/Analysis/MidiStreamVertexAnalyzerTests.cpp
+++ b/Engine/Plugins/Runtime/Harmonix/Source/HarmonixMetasoundTests/Private/Analysis/MidiStreamVertexAnalyzerTests.cpp
@@ -53,16 +53,13 @@ namespace HarmonixMetasoundTests::MidiStreamVertexAnalyzer
 		const TSharedPtr<FMidiFileData> MidiData = MakeShared<FMidiFileData>();
 		check(MidiData);
 		MidiData->Tracks.Add(FMidiTrack(TEXT("conductor")));
-		MidiData->LastEventTick = std::numeric_limits<int32>::max();
 		MidiData->Tracks[0].AddEvent(FMidiEvent(0, FMidiMsg(static_cast<uint8>(TimeSigNum), static_cast<uint8>(TimeSigDenom))));
 		MidiData->SongMaps.GetBarMap().AddTimeSignatureAtBarIncludingCountIn(0, TimeSigNum, TimeSigDenom);
 		const int32 MidiTempo = Harmonix::Midi::Constants::BPMToMidiTempo(Tempo);
 		MidiData->Tracks[0].AddEvent(FMidiEvent(0, FMidiMsg(MidiTempo)));
 		MidiData->SongMaps.GetTempoMap().AddTempoInfoPoint(MidiTempo, 0);
 		MidiData->Tracks[0].Sort();
-		MidiData->SongMaps.GetSongLengthData().LastTick = std::numeric_limits<int32>::max();
-		MidiData->SongMaps.GetSongLengthData().LengthTicks = std::numeric_limits<int32>::max();
-		MidiData->SongMaps.GetSongLengthData().LengthBars = std::numeric_limits<int32>::max();
+		MidiData->ConformToLength(std::numeric_limits<int32>::max());
 
 		// Attach the maps
 		Clock->AttachToMidiResource(MidiData);
diff --git a/Engine/Plugins/Runtime/Harmonix/Source/HarmonixMetasoundTests/Private/Nodes/DelayNodeTests.cpp b/Engine/Plugins/Runtime/Harmonix/Source/HarmonixMetasoundTests/Private/Nodes/DelayNodeTests.cpp
--- a/Engine/Plugins/Runtime/Harmonix/Source/HarmonixMetasoundTests/Private/Nodes/DelayNodeTests.cpp
+++ b/Engine/Plugins/Runtime/Harmonix/Source/HarmonixMetasoundTests/Private/Nodes/DelayNodeTests.cpp
@@ -248,16 +248,15 @@ namespace HarmonixMetasoundTests::DelayNode
 			MidiData->Tracks.Empty();
 
 			MidiData->Tracks.Add(FMidiTrack(TEXT("conductor")));
-			MidiData->LastEventTick = std::numeric_limits<int32>::max();
 			MidiData->Tracks[0].AddEvent(FMidiEvent(0, FMidiMsg(static_cast<uint8>(TimeSigNum), static_cast<uint8>(TimeSigDenom))));
 			BarMap.AddTimeSignatureAtBarIncludingCountIn(0, TimeSigNum, TimeSigDenom);
 			const int32 MidiTempo = Harmonix::Midi::Constants::BPMToMidiTempo(Tempo);
 			MidiData->Tracks[0].AddEvent(FMidiEvent(0, FMidiMsg(MidiTempo)));
 			TempoMap.AddTempoInfoPoint(MidiTempo, 0);
 			MidiData->Tracks[0].Sort();
-			MidiData->SongMaps.GetSongLengthData().LastTick = std::numeric_limits<int32>::max();
+			MidiData->ConformToLength(std::numeric_limits<int32>::max());
 			MidiData->SongMaps.GetSongLengthData().LengthTicks = std::numeric_limits<int32>::max();
-			MidiData->SongMaps.GetSongLengthData().LengthBars = std::numeric_limits<int32>::max();
+			MidiData->SongMaps.GetSongLengthData().LengthFractionalBars = std::numeric_limits<float>::max();
 
 			(*ClockInput)->AttachToMidiResource(MidiData);
 			(*ClockInput)->ResetAndStart(0);
diff --git a/Engine/Plugins/Runtime/Harmonix/Source/HarmonixMetasoundTests/Private/Nodes/MidiPlayerNodeTests.cpp b/Engine/Plugins/Runtime/Harmonix/Source/HarmonixMetasoundTests/Private/Nodes/MidiPlayerNodeTests.cpp
--- a/Engine/Plugins/Runtime/Harmonix/Source/HarmonixMetasoundTests/Private/Nodes/MidiPlayerNodeTests.cpp
+++ b/Engine/Plugins/Runtime/Harmonix/Source/HarmonixMetasoundTests/Private/Nodes/MidiPlayerNodeTests.cpp
@@ -265,7 +265,6 @@ namespace HarmonixMetasoundTests::MidiPlayerNode
 
 			// create conductor track
 			FMidiTrack& Track = OutMidiData->Tracks.Add_GetRef(FMidiTrack(TEXT("conductor")));
-			OutMidiData->LastEventTick = std::numeric_limits<int32>::max();
 
 			// add time sig info
 			int32 TimeSigNum = FMath::Clamp(InTimeSigNum, 1, 64);
diff --git a/Engine/Plugins/Runtime/Harmonix/Source/HarmonixMetasoundTests/Private/Nodes/MidiTextTriggerNodeTests.cpp b/Engine/Plugins/Runtime/Harmonix/Source/HarmonixMetasoundTests/Private/Nodes/MidiTextTriggerNodeTests.cpp
--- a/Engine/Plugins/Runtime/Harmonix/Source/HarmonixMetasoundTests/Private/Nodes/MidiTextTriggerNodeTests.cpp
+++ b/Engine/Plugins/Runtime/Harmonix/Source/HarmonixMetasoundTests/Private/Nodes/MidiTextTriggerNodeTests.cpp
@@ -33,9 +33,8 @@ namespace HarmonixMetasoundTests::MidiTextTriggerNode
 		FMidiEvent ASecondTextEvent(TheMidi->GetSongMaps()->MsToTick(15.0f), FMidiMsg::CreateText(TextIndex, Harmonix::Midi::Constants::GMeta_Text));
 		TextTrack->AddEvent(ASecondTextEvent);
 		// Tell the midi file its tracks have been changed so it can recalculate song length data...
-		TheMidi->TracksChanged();
+		TheMidi->ScanTracksForSongLengthChange();
 		// Now round appropriately...
-		TheMidi->ConformMidiFileLength(EMidiFileLengthConformOption::RoundUp);
 		return TheMidi;
 	}
 
diff --git a/Engine/Plugins/Runtime/Harmonix/Source/HarmonixMetasoundTests/Private/Nodes/MorphingLfoNode.spec.cpp b/Engine/Plugins/Runtime/Harmonix/Source/HarmonixMetasoundTests/Private/Nodes/MorphingLfoNode.spec.cpp
--- a/Engine/Plugins/Runtime/Harmonix/Source/HarmonixMetasoundTests/Private/Nodes/MorphingLfoNode.spec.cpp
+++ b/Engine/Plugins/Runtime/Harmonix/Source/HarmonixMetasoundTests/Private/Nodes/MorphingLfoNode.spec.cpp
@@ -72,16 +72,13 @@ namespace HarmonixMetasoundTests::MorphingLFONode
 		MidiData->Tracks.Empty();
 
 		MidiData->Tracks.Add(FMidiTrack(TEXT("conductor")));
-		MidiData->LastEventTick = std::numeric_limits<int32>::max();
 		MidiData->Tracks[0].AddEvent(FMidiEvent(0, FMidiMsg(static_cast<uint8>(TimeSigNumerator), static_cast<uint8>(TimeSigDenominator))));
 		BarMap.AddTimeSignatureAtBarIncludingCountIn(0, TimeSigNumerator, TimeSigNumerator);
 		const int32 MidiTempo = Harmonix::Midi::Constants::BPMToMidiTempo(Tempo);
 		MidiData->Tracks[0].AddEvent(FMidiEvent(0, FMidiMsg(MidiTempo)));
 		TempoMap.AddTempoInfoPoint(MidiTempo, 0);
 		MidiData->Tracks[0].Sort();
-		MidiData->SongMaps.GetSongLengthData().LastTick = std::numeric_limits<int32>::max();
-		MidiData->SongMaps.GetSongLengthData().LengthTicks = std::numeric_limits<int32>::max();
-		MidiData->SongMaps.GetSongLengthData().LengthBars = std::numeric_limits<int32>::max();
+		MidiData->ConformToLength(std::numeric_limits<int32>::max());
 
 		(*ClockInput)->AttachToMidiResource(MidiData);
 		(*ClockInput)->ResetAndStart(0);
diff --git a/Engine/Plugins/Runtime/Harmonix/Source/HarmonixMetasoundTests/Private/Subsystems/MidiClockUpdateSubsystemTests.cpp b/Engine/Plugins/Runtime/Harmonix/Source/HarmonixMetasoundTests/Private/Subsystems/MidiClockUpdateSubsystemTests.cpp
--- a/Engine/Plugins/Runtime/Harmonix/Source/HarmonixMetasoundTests/Private/Subsystems/MidiClockUpdateSubsystemTests.cpp
+++ b/Engine/Plugins/Runtime/Harmonix/Source/HarmonixMetasoundTests/Private/Subsystems/MidiClockUpdateSubsystemTests.cpp
@@ -21,7 +21,6 @@ namespace HarmonixMetasoundTests::MidiClockUpdateSubsystem
 			check(MidiData);
 			
 			MidiData->Tracks.Add(FMidiTrack(TEXT("conductor")));
-			MidiData->LastEventTick = std::numeric_limits<int32>::max();
 			MidiData->Tracks[0].AddEvent(FMidiEvent(0, FMidiMsg(static_cast<uint8>(TimeSigNum), static_cast<uint8>(TimeSigDenom))));
 			FBarMap& BarMap = MidiData->SongMaps.GetBarMap();
 			BarMap.AddTimeSignatureAtBarIncludingCountIn(0, TimeSigNum, TimeSigDenom);
@@ -30,9 +29,7 @@ namespace HarmonixMetasoundTests::MidiClockUpdateSubsystem
 			FTempoMap& TempoMap = MidiData->SongMaps.GetTempoMap();
 			TempoMap.AddTempoInfoPoint(MidiTempo, 0);
 			MidiData->Tracks[0].Sort();
-			MidiData->SongMaps.GetSongLengthData().LastTick = std::numeric_limits<int32>::max();
-			MidiData->SongMaps.GetSongLengthData().LengthTicks = std::numeric_limits<int32>::max();
-			MidiData->SongMaps.GetSongLengthData().LengthBars = std::numeric_limits<int32>::max();
+			MidiData->ConformToLength(std::numeric_limits<int32>::max());
 
 			HarmonixMetasound::FMidiClock Clock{ OperatorSettings };
 			Clock.AttachToMidiResource(MidiData);
diff --git a/Engine/Plugins/Runtime/Harmonix/Source/HarmonixMidi/Private/HarmonixMidi/BarMap.cpp b/Engine/Plugins/Runtime/Harmonix/Source/HarmonixMidi/Private/HarmonixMidi/BarMap.cpp
--- a/Engine/Plugins/Runtime/Harmonix/Source/HarmonixMidi/Private/HarmonixMidi/BarMap.cpp
+++ b/Engine/Plugins/Runtime/Harmonix/Source/HarmonixMidi/Private/HarmonixMidi/BarMap.cpp
@@ -182,6 +182,50 @@ int32 FBarMap::BarBeatTickIncludingCountInToTick(int32 BarIndex, int32 BeatInBar
 			TickInBeat;
 }
 
+void FBarMap::TickToBarBeatTickIncludingCountIn(int32 InTick, int32& OutBarIndex, int32& OutBeatInBar, int32& OutTickIndexInBeat, int32* OutBeatsPerBar, int32* OutTicksPerBeat) const
+{
+	if (Points.IsEmpty())
+	{
+		// assume 4/4
+		OutBarIndex = InTick / (TicksPerQuarterNote * 4);
+		InTick -= OutBarIndex * (TicksPerQuarterNote * 4);
+		OutBeatInBar = InTick / TicksPerQuarterNote;
+		OutTickIndexInBeat = InTick - (OutBeatInBar * TicksPerQuarterNote);
+		OutBeatInBar++; // Beat is one based to be consistent with BarBeatTickIncludingCountInToTick
+		if (OutBeatsPerBar)
+		{
+			*OutBeatsPerBar = 4;
+		}
+		if (OutTicksPerBeat)
+		{
+			*OutTicksPerBeat = TicksPerQuarterNote;
+		}
+		return;
+	}
+	int32 TimeSigIndex = FMusicMapUtl::GetPointIndexForTick(Points, InTick);
+	if (!Points.IsValidIndex(TimeSigIndex))
+	{
+		TimeSigIndex = 0;
+	}
+	int32 TicksPerBar = GetTicksInBarAfterPoint(TimeSigIndex);
+	int32 TicksPerBeat = GetTicksInBeatAfterPoint(TimeSigIndex);
+	int32 TicksPast = InTick - Points[TimeSigIndex].StartTick;
+	int32 BarsPassed = TicksPast / TicksPerBar;
+	OutBarIndex = Points[TimeSigIndex].BarIndex + BarsPassed;
+	TicksPast -= BarsPassed * TicksPerBar;
+	OutBeatInBar = TicksPast / TicksPerBeat;
+	OutTickIndexInBeat = TicksPast - (OutBeatInBar * TicksPerBeat);
+	OutBeatInBar++; // Beat is one based to be consistent with BarBeatTickIncludingCountInToTick
+	if (OutBeatsPerBar)
+	{
+		*OutBeatsPerBar = Points[TimeSigIndex].TimeSignature.Numerator;
+	}
+	if (OutTicksPerBeat)
+	{
+		*OutTicksPerBeat = TicksPerBeat;
+	}
+}
+
 FMusicTimestamp FBarMap::TickToMusicTimestamp(float Tick, int32* OutBeatsPerBar) const
 {
 	FMusicTimestamp Result;
@@ -234,85 +278,6 @@ FMusicTimestamp FBarMap::TickFromBarOneToMusicTimestamp(float InTickFromBarOne,
 {
 	return TickToMusicTimestamp(InTickFromBarOne + GetTickOfBarOne(), OutBeatsPerBar);
 }
-/*
-
-void FBarMap::TickToRawBarBeatTick(int32 InTick, int32& OutBar, int32& OutBeat, int32& OutTickInBeat, int32* OutBeatsPerBar, int32* OutTicksPerBeat) const
-{
-	if (Points.IsEmpty())
-	{
-		// Assume 4/4 time.
-		if (InTick == 0)
-		{
-			OutBar = 1;
-			OutBeat = 1;
-			OutTickInBeat = 0;
-		}
-		else if (InTick > 0)
-		{
-			OutBar = InTick / (TicksPerQuarterNote * 4);
-			InTick -= (OutBar * TicksPerQuarterNote * 4);
-			OutBeat = (float)InTick / (float)TicksPerQuarterNote;
-			OutTickInBeat = InTick - (OutBeat * TicksPerQuarterNote);
-			OutBar += 1; // 1 based
-			OutBeat += 1; // 1 based
-		}
-		else // InTick < 0
-		{
-			OutBar = InTick / (TicksPerQuarterNote * 4);
-			OutBar -= 1; // 1 based
-			InTick -= (OutBar * TicksPerQuarterNote * 4);
-			OutBeat = (float)InTick / (float)TicksPerQuarterNote;
-			OutTickInBeat = InTick - (OutBeat * TicksPerQuarterNote);
-			OutBeat += 1;
-		}
-		if (OutBeatsPerBar) *OutBeatsPerBar = 4;
-		if (OutTicksPerBeat) *OutTicksPerBeat = TicksPerQuarterNote;
-		return;
-	}
-
-	int32 Index = FMusicMapUtl::GetPointIndexForTick(Points, InTick);
-	if (Index < 0)
-	{
-		Index = 0;
-	}
-
-	int32 TicksPerBar = GetTicksInBarAfterPoint(Index);
-	int32 TicksPerBeat = GetTicksInBeatAfterPoint(Index);
-	if (InTick < 0)
-	{
-		// before the beginning? 
-		int32 BarsPassed = (-InTick / TicksPerBar) + 1;
-		int32 TickInBar = (BarsPassed * TicksPerBar) + InTick;
-		int32 BeatsPassed = TickInBar / TicksPerBeat;
-		int32 TickInBeat = TickInBar - (BeatsPassed * TicksPerBeat);
-		OutBar = -BarsPassed; // already 1 based
-		OutBeat = BeatsPassed + 1; // 1 based
-		OutTickInBeat = TickInBeat; // 0 based
-	}
-	else
-	{
-		int32 TicksPassed = InTick - Points[Index].StartTick;
-		int32 BarsPassed = TicksPassed / TicksPerBar;
-		TicksPassed -= BarsPassed * TicksPerBar;
-		int32 BeatsPassed = TicksPassed / TicksPerBeat;
-		TicksPassed -= BeatsPassed * TicksPerBeat;
-		OutBar = Points[Index].Bar + BarsPassed + 1; // 1 based
-		OutBeat = BeatsPassed + 1; // 1 based
-		OutTickInBeat = TicksPassed; // 0 based
-	}
-	if (OutBeatsPerBar) *OutBeatsPerBar = Points[Index].TimeSignature.Numerator;
-	if (OutTicksPerBeat) *OutTicksPerBeat = TicksPerBeat;
-}
-*/
-
-/*
-float FBarMap::TickToFractionalMusicalTimestampBar(int32 Tick) const
-{
-	int32 BeatsPerBar;
-	FMusicTimestamp Timestamp = TickToMusicTimestamp(Tick, &BeatsPerBar);
-	return (float)Timestamp.Bar + ((Timestamp.Beat - 1.0f) / (float)BeatsPerBar);
-}
-*/
 
 float FBarMap::TickToFractionalBarIncludingCountIn(float Tick) const
 {
diff --git a/Engine/Plugins/Runtime/Harmonix/Source/HarmonixMidi/Private/HarmonixMidi/MidiFile.cpp b/Engine/Plugins/Runtime/Harmonix/Source/HarmonixMidi/Private/HarmonixMidi/MidiFile.cpp
--- a/Engine/Plugins/Runtime/Harmonix/Source/HarmonixMidi/Private/HarmonixMidi/MidiFile.cpp
+++ b/Engine/Plugins/Runtime/Harmonix/Source/HarmonixMidi/Private/HarmonixMidi/MidiFile.cpp
@@ -69,7 +69,7 @@ void UMidiFile::GetAssetRegistryTags(FAssetRegistryTagsContext Context) const
 	FString Tempo = FString::Printf(TEXT("%.2f BPM"), TheMidiData.SongMaps.GetTempoAtTick(0));
 	Context.AddTag(UObject::FAssetRegistryTag("InitialTempo",  Tempo, UObject::FAssetRegistryTag::TT_Numerical));
 	
-	FString Length = FString::Printf(TEXT("%d Bars"), TheMidiData.SongMaps.GetSongLengthData().LengthBars);
+	FString Length = FString::Printf(TEXT("%.5f Bars ( %s )"), TheMidiData.SongMaps.GetSongLengthData().LengthFractionalBars, *TheMidiData.SongMaps.GetSongLengthString());
 	Context.AddTag(UObject::FAssetRegistryTag("Length", Length, UObject::FAssetRegistryTag::TT_Numerical));
 
 #if WITH_EDITORONLY_DATA
@@ -180,8 +180,6 @@ void UMidiFile::LoadStdMidiFile(
 	FStdMidiFileReader reader(Archive, Filename, &AllReceivers, TheMidiData.TicksPerQuarterNote, InTextEncoding);
 	reader.ReadAllTracks();
 
-	TheMidiData.LastEventTick = EventReceiver.GetLastTick();
-
 	if (reader.IsFailed())
 	{
 		if (EnsureNotFailed)
@@ -419,357 +417,98 @@ void UMidiFile::FindAllTextEvents(const FString& EventText, int32 TrackIndex, TA
 	}
 }
 
-void UMidiFile::TracksChanged()
+void UMidiFile::ScanTracksForSongLengthChange()
 {
-	int32 NewLastEventTick = 0;
-	bool bIsFileEmpty = true;
-	for (auto& Track : TheMidiData.Tracks)
-	{
-		FMidiEventList& Events = Track.GetRawEvents();
-		//check if the file is empty
-		if (bIsFileEmpty && !Events.IsEmpty())
-		{
-			bIsFileEmpty = false;
-		}
-		int32 LastTickOnTrack = Events.IsEmpty() ? 0 : Events.Last().GetTick();
-		if (NewLastEventTick < LastTickOnTrack)
-		{
-			NewLastEventTick = LastTickOnTrack;
-		}
-	}
-
-	if (NewLastEventTick == 0 && !bIsFileEmpty)
-	{
-		TheMidiData.LastEventTick = NewLastEventTick;
-		FSongLengthData& LengthData = TheMidiData.SongMaps.GetSongLengthData();
-		//if the last event tick in a file is 0 and the file is not empty, we know that the length in tick is 1
-		LengthData.LengthTicks = 1;
-		LengthData.LastTick = NewLastEventTick;
-		LengthData.LengthBars = TheMidiData.SongMaps.GetBarMap().TickToBarIncludingCountIn(LengthData.LengthTicks);
-		return;
-	}
-
-	//update length data in song length data if needed 
-	if (NewLastEventTick != TheMidiData.LastEventTick)
-	{
-		TheMidiData.LastEventTick = NewLastEventTick;
-		FMusicTimestamp Timestamp = TheMidiData.SongMaps.GetBarMap().TickToMusicTimestamp(NewLastEventTick + 1);
-		FSongLengthData& LengthData = TheMidiData.SongMaps.GetSongLengthData();
-		LengthData.LengthTicks = TheMidiData.SongMaps.GetBarMap().MusicTimestampToTick(Timestamp);
-		LengthData.LastTick = LengthData.LengthTicks - 1;
-		LengthData.LengthBars = TheMidiData.SongMaps.GetBarMap().TickToBarIncludingCountIn(LengthData.LengthTicks);
-	}
+	TheMidiData.ScanTracksForSongLengthChange();
+	MarkPackageDirty();
 }
 
-bool UMidiFile::ShouldConformMidiFileLength(EMidiFileLengthConformOption Option)
+bool UMidiFile::LengthIsAPerfectSubdivision() const
 {
-	//if file length is already rounded down, there's nothing more can be done
-	if (bLengthRoundedDown) return false;
+	return TheMidiData.LengthIsAPerfectSubdivision();
+}
 
-	//if a file's length is rounded up, it can be rounded down once more, but nothing else 
-	if (bLengthRoundedUp)
+void UMidiFile::QuantizeLengthToNearestPerfectSubdivision(const EMidiFileQuantizeDirection Direction)
+{
+	EMidiClockSubdivisionQuantization Subdivision = EMidiClockSubdivisionQuantization::None;
+	int32 QuantizedLength = TheMidiData.SongMaps.QuantizeTickToAnyNearestSubdivision(TheMidiData.SongMaps.GetSongLengthData().LengthTicks, Direction, Subdivision);
+	if (QuantizedLength == 0)
 	{
-		if (Option == EMidiFileLengthConformOption::RoundDown)
+		// If we were specifically asked to go down and going down gets
+		// us a zero length file, then warn and don't quantize...
+		if (Direction == EMidiFileQuantizeDirection::Down)
 		{
-			//if the file's length is less than 1, cannot round down
-			return TheMidiData.SongMaps.GetSongLengthData().LengthBars > 1 ? true : false;
-		}
-		else
-		{
-			return false;
+			UE_LOG(LogMIDI, Warning, TEXT("QuantizeLengthToNearestPerfectSubdivision: Asked to Quantize file length DOWN, but that would result in a zero length midi file. NOT ALLOWED! Skipping quantization!"));
+			return;
 		}
-	}
 
-	if (bLengthRoundedToNearest && Option == EMidiFileLengthConformOption::Nearest) return false;
-
-	//if bar length is 1 bar (or less, it will round up to 1 bar if the file is imported, 0 if the file is manually created), 
-	//it CANNOT be rounded down to 0 bar, can only be rounded up
-	if (TheMidiData.SongMaps.GetSongLengthData().LengthBars <= 1 && Option != EMidiFileLengthConformOption::RoundUp)
-	{
-		return false;
-	}
-
-	//Special case: a manually created Midi file only has events on tick 0
-	//This should still be able to be conformed by rounding up ONLY
-	bool bIsFileEmpty = true;
-	for (auto& Track : TheMidiData.Tracks)
-	{
-		FMidiEventList& Events = Track.GetRawEvents();
-		//check if the file is empty
-		if (bIsFileEmpty && !Events.IsEmpty())
+		// If we were ask to go up and we STILL got a zero length
+		// file that is super weird! Warn the user and skip quantizing...
+		if (Direction == EMidiFileQuantizeDirection::Up)
 		{
-			bIsFileEmpty = false;
-			break;
+			UE_LOG(LogMIDI, Error, TEXT("This is odd. Asked to quantize MIDI file length UP to the nearest subdivision, but returned length is zero! Skipping."));
+			return;
 		}
-	}
 
-	//get the last event tick in file
-	int32 LastEventTick = TheMidiData.LastEventTick;
-
-	if (LastEventTick == 0 && !bIsFileEmpty && Option == EMidiFileLengthConformOption::RoundUp) return true;
-
-	//if the last event tick is at the very end of an integer bar, the next tick should be the at the next bar, beat 1.0 
-	FMusicTimestamp Timestamp = TheMidiData.SongMaps.GetBarMap().TickToMusicTimestamp(LastEventTick + 1);
-	//for manually created Midi files, we can specify that the last event tick is at the very end of an integer bar
-	//however,some DAWs places the very last event of an integer bar at the very beginning of the next bar, causing a 1-tick difference
-	//the current solution is to tolerate this 1-tick error when checking the last event tick
-	double OneTickTimeStampBeat = TheMidiData.SongMaps.GetBarMap().TickToMusicTimestamp(1).Beat;
-	double OneTickError = OneTickTimeStampBeat - 1.0;
-	if (!FMath::IsNearlyEqual(Timestamp.Beat, 1.0f, OneTickError))
-	{
-		return true;
+		// We MUST have been asked to quantize to nearest, but nearest must be
+		// down and result in a zero length file. So... this time force up, as
+		// that will be the only valid "nearest"!
+		QuantizedLength = TheMidiData.SongMaps.QuantizeTickToAnyNearestSubdivision(TheMidiData.SongMaps.GetSongLengthData().LengthTicks, EMidiFileQuantizeDirection::Up, Subdivision);
+		check (QuantizedLength > 0);
 	}
-	else if (!FMath::IsNearlyEqual(Timestamp.Beat, 1.0f))
-	{
-		ConformMidiFileLength(EMidiFileLengthConformOption::RoundDown, true);
-	}
-	return false;
+	TheMidiData.ConformToLength(QuantizedLength);
+	MarkPackageDirty();
 }
 
-void UMidiFile::ConformMidiFileLength(EMidiFileLengthConformOption Option, bool Force)
+void UMidiFile::QuantizeLengthToSubdivision(const EMidiFileQuantizeDirection Direction, const EMidiClockSubdivisionQuantization Subdivision)
 {
-	//check whether midi file length needs to be conformed
-	if (!Force && !ShouldConformMidiFileLength(Option))
+	int32 QuantizedLength = TheMidiData.SongMaps.QuantizeTickToNearestSubdivision(TheMidiData.SongMaps.GetSongLengthData().LengthTicks, Direction, Subdivision);
+	if (QuantizedLength == 0)
 	{
-		UE_LOG(LogMIDI, Log, TEXT("Midi file does not need to be conformed."));
-		return;
-	}
-	
-	int32 LastTickInSongLengthData = TheMidiData.SongMaps.GetSongLengthData().LastTick;
-	//Currently, file length is automatically rounded up to the nearest integer bar for imported midi files, if the conform option is Round Up, early out
-	//(NOTE: Round Up doesn't add/change/remove anything to the midi file or events)
-	if (Option == EMidiFileLengthConformOption::RoundUp)
-	{
-		//for manually created Midi files, update the song length data to keep up with 
-		//results of imported Midi files WITHOUT adding/changing/removing events
-		FMusicTimestamp Timestamp = TheMidiData.SongMaps.GetBarMap().TickToMusicTimestamp(TheMidiData.LastEventTick + 1);
-		if (!FMath::IsNearlyEqual(Timestamp.Beat, 1.0f))
+		// If we were specifically asked to go down and going down gets
+		// us a zero length file, then warn and don't quantize...
+		if (Direction == EMidiFileQuantizeDirection::Down)
 		{
-			Timestamp.Bar++;
-			Timestamp.Beat = 1.0f;
-			FSongLengthData& LengthData = TheMidiData.SongMaps.GetSongLengthData();
-			LengthData.LengthTicks = TheMidiData.SongMaps.GetBarMap().MusicTimestampToTick(Timestamp);
-			LengthData.LastTick = LengthData.LengthTicks - 1;
-			LengthData.LengthBars = TheMidiData.SongMaps.GetBarMap().TickToBarIncludingCountIn(LengthData.LengthTicks);
+			UE_LOG(LogMIDI, Warning, TEXT("QuantizeLengthToNearestPerfectSubdivision: Asked to Quantize file length DOWN, but that would result in a zero length midi file. NOT ALLOWED! Skipping quantization!"));
+			return;
 		}
-		
-		bLengthRoundedUp = true;
-		return;
-	}
-
-	//Get file's fractional bar length (not rounded up)
-	float OriginalLengthFractionalBar = TheMidiData.SongMaps.GetBarIncludingCountInAtTick(TheMidiData.LastEventTick);
 
-	//Round Down to the nearest integer bar
-	int32 ConformedLastBar = (int32)OriginalLengthFractionalBar;
-
-	//rounding to nearest is still either round down or round up
-	//depending on the result from FMath::RoundToInt()
-	if (Option == EMidiFileLengthConformOption::Nearest)
-	{
-		ConformedLastBar = FMath::RoundToInt(OriginalLengthFractionalBar);
-		EMidiFileLengthConformOption NearestRoundingOption = ConformedLastBar > (int32)OriginalLengthFractionalBar ? EMidiFileLengthConformOption::RoundUp : EMidiFileLengthConformOption::RoundDown;
-		
-		//if round to nearest ends up being rounding up, early out
-		if (NearestRoundingOption == EMidiFileLengthConformOption::RoundUp)
+		// If we were ask to go up and we STILL got a zero length
+		// file that is super weird! Warn the user and skip quantizing...
+		if (Direction == EMidiFileQuantizeDirection::Up)
 		{
-			//round the length up by updating the song length data
-			FMusicTimestamp Timestamp = TheMidiData.SongMaps.GetBarMap().TickToMusicTimestamp(TheMidiData.LastEventTick + 1);
-			if (!FMath::IsNearlyEqual(Timestamp.Beat, 1.0f))
-			{
-				Timestamp.Bar++;
-				Timestamp.Beat = 1.0f;
-				FSongLengthData& LengthData = TheMidiData.SongMaps.GetSongLengthData();
-				LengthData.LengthTicks = TheMidiData.SongMaps.GetBarMap().MusicTimestampToTick(Timestamp);
-				LengthData.LastTick = LengthData.LengthTicks - 1;
-				LengthData.LengthBars = TheMidiData.SongMaps.GetBarMap().TickToBarIncludingCountIn(LengthData.LengthTicks);
-			}
-			bLengthRoundedUp = true;
-			bLengthRoundedToNearest = true;
+			UE_LOG(LogMIDI, Error, TEXT("This is odd. Asked to quantize MIDI file length UP to the nearest subdivision, but returned length is zero! Skipping."));
 			return;
 		}
-	}
-
-	//the actual tick to conform to by rounding down
-	int32 ConformedLastEventTick = TheMidiData.SongMaps.GetBarMap().BarIncludingCountInToTick(ConformedLastBar) - 1;
 
-	//Retrieve midi file name for logging after moving/removing midi events
-	FString MidiFileName = this->GetName();
-
-	//keep track of the events that are moved to the last tick and excessive events removed after moving
-	int32 NumEventsMovedToLastTick = 0;
-	int32 TotalNumEventsRemoved = 0;
-	int32 NumNoteOnNoteOffPairsRemoved = 0;
-	int32 NumAftertouchEventsRemoved = 0;
-	int32 NumPitchEventsRemoved = 0;
-	int32 NumControlEventsRemoved = 0;
+		// We MUST have been asked to quantize to nearest, but nearest must be
+		// down and result in a zero length file. So... this time force up, as
+		// that will be the only valid "nearest"!
+		QuantizedLength = TheMidiData.SongMaps.QuantizeTickToNearestSubdivision(TheMidiData.SongMaps.GetSongLengthData().LengthTicks, EMidiFileQuantizeDirection::Up, Subdivision);
+		check(QuantizedLength > 0);
+	}
+	TheMidiData.ConformToLength(QuantizedLength);
+	MarkPackageDirty();
+}
 
-	//go through midi events in each midi track 
-	for (int32 TrackIndex = 0; TrackIndex < TheMidiData.Tracks.Num(); ++TrackIndex)
+bool UMidiFile::ConformToLength(int32 NewLengthTicks)
+{
+	if (TheMidiData.ConformToLength(NewLengthTicks))
 	{
-		//Get the raw midi events and the last tick of the last event
-		FMidiTrack& Track = TheMidiData.Tracks[TrackIndex];
-		FMidiEventList& Events = Track.GetRawEvents();
-		int32 LastEventTick = Events.Last().GetTick();
-	
-		//if the last tick of this event is less than the tick of the conformed last tick,
-		//this midi track does not need to be conformed, go to the next midi track
-		if (LastEventTick <= ConformedLastEventTick) continue;
-
-		//move all events with tick > ConfromedLastEventTick to ConformedLastTick
-		for (int32 EventIndex = Events.Num() - 1; EventIndex >= 0 && Events[EventIndex].GetTick() > ConformedLastEventTick; --EventIndex)
-		{
-			// NOTE: It is safe for us to do this BECAUSE...
-			// We do not have to re-sort the midi data on the tracks because while we may have slid events earlier,
-			// their relative position to each other MUST be unchanged!
-			Events[EventIndex].Tick = ConformedLastEventTick;
-			NumEventsMovedToLastTick++;
-		}
-
-		//filter the midi events (at the new conformed last tick) after moving them, remove note on/note off pairs on the last tick, and other excessive events
-		int32 NumItemsRemoved = 0;
-		int32 CurrentEventIndex = Events.Num() - 1;
-
-		while (CurrentEventIndex >= 0 && Events[CurrentEventIndex].GetTick() == ConformedLastEventTick)
-		{
-			FMidiEvent& Event = Events[CurrentEventIndex];
-			FMidiMsg& Msg = Event.GetMsg();
-
-			//filter Note On/Note off pairs on the last tick
-			if (Msg.IsNoteOff())
-			{
-				int32 NoteOnEventIndex = -1;
-				for (int32 i = CurrentEventIndex - 1; i > 0 && Events[i].GetTick() == ConformedLastEventTick; --i)
-				{
-					//check equality of note on/note off events' midi note number (data1) 
-					if (Events[i].GetMsg().IsNoteOn() && Events[i].GetMsg().GetStdData1() == Msg.GetStdData1())
-					{
-						NoteOnEventIndex = i;
-						break;
-					}
-				}
-				//remove note on/note off pairs on the conformed last tick
-				if (NoteOnEventIndex != -1) 
-				{
-					Events.RemoveAt(CurrentEventIndex);
-					NumItemsRemoved++;
-					Events.RemoveAt(NoteOnEventIndex);
-					NumItemsRemoved++;
-					NumNoteOnNoteOffPairsRemoved++;
-					TotalNumEventsRemoved += 2;
-				}
-			}
-			uint8 MsgStatus = Msg.Status;
-
-			using namespace Harmonix::Midi::Constants;
-			
-			//filter chan press/pitch bend/program change events
-			if (MsgStatus == GChanPres|| MsgStatus == GPitch || MsgStatus == GProgram)
-			{
-				int32 OtherEventIndexToRemove = -1;
-				//check if there exist multiple events with same status on the last tick
-				for (int32 i = CurrentEventIndex - 1; i >= 0 && Events[i].GetTick() == ConformedLastEventTick; --i)
-				{
-					if (Events[i].GetMsg().Status == MsgStatus) {
-						OtherEventIndexToRemove = CurrentEventIndex;
-						break;
-					}
-				}
-				//remove excessive events on the conformed last tick
-				if (OtherEventIndexToRemove != -1)
-				{
-					Events.RemoveAt(OtherEventIndexToRemove);
-					NumItemsRemoved++;
-					if (MsgStatus == GChanPres || MsgStatus == GPolyPres)
-					{
-						NumAftertouchEventsRemoved++;
-					}
-					else {
-						NumPitchEventsRemoved++;
-					}
-					TotalNumEventsRemoved++;
-				}
-
-			}
-
-			//filter Control Change events and Poly Pres events
-			if (MsgStatus == GControl || MsgStatus == GPolyPres)
-			{
-				uint8 CurrentEventControllerId = Msg.Data1;
-				int32 ControlEventIndexToRemove = -1; 
-				for (int32 i = CurrentEventIndex - 1; i >= 0 && Events[i].GetTick() == ConformedLastEventTick; --i)
-				{
-					//check control change events for identical controller ID (data1) on the same tick and remove the later one
-					//check poly pres events for the same note number (data1) on the same tick and remove the later one
-					if (Events[i].GetMsg().Data1 == CurrentEventControllerId) {
-						ControlEventIndexToRemove = CurrentEventIndex;
-						break;
-					}
-				}
-				//remove excessive control changes or poly pressure on conformed last tick
-				if (ControlEventIndexToRemove != -1)
-				{
-					Events.RemoveAt(ControlEventIndexToRemove);
-					NumItemsRemoved++;
-					NumControlEventsRemoved++;
-					TotalNumEventsRemoved++;
-				}
-			}
-
-			//update indices after removing excessive events
-			if (NumItemsRemoved == 0)
-			{
-				CurrentEventIndex -= 1;
-			}
-			else {
-				CurrentEventIndex -= NumItemsRemoved;
-				NumItemsRemoved = 0;
-			}
-		}
+		MarkPackageDirty();
+		return true;
 	}
+	return false;
+}
 
-	TracksChanged();
-
-	//update the conformed length in SongLengthData if conformed to an integer bar, and marks the file is being conformed with the current conform option(s)
-	if (Option == EMidiFileLengthConformOption::RoundDown || Option == EMidiFileLengthConformOption::Nearest)
+bool UMidiFile::ConformToLengthGivenLastEventTick(int32 NewLastEventTick)
+{
+	if (TheMidiData.ConformToLengthGivenLastEventTick(NewLastEventTick))
 	{
-		FSongLengthData& LengthData = TheMidiData.SongMaps.GetSongLengthData();
-		LengthData.LengthTicks = ConformedLastEventTick + 1;
-		LengthData.LastTick = ConformedLastEventTick;
-		LengthData.LengthBars = TheMidiData.SongMaps.GetBarMap().TickToBarIncludingCountIn(LengthData.LengthTicks);
-		bLengthRoundedDown = true;
-
-		//if we get here by rounding up first (or round up by rounding to the nearest integer) and then round down
-		//this means the file is ultimately rounded down, NOT rounded up
-		if (bLengthRoundedUp) bLengthRoundedUp = false;
-
-		bLengthRoundedToNearest = Option == EMidiFileLengthConformOption::Nearest ? true : false;
-		
-		//Log conform summary 
-		FString MidiFileLengthConformSummary = FString::Printf(
-			TEXT("Midi file '%s.mid' length conformance summary: original length: %f bars, rounded down to new length: %d bar(s), moved %d Midi event(s) to tick %d"),
-			*MidiFileName,
-			OriginalLengthFractionalBar,
-			ConformedLastBar,
-			NumEventsMovedToLastTick,
-			ConformedLastEventTick
-		);
-		if (TotalNumEventsRemoved > 0)
-		{
-			MidiFileLengthConformSummary.Append(FString::Printf(TEXT("; removed %d Midi event(s):"),TotalNumEventsRemoved));
-			if (NumNoteOnNoteOffPairsRemoved > 0) MidiFileLengthConformSummary.Append(FString::Printf(TEXT(" %d Note On/Note off event pair(s),"), NumNoteOnNoteOffPairsRemoved));
-			if (NumAftertouchEventsRemoved > 0) MidiFileLengthConformSummary.Append(FString::Printf(TEXT(" %d Aftertouch event(s),"), NumAftertouchEventsRemoved));
-			if (NumPitchEventsRemoved > 0) MidiFileLengthConformSummary.Append(FString::Printf(TEXT(" %d Pitch Bend event(s),"), NumPitchEventsRemoved));
-			if (NumControlEventsRemoved > 0) MidiFileLengthConformSummary.Append(FString::Printf(TEXT(" %d Control Change event(s),"), NumControlEventsRemoved));
-			MidiFileLengthConformSummary.Append(FString::Printf(TEXT(" at tick %d "), ConformedLastEventTick));
-		}
-		UE_LOG(LogMIDI, Log, TEXT("%s"),*MidiFileLengthConformSummary);
+		MarkPackageDirty();
+		return true;
 	}
-
-	//This prevents the current conform option affecting any midi file that is currently being obtained from RenderableCopyOfMidiFileData
-	//e.g. if file length is being conformed while the un-conformed file is occupied during playback, changes to the midi file won't be reflected until the 
-	//current playback is stopped and starts playing again (we only create a new RenderableCopyOfMidiFileData if there isn't one existed already)
-	RenderableCopyOfMidiFileData = nullptr;
+	return false;
 }
 
 #if WITH_EDITORONLY_DATA
@@ -905,8 +644,12 @@ bool FMidiFileData::operator==(const FMidiFileData& Other) const
 
 	return MidiFileName == Other.MidiFileName &&
 		TicksPerQuarterNote == Other.TicksPerQuarterNote &&
-		SongMaps == Other.SongMaps &&
-		LastEventTick == Other.LastEventTick;
+		SongMaps == Other.SongMaps;
+}
+
+int32 FMidiFileData::GetLastEventTick() const
+{
+	return SongMaps.GetSongLengthData().LastTick;
 }
 
 int32 FMidiFileData::FindTrackIndexByName(const FString& TrackName)
@@ -948,10 +691,202 @@ void FMidiFileData::PostSerialize(const FArchive& Ar)
 				BarMap.AddTimeSignatureAtBarIncludingCountIn(BarIndex, Event.GetMsg().GetTimeSigNumerator(), Event.GetMsg().GetTimeSigDenominator());
 			}
 		}
-		BarMap.Finalize(LastEventTick);
+		BarMap.Finalize(GetLastEventTick());
 	}
 }
 
+bool FMidiFileData::ConformToLength(int32 NewLengthTicks)
+{
+	return ConformToLengthGivenLastEventTick(NewLengthTicks - 1);
+}
+
+bool FMidiFileData::ConformToLengthGivenLastEventTick(int32 NewLastEventTick)
+{
+	if (NewLastEventTick < 0)
+	{
+		NewLastEventTick = 0;
+		UE_LOG(LogMIDI, Warning, TEXT("Asked to comform length such that the last event tick would be less than 0... Not possible! Conforming length to 1 tick, with last event being placed on tick 0!"));
+	}
+
+	if (SongMaps.GetSongLengthData().LastTick == NewLastEventTick && SongMaps.GetSongLengthData().LengthTicks != 0)
+	{
+		// nothing to change...
+		return false;
+	}
+
+	if (SongMaps.GetSongLengthData().LastTick < NewLastEventTick)
+	{
+		// easy... just change the length data...
+		SongMaps.SetSongLengthTicks(NewLastEventTick + 1);
+		// we changed something so report it...
+		return true;
+	}
+
+	//keep track of the events that are moved to the last tick and excessive events removed after moving
+	int32 NumEventsMovedToLastTick = 0;
+	int32 TotalNumEventsRemoved = 0;
+	int32 NumNoteOnNoteOffPairsRemoved = 0;
+	int32 NumControlEventsRemoved = 0;
+
+	//go through midi events in each midi track 
+	for (int32 TrackIndex = 0; TrackIndex < Tracks.Num(); ++TrackIndex)
+	{
+		//Get the raw midi events and the last tick of the last event
+		FMidiTrack& Track = Tracks[TrackIndex];
+		FMidiEventList& Events = Track.GetRawEvents();
+		int32 LastEventTick = Events.Last().GetTick();
+
+		//if the last tick of this event is less than the tick of the conformed last tick,
+		//this midi track does not need to be conformed, go to the next midi track
+		if (LastEventTick <= NewLastEventTick) continue;
+
+		//move all events with tick > ConfromedLastEventTick to ConformedLastTick
+		for (int32 EventIndex = Events.Num() - 1; EventIndex >= 0 && Events[EventIndex].GetTick() > NewLastEventTick; --EventIndex)
+		{
+			// NOTE: It is safe for us to do this BECAUSE...
+			// We do not have to re-sort the midi data on the tracks because while we may have slid events earlier,
+			// their relative position to each other MUST be unchanged!
+			Events[EventIndex].Tick = NewLastEventTick;
+			NumEventsMovedToLastTick++;
+		}
+
+		//filter the midi events (at the new conformed last tick) after moving them, remove note on/note off pairs on the last tick, and other excessive events
+		int32 NumItemsRemoved = 0;
+		int32 CurrentEventIndex = Events.Num() - 1;
+
+		while (CurrentEventIndex >= 0 && Events[CurrentEventIndex].GetTick() == NewLastEventTick)
+		{
+			using namespace Harmonix::Midi::Constants;
+
+			FMidiEvent& Event = Events[CurrentEventIndex];
+			FMidiMsg& Msg = Event.GetMsg();
+
+			if (Msg.IsStd())
+			{
+				uint8 MsgType = Msg.GetStdStatusType();
+
+				//filter Note On/Note off pairs on the last tick
+				if (MsgType == GNoteOff)
+				{
+					int32 NoteOnEventIndex = -1;
+					for (int32 i = CurrentEventIndex - 1; i > 0 && Events[i].GetTick() == NewLastEventTick; --i)
+					{
+						//check equality of note on/note off events' midi note number (data1) 
+						if (Events[i].GetMsg().IsNoteOn() && Events[i].GetMsg().GetStdData1() == Msg.GetStdData1())
+						{
+							NoteOnEventIndex = i;
+							break;
+						}
+					}
+					//remove note on/note off pairs on the conformed last tick
+					if (NoteOnEventIndex != -1)
+					{
+						Events.RemoveAt(CurrentEventIndex);
+						NumItemsRemoved++;
+						Events.RemoveAt(NoteOnEventIndex);
+						NumItemsRemoved++;
+						NumNoteOnNoteOffPairsRemoved++;
+						TotalNumEventsRemoved += 2;
+					}
+				}
+
+				//filter chan press/pitch bend/program change events
+				else if (MsgType == GChanPres || MsgType == GPitch || MsgType == GProgram)
+				{
+					TArray<int32> OtherEventIndexesToRemove;
+					//check if there exist multiple events with same status on the last tick
+					for (int32 i = CurrentEventIndex - 1; i >= 0 && Events[i].GetTick() == NewLastEventTick; --i)
+					{
+						if (Events[i].GetMsg().IsStd() &&
+							Events[i].GetMsg().Status == Msg.Status) 
+						{
+							OtherEventIndexesToRemove.Add(i);
+						}
+					}
+					//remove excessive events on the conformed last tick
+					if (!OtherEventIndexesToRemove.IsEmpty())
+					{
+						for (int32 i : OtherEventIndexesToRemove)
+						{
+							Events.RemoveAt(i);
+						}
+						NumItemsRemoved += OtherEventIndexesToRemove.Num();
+						NumControlEventsRemoved += OtherEventIndexesToRemove.Num();
+						TotalNumEventsRemoved += OtherEventIndexesToRemove.Num();
+					}
+				}
+
+				//filter Control Change events and Poly Pres events
+				else if (MsgType == GControl || MsgType == GPolyPres)
+				{
+					TArray<int32> ControlEventIndexesToRemove;
+					for (int32 i = CurrentEventIndex - 1; i >= 0 && Events[i].GetTick() == NewLastEventTick; --i)
+					{
+						//check control change events for identical controller ID (data1) on the same tick and remove the later one
+						//check poly pres events for the same note number (data1) on the same tick and remove the later one
+						if (Events[i].GetMsg().IsStd() &&
+							Events[i].GetMsg().GetStdStatus() == Msg.GetStdStatus() && 
+							Events[i].GetMsg().Data1 == Msg.Data1)
+						{
+							ControlEventIndexesToRemove.Add(i);
+						}
+					}
+					//remove excessive control changes or poly pressure on conformed last tick
+					if (!ControlEventIndexesToRemove.IsEmpty())
+					{
+						for (int32 i : ControlEventIndexesToRemove)
+						{
+							Events.RemoveAt(i);
+						}
+						NumItemsRemoved += ControlEventIndexesToRemove.Num();
+						NumControlEventsRemoved += ControlEventIndexesToRemove.Num();
+						TotalNumEventsRemoved += ControlEventIndexesToRemove.Num();
+					}
+				}
+
+				//update indices after removing excessive events
+				if (NumItemsRemoved == 0)
+				{
+					CurrentEventIndex -= 1;
+				}
+				else 
+				{
+					CurrentEventIndex -= NumItemsRemoved;
+					NumItemsRemoved = 0;
+				}
+			} 
+			else
+			{
+				CurrentEventIndex--;
+			}
+		}
+	}
+
+	// TracksChanged may have cause the song length data to
+	// shrink down because we have have removed all of the 
+	// "past-the-end" events above, and be left with an 
+	// unquantized length. So here will will force the song
+	// length to the proper valyue.
+	SongMaps.SetSongLengthTicks(NewLastEventTick + 1);
+
+	//Log conform summary 
+	FString MidiFileLengthConformSummary = FString::Printf(
+		TEXT("Midi file '%s.mid' length conformance summary: moved %d Midi event(s) to tick %d"),
+		*MidiFileName,
+		NumEventsMovedToLastTick,
+		NewLastEventTick);
+	if (TotalNumEventsRemoved > 0)
+	{
+		MidiFileLengthConformSummary.Append(FString::Printf(TEXT("\n\tRemoved %d Midi event(s):"), TotalNumEventsRemoved));
+		if (NumNoteOnNoteOffPairsRemoved > 0) MidiFileLengthConformSummary.Append(FString::Printf(TEXT("\n\t\t%d Note On/Note off event pair(s)"), NumNoteOnNoteOffPairsRemoved));
+		if (NumControlEventsRemoved > 0) MidiFileLengthConformSummary.Append(FString::Printf(TEXT("\n\t\t%d Aftertouch / Program / Pitch Wheel / CC event(s)"), NumControlEventsRemoved));
+		MidiFileLengthConformSummary.Append(FString::Printf(TEXT("\n\tat tick %d "), NewLastEventTick));
+	}
+	UE_LOG(LogMIDI, Log, TEXT("%s"), *MidiFileLengthConformSummary);
+	
+	return true;
+}
+
 void FMidiFileData::AddTempoChange(int32 TrackIdx, int32 Tick, float TempoBPM)
 {
 	check(Tracks.IsValidIndex(TrackIdx));
@@ -980,3 +915,34 @@ void FMidiFileData::AddTimeSigChange(int32 TrackIdx, int32 Tick, int32 InTimeSig
 	BarMap.AddTimeSignatureAtBarIncludingCountIn(AbsoluteBar, TimeSigNum, TimeSigDenom);
 }
 
+void FMidiFileData::ScanTracksForSongLengthChange()
+{
+	int32 NewLastEventTick = 0;
+	bool bIsFileEmpty = true;
+	for (auto& Track : Tracks)
+	{
+		FMidiEventList& Events = Track.GetRawEvents();
+		if (bIsFileEmpty && !Events.IsEmpty())
+		{
+			bIsFileEmpty = false;
+		}
+		int32 LastTickOnTrack = Events.IsEmpty() ? 0 : Events.Last().GetTick();
+		if (NewLastEventTick < LastTickOnTrack)
+		{
+			NewLastEventTick = LastTickOnTrack;
+		}
+	}
+
+	FSongLengthData& LengthData = SongMaps.GetSongLengthData();
+
+	//update length data in song length data if needed 
+	LengthData.LengthTicks = NewLastEventTick + 1;
+	LengthData.LastTick = NewLastEventTick;
+	LengthData.LengthFractionalBars = SongMaps.GetBarMap().TickToFractionalBarIncludingCountIn(LengthData.LengthTicks);
+}
+
+bool FMidiFileData::LengthIsAPerfectSubdivision() const
+{
+	return SongMaps.LengthIsAPerfectSubdivision();
+}
+
diff --git a/Engine/Plugins/Runtime/Harmonix/Source/HarmonixMidi/Private/HarmonixMidi/SongMaps.cpp b/Engine/Plugins/Runtime/Harmonix/Source/HarmonixMidi/Private/HarmonixMidi/SongMaps.cpp
--- a/Engine/Plugins/Runtime/Harmonix/Source/HarmonixMidi/Private/HarmonixMidi/SongMaps.cpp
+++ b/Engine/Plugins/Runtime/Harmonix/Source/HarmonixMidi/Private/HarmonixMidi/SongMaps.cpp
@@ -66,16 +66,10 @@ void FSongMaps::Copy(const FSongMaps& Other, int32 StartTick, int32 EndTick)
 	ChordMap.Finalize(LastTick);
 
 	memset(&LengthData, 0, sizeof(LengthData));
-	LengthData.LastTick = LastTick;
+	EMidiClockSubdivisionQuantization Division = EMidiClockSubdivisionQuantization::None;
+	LengthData.LastTick = QuantizeTickToAnyNearestSubdivision(EndTick, EMidiFileQuantizeDirection::Nearest, Division) - 1;
 	LengthData.LengthTicks = LengthData.LastTick + 1;
-	FMusicTimestamp Timestamp;
-	Timestamp = BarMap.TickToMusicTimestamp(float(LengthData.LastTick));
-	LengthData.LengthBars = FMath::IsNearlyZero(Timestamp.Beat-1.0f,1.0E-4) ? Timestamp.Bar - 1 : Timestamp.Bar;
-	if (LengthData.LengthBars == Timestamp.Bar)
-	{
-		LengthData.LengthTicks = int32(BarMap.MusicTimestampToTick({LengthData.LengthBars + BarMap.GetStartBar(), 1.0f}));
-		LengthData.LastTick = LengthData.LengthTicks - 1;
-	}
+	LengthData.LengthFractionalBars = BarMap.TickToFractionalBarIncludingCountIn(LengthData.LengthTicks);
 }
 
 bool FSongMaps::LoadFromStdMidiFile(const FString& FilePath)
@@ -144,14 +138,7 @@ bool FSongMaps::FinalizeRead(IMidiReader* Reader)
 	memset(&LengthData, 0, sizeof(LengthData));
 	LengthData.LastTick = LastTick;
 	LengthData.LengthTicks = LengthData.LastTick + 1;
-	FMusicTimestamp Timestamp;
-	Timestamp = BarMap.TickToMusicTimestamp(float(LengthData.LastTick));
-	LengthData.LengthBars = FMath::IsNearlyZero(Timestamp.Beat - 1.0f, 1.0E-4) ? Timestamp.Bar - 1 : Timestamp.Bar;
-	if (LengthData.LengthBars == Timestamp.Bar)
-	{
-		LengthData.LengthTicks = int32(BarMap.MusicTimestampToTick({ LengthData.LengthBars + BarMap.GetStartBar(), 1.0f }));
-		LengthData.LastTick = LengthData.LengthTicks - 1;
-	}
+	LengthData.LengthFractionalBars = BarMap.TickToFractionalBarIncludingCountIn(LengthData.LengthTicks);
 	return true;
 }
 
@@ -161,14 +148,6 @@ bool FSongMaps::ReadWithReader(FStdMidiFileReader& Reader)
 	return FinalizeRead(&Reader);
 }
 
-FArchive& operator<<(FArchive& Archive, FSongLengthData& LengthData)
-{
-	Archive << LengthData.LengthTicks;  
-	Archive << LengthData.LengthBars;   
-	Archive << LengthData.LastTick;
-	return Archive;
-} 
-
 void FSongMaps::EmptyAllMaps()
 {
 	TempoMap.Empty();
@@ -178,7 +157,7 @@ void FSongMaps::EmptyAllMaps()
 	ChordMap.Empty();
 	TrackNames.Empty();
 	LengthData.LastTick = 0;
-	LengthData.LengthBars = 0;
+	LengthData.LengthFractionalBars = 0.0f;
 	LengthData.LengthTicks = 0;
 }
 
@@ -191,7 +170,7 @@ bool FSongMaps::IsEmpty() const
 			ChordMap.IsEmpty() &&
 			TrackNames.IsEmpty() &&
 			LengthData.LastTick == 0 &&
-			LengthData.LengthBars == 0 &&
+			LengthData.LengthFractionalBars == 0 &&
 			LengthData.LengthTicks == 0;
 }
 
@@ -205,9 +184,207 @@ int32 FSongMaps::GetSongLengthBeats() const
 	return BeatMap.GetNumMapPoints();
 }
 
-int32 FSongMaps::GetSongLengthBars() const
+float FSongMaps::GetSongLengthFractionalBars() const
+{
+	return LengthData.LengthFractionalBars;
+}
+
+void FSongMaps::SetSongLengthTicks(int32 NewLengthTicks)
 {
-	return GetSongLengthData().LengthBars;
+	if (NewLengthTicks < 1)
+	{
+		UE_LOG(LogMIDI, Warning, TEXT("SetSongLengthTicks: Asked to set length less than 1. That is not possible. Setting to length 1!"));
+		NewLengthTicks = 1;
+	}
+	LengthData.LengthTicks = NewLengthTicks;
+	LengthData.LastTick    = NewLengthTicks - 1;
+	LengthData.LengthFractionalBars  = BarMap.TickToFractionalBarIncludingCountIn(LengthData.LengthTicks);
+}
+
+bool FSongMaps::LengthIsAPerfectSubdivision() const
+{
+	int32 BarIndex = 0;
+	int32 BeatInBar = 0;
+	int32 TickIndexInBeat = 0;
+	BarMap.TickToBarBeatTickIncludingCountIn(LengthData.LengthTicks, BarIndex, BeatInBar, TickIndexInBeat);
+	// The smallest subdivision we will consider is a 64th note triplet. 
+	// A sixty fourth note triplet divides a quarter note into 24 parts. 
+	int32 TicksPer64thTriplet = TicksPerQuarterNote / 24;
+	int32 TicksPer64th = TicksPerQuarterNote / 16;
+	return (TickIndexInBeat % TicksPer64thTriplet) == 0 || (TickIndexInBeat % TicksPer64th) == 0;
+}
+
+int32 FSongMaps::QuantizeTickToAnyNearestSubdivision(int32 InTick, EMidiFileQuantizeDirection Direction, EMidiClockSubdivisionQuantization& Division) const
+{
+	int32 BarIndex = 0;
+	int32 BeatInBar = 0;
+	int32 TickIndexInBeat = 0;
+	int32 TicksPerBar = 0;
+	int32 TicksPerBeat = 0;
+	BarMap.TickToBarBeatTickIncludingCountIn(InTick, BarIndex, BeatInBar, TickIndexInBeat, &TicksPerBar, &TicksPerBeat);
+	int32 BeatIndex = BeatInBar - 1; // BeatInBar is 1 bases!
+
+	if (BeatIndex == 0 && TickIndexInBeat == 0)
+	{
+		Division = EMidiClockSubdivisionQuantization::Bar;
+		return InTick;
+	}
+	if (TickIndexInBeat == 0)
+	{
+		Division = EMidiClockSubdivisionQuantization::Beat;
+		return InTick;
+	}
+
+	// Not so simple. Now we need to know the time signature...
+	int32 ZeroPoint = 0;
+	FTimeSignature TimeSignature(4, 4);
+	int32 BarMapPointIndex = BarMap.GetPointIndexForTick(InTick);
+	if (BarMapPointIndex >= 0)
+	{
+		const FTimeSignaturePoint& TimeSignaturePoint = BarMap.GetTimeSignaturePoint(BarMapPointIndex);
+		TimeSignature = TimeSignaturePoint.TimeSignature;
+		ZeroPoint = TimeSignaturePoint.StartTick;
+	}
+
+	// We "start" the quantization grid at the nearest proceeding time signature change...
+	int32 TickAtTimeSignature = InTick - ZeroPoint;
+	EMidiClockSubdivisionQuantization BestDivision = EMidiClockSubdivisionQuantization::None;
+	int32 BestDistanceFromDivision = std::numeric_limits<int32>::max();
+	
+	auto TryDivision = [this, &TickAtTimeSignature, &BestDivision, &BestDistanceFromDivision, &TimeSignature, &Division, &Direction](EMidiClockSubdivisionQuantization TryDivision)
+		{
+			int32 TicksPerDivision = SubdivisionToMidiTicks(TryDivision, TimeSignature);
+			int32 DistanceFromDivision = TickAtTimeSignature % TicksPerDivision;
+			switch (Direction)
+			{
+			case EMidiFileQuantizeDirection::Up:
+				DistanceFromDivision -= TicksPerDivision;
+				break;
+			case EMidiFileQuantizeDirection::Down:
+				// nothing to do here
+				break;
+			default:
+			case EMidiFileQuantizeDirection::Nearest:
+				if (DistanceFromDivision > (TicksPerDivision / 2))
+					DistanceFromDivision -= TicksPerDivision;
+				break;
+			}
+			if (DistanceFromDivision == 0)
+			{
+				Division = TryDivision;
+				return true;
+			}
+			else if (FMath::Abs(BestDistanceFromDivision) > FMath::Abs(DistanceFromDivision))
+			{
+				BestDivision = TryDivision;
+				BestDistanceFromDivision = DistanceFromDivision;
+			}
+			return false;
+		};
+
+	if (TryDivision(EMidiClockSubdivisionQuantization::Bar))
+		return InTick;
+	if (TryDivision(EMidiClockSubdivisionQuantization::Beat))
+		return InTick;
+	if (TryDivision(EMidiClockSubdivisionQuantization::QuarterNote))
+		return InTick;
+	if (TryDivision(EMidiClockSubdivisionQuantization::QuarterNoteTriplet))
+		return InTick;
+	if (TryDivision(EMidiClockSubdivisionQuantization::EighthNote))
+		return InTick;
+	if (TryDivision(EMidiClockSubdivisionQuantization::EighthNoteTriplet))
+		return InTick;
+	if (TryDivision(EMidiClockSubdivisionQuantization::SixteenthNote))
+		return InTick;
+	if (TryDivision(EMidiClockSubdivisionQuantization::SixteenthNoteTriplet))
+		return InTick;
+	if (TryDivision(EMidiClockSubdivisionQuantization::ThirtySecondNote))
+		return InTick;
+
+	Division = BestDivision;
+	return ZeroPoint + (TickAtTimeSignature - BestDistanceFromDivision);
+}
+
+int32 FSongMaps::QuantizeTickToNearestSubdivision(int32 InTick, EMidiFileQuantizeDirection Direction, EMidiClockSubdivisionQuantization Division) const
+{
+	int32 LowerTick = 0;
+	int32 UpperTick = 0;
+	GetTicksForNearestSubdivision(InTick, Division, LowerTick, UpperTick);
+	if (Direction == EMidiFileQuantizeDirection::Down)
+		return LowerTick;
+	if (Direction == EMidiFileQuantizeDirection::Up)
+		return UpperTick;
+	if ((InTick - LowerTick) < (UpperTick - InTick))
+		return LowerTick;
+	return UpperTick;
+}
+
+void FSongMaps::GetTicksForNearestSubdivision(int32 InTick, EMidiClockSubdivisionQuantization Division, int32& LowerTick, int32& UpperTick) const
+{
+	int32 TicksInSubdivision = 0;
+	int32 TickError = 0;
+	if (BarMap.IsEmpty())
+	{
+		FTimeSignature TimeSignature(4, 4);
+		TicksInSubdivision = SubdivisionToMidiTicks(Division, TimeSignature);
+		TickError = InTick % TicksInSubdivision;
+		LowerTick = InTick - TickError;
+		UpperTick = LowerTick + TicksInSubdivision;
+		return;
+	}
+
+	int32 BarIndex = 0;
+	int32 BeatInBar = 0;
+	int32 TickIndexInBeat = 0;
+	int32 BeatsPerBar = 0;
+	int32 TicksPerBeat = 0;
+	BarMap.TickToBarBeatTickIncludingCountIn(InTick, BarIndex, BeatInBar, TickIndexInBeat, &BeatsPerBar, &TicksPerBeat);
+	int32 TicksInBar = BeatsPerBar * TicksPerBeat;
+
+	if (Division == EMidiClockSubdivisionQuantization::Bar)
+	{
+		LowerTick = BarMap.BarBeatTickIncludingCountInToTick(BarIndex, 1, 0);
+		UpperTick = LowerTick + TicksInBar;
+		return;
+	}
+
+	if (Division == EMidiClockSubdivisionQuantization::Beat)
+	{
+		LowerTick = BarMap.BarBeatTickIncludingCountInToTick(BarIndex, BeatInBar, 0);
+		UpperTick = LowerTick + TicksPerBeat;
+		return;
+	}
+
+	// Not so simple. Now we need to know the time signature...
+	int32 ZeroPoint = 0;
+	FTimeSignature TimeSignature(4, 4);
+	int32 BarMapPointIndex = BarMap.GetPointIndexForTick(InTick);
+	if (BarMapPointIndex >= 0)
+	{
+		const FTimeSignaturePoint& TimeSignaturePoint = BarMap.GetTimeSignaturePoint(BarMapPointIndex);
+		TimeSignature = TimeSignaturePoint.TimeSignature;
+		ZeroPoint = TimeSignaturePoint.StartTick;
+	}
+
+	// We "start" the quantization grid at the nearest proceeding time signature change...
+	int32 TickAtTimeSignature = InTick - ZeroPoint;
+	TicksInSubdivision = SubdivisionToMidiTicks(Division, InTick);
+	TickError = TickAtTimeSignature % TicksInSubdivision;
+	// Now that we know the tick error we can apply it to our original input tick...
+	LowerTick = InTick - TickError;
+	UpperTick = LowerTick + TicksInSubdivision;
+}
+
+FString FSongMaps::GetSongLengthString() const
+{
+	int32 BarIndex;
+	int32 BeatInBar;
+	int32 TickIndexInBeat;
+	int32 BeatsPerBar;
+	int32 TicksPerBeat;
+	BarMap.TickToBarBeatTickIncludingCountIn(LengthData.LengthTicks, BarIndex, BeatInBar, TickIndexInBeat, &BeatsPerBar, &TicksPerBeat);
+	int32 BeatIndex = BeatInBar - 1; // BeatInBar is 1 based
+	return FString::Printf(TEXT("%d | %.3f"), BarIndex, (float)BeatIndex + (float)TickIndexInBeat / (float)TicksPerBeat );
 }
 
 void FSongMaps::StringLengthToMT(const FString& LengthString, int32& OutBars, int32& OutTicks)
@@ -582,7 +759,7 @@ void FSongMaps::SetLengthTotalBars(int32 Bars)
 	memset(&LengthData, 0, sizeof(LengthData));
 	LengthData.LengthTicks = BarMap.BarIncludingCountInToTick(Bars);
 	LengthData.LastTick = LengthData.LengthTicks - 1;
-	LengthData.LengthBars = Bars;
+	LengthData.LengthFractionalBars = Bars;
 }
 
 int32 FSongMaps::CalculateMidiTick(const FMusicTimestamp& Timestamp, const EMidiClockSubdivisionQuantization Quantize) const
@@ -611,45 +788,44 @@ int32 FSongMaps::CalculateMidiTick(const FMusicTimestamp& Timestamp, const EMidi
 	return TriggerTick;
 }
 
-int32 FSongMaps::SubdivisionToMidiTicks(const EMidiClockSubdivisionQuantization Division, const int32 Tick) const
+int32 FSongMaps::SubdivisionToMidiTicks(const EMidiClockSubdivisionQuantization Division, const FTimeSignature& TimeSignature) const
 {
-	int32 BarMapPointIndex = GetBarMap().GetPointIndexForTick(Tick);
-	if (BarMapPointIndex < 0)
-	{
-		return 1;
-	}
-	const FTimeSignature* TimeSignature = GetTimeSignatureAtTick(Tick);
-	if (!TimeSignature)
-	{
-		return 1;
-	}
-
-	using namespace Harmonix::Midi::Constants;
-
 	switch (Division)
 	{
 	case EMidiClockSubdivisionQuantization::None:                   return 1;
-	case EMidiClockSubdivisionQuantization::Bar: 					return GetBarMap().GetTicksInBarAfterPoint(BarMapPointIndex);
-	case EMidiClockSubdivisionQuantization::Beat:					return (GTicksPerQuarterNoteInt * 4) / TimeSignature->Denominator;
-	case EMidiClockSubdivisionQuantization::ThirtySecondNote:		return GTicksPerQuarterNoteInt / 8;
-	case EMidiClockSubdivisionQuantization::SixteenthNote:			return GTicksPerQuarterNoteInt / 4;
-	case EMidiClockSubdivisionQuantization::EighthNote:				return GTicksPerQuarterNoteInt / 2;
-	case EMidiClockSubdivisionQuantization::QuarterNote:			return GTicksPerQuarterNoteInt;
-	case EMidiClockSubdivisionQuantization::HalfNote:				return GTicksPerQuarterNoteInt * 2;
-	case EMidiClockSubdivisionQuantization::WholeNote:				return GTicksPerQuarterNoteInt * 4;
-	case EMidiClockSubdivisionQuantization::DottedSixteenthNote:	return (GTicksPerQuarterNoteInt / 4) + (GTicksPerQuarterNoteInt / 8);
-	case EMidiClockSubdivisionQuantization::DottedEighthNote:		return (GTicksPerQuarterNoteInt / 2) + (GTicksPerQuarterNoteInt / 4);
-	case EMidiClockSubdivisionQuantization::DottedQuarterNote:		return (GTicksPerQuarterNoteInt)+(GTicksPerQuarterNoteInt / 2);
-	case EMidiClockSubdivisionQuantization::DottedHalfNote:			return (GTicksPerQuarterNoteInt * 2) + (GTicksPerQuarterNoteInt);
-	case EMidiClockSubdivisionQuantization::DottedWholeNote:		return (GTicksPerQuarterNoteInt * 4) + (GTicksPerQuarterNoteInt * 2);
-	case EMidiClockSubdivisionQuantization::SixteenthNoteTriplet:   return (GTicksPerQuarterNoteInt / 2) / 3;
-	case EMidiClockSubdivisionQuantization::EighthNoteTriplet:		return GTicksPerQuarterNoteInt / 3;
-	case EMidiClockSubdivisionQuantization::QuarterNoteTriplet:		return (GTicksPerQuarterNoteInt * 2) / 3;
-	case EMidiClockSubdivisionQuantization::HalfNoteTriplet:        return (GTicksPerQuarterNoteInt * 4) / 3;
+	case EMidiClockSubdivisionQuantization::Bar: 					return TimeSignature.Numerator * ((TicksPerQuarterNote * 4)/TimeSignature.Denominator);
+	case EMidiClockSubdivisionQuantization::Beat:					return (TicksPerQuarterNote * 4) / TimeSignature.Denominator;
+	case EMidiClockSubdivisionQuantization::ThirtySecondNote:		return TicksPerQuarterNote / 8;
+	case EMidiClockSubdivisionQuantization::SixteenthNote:			return TicksPerQuarterNote / 4;
+	case EMidiClockSubdivisionQuantization::EighthNote:				return TicksPerQuarterNote / 2;
+	case EMidiClockSubdivisionQuantization::QuarterNote:			return TicksPerQuarterNote;
+	case EMidiClockSubdivisionQuantization::HalfNote:				return TicksPerQuarterNote * 2;
+	case EMidiClockSubdivisionQuantization::WholeNote:				return TicksPerQuarterNote * 4;
+	case EMidiClockSubdivisionQuantization::DottedSixteenthNote:	return (TicksPerQuarterNote / 4) + (TicksPerQuarterNote / 8);
+	case EMidiClockSubdivisionQuantization::DottedEighthNote:		return (TicksPerQuarterNote / 2) + (TicksPerQuarterNote / 4);
+	case EMidiClockSubdivisionQuantization::DottedQuarterNote:		return (TicksPerQuarterNote)+(TicksPerQuarterNote / 2);
+	case EMidiClockSubdivisionQuantization::DottedHalfNote:			return (TicksPerQuarterNote * 2) + (TicksPerQuarterNote);
+	case EMidiClockSubdivisionQuantization::DottedWholeNote:		return (TicksPerQuarterNote * 4) + (TicksPerQuarterNote * 2);
+	case EMidiClockSubdivisionQuantization::SixteenthNoteTriplet:   return (TicksPerQuarterNote / 2) / 3;
+	case EMidiClockSubdivisionQuantization::EighthNoteTriplet:		return TicksPerQuarterNote / 3;
+	case EMidiClockSubdivisionQuantization::QuarterNoteTriplet:		return (TicksPerQuarterNote * 2) / 3;
+	case EMidiClockSubdivisionQuantization::HalfNoteTriplet:        return (TicksPerQuarterNote * 4) / 3;
 	default:	                                             		checkNoEntry();	return 1;
 	}
 }
 
+int32 FSongMaps::SubdivisionToMidiTicks(const EMidiClockSubdivisionQuantization Division, const int32 AtTick) const
+{
+	FTimeSignature TimeSignature(4,4);
+	int32 BarMapPointIndex = BarMap.GetPointIndexForTick(AtTick);
+	if (BarMapPointIndex >= 0)
+	{
+		const FTimeSignaturePoint& TimeSignaturePoint = BarMap.GetTimeSignaturePoint(BarMapPointIndex);
+		TimeSignature = TimeSignaturePoint.TimeSignature;
+	}
+	return SubdivisionToMidiTicks(Division, TimeSignature);
+}
+
 float FSongMaps::SubdivisionToBeats(EMidiClockSubdivisionQuantization Subdivision, const FTimeSignature& TimeSignature)
 {
 	// Easy cases first
diff --git a/Engine/Plugins/Runtime/Harmonix/Source/HarmonixMidi/Public/HarmonixMidi/BarMap.h b/Engine/Plugins/Runtime/Harmonix/Source/HarmonixMidi/Public/HarmonixMidi/BarMap.h
--- a/Engine/Plugins/Runtime/Harmonix/Source/HarmonixMidi/Public/HarmonixMidi/BarMap.h
+++ b/Engine/Plugins/Runtime/Harmonix/Source/HarmonixMidi/Public/HarmonixMidi/BarMap.h
@@ -244,6 +244,21 @@ public:
 	 */
 	int32 BarBeatTickIncludingCountInToTick(int32 BarIndex, int32 BeatInBar, int32 TickInBeat) const;
 
+	/**
+	 * Take "raw" tick (tick from the beginning of the midi data and in units of TicksPerQuarterNote)
+	 * and return a bar (0 based), beat in bar (1 based) within that bar, and tick (0 based) within that beat. 
+	 * Here, the returned position is a 'raw' position from the beginning of the midi data. There is
+	 * no consideration of count-in or pickup bars.
+	 *
+	 * @param RawTick from the beginning of the midi data.
+	 * @param OutBarIndex (output)
+	 * @param OutBeatInBar (1 based!) (output)
+	 * @param OutTickIndexInBeat (output)
+	 * @param OutBeatsPerBar (output - optional)
+	 * @param OutTicksPerBeat (output - optional)
+	 */
+	void TickToBarBeatTickIncludingCountIn(int32 RawTick, int32& OutBarIndex, int32& OutBeatInBarIndex, int32& OutTickIndexInBeat, int32* OutBeatsPerBar = nullptr, int32* OutTicksPerBeat = nullptr) const;
+
 	/**
 	 * Takes an absolute tick (from the beginning of the midi data) and calculates the fractional bar,
 	 * where bar 0.0 is the beginning of the raw midi data.
diff --git a/Engine/Plugins/Runtime/Harmonix/Source/HarmonixMidi/Public/HarmonixMidi/MidiEvent.h b/Engine/Plugins/Runtime/Harmonix/Source/HarmonixMidi/Public/HarmonixMidi/MidiEvent.h
--- a/Engine/Plugins/Runtime/Harmonix/Source/HarmonixMidi/Public/HarmonixMidi/MidiEvent.h
+++ b/Engine/Plugins/Runtime/Harmonix/Source/HarmonixMidi/Public/HarmonixMidi/MidiEvent.h
@@ -57,6 +57,7 @@ protected:
 
 private:
 	friend struct FMidiTrack;
+	friend struct FMidiFileData;
 	friend class  UMidiFile;
 
 	UPROPERTY()
diff --git a/Engine/Plugins/Runtime/Harmonix/Source/HarmonixMidi/Public/HarmonixMidi/MidiFile.h b/Engine/Plugins/Runtime/Harmonix/Source/HarmonixMidi/Public/HarmonixMidi/MidiFile.h
--- a/Engine/Plugins/Runtime/Harmonix/Source/HarmonixMidi/Public/HarmonixMidi/MidiFile.h
+++ b/Engine/Plugins/Runtime/Harmonix/Source/HarmonixMidi/Public/HarmonixMidi/MidiFile.h
@@ -16,14 +16,6 @@ class UAssetImportData;
 class FMidiFileProxy;
 using FMidiFileProxyPtr = TSharedPtr<FMidiFileProxy, ESPMode::ThreadSafe>;
 
-//midi file length conform options,rounding file length to a whole bar up or down 
-enum class EMidiFileLengthConformOption : uint8
-{
-	RoundDown,
-	RoundUp,
-	Nearest
-};
-
 USTRUCT(BlueprintType, Meta = (DisplayName = "MIDI File Data"))
 struct HARMONIXMIDI_API FMidiFileData
 {
@@ -31,7 +23,6 @@ struct HARMONIXMIDI_API FMidiFileData
 
 	FMidiFileData()
 		: TicksPerQuarterNote(Harmonix::Midi::Constants::GTicksPerQuarterNoteInt)
-		, LastEventTick(0)
 	{}
 
 	bool operator==(const FMidiFileData& Other) const;
@@ -43,7 +34,6 @@ struct HARMONIXMIDI_API FMidiFileData
 		SongMaps.EmptyAllMaps();
 		Tracks.Empty();
 		Tracks.Emplace("Conductor");
-		LastEventTick = 0;
 	}
 
 	bool IsEmpty() const 
@@ -51,8 +41,7 @@ struct HARMONIXMIDI_API FMidiFileData
 		return MidiFileName.IsEmpty() &&
 			TicksPerQuarterNote == Harmonix::Midi::Constants::GTicksPerQuarterNoteInt &&
 			SongMaps.IsEmpty() &&
-			(Tracks.IsEmpty() || (Tracks.Num() == 1 && Tracks[0].GetNumEvents() == 1 && Tracks[0].GetName()->Compare("Conductor", ESearchCase::IgnoreCase) == 0)) &&
-			LastEventTick == 0;
+			(Tracks.IsEmpty() || (Tracks.Num() == 1 && Tracks[0].GetNumEvents() == 1 && Tracks[0].GetName()->Compare("Conductor", ESearchCase::IgnoreCase) == 0));
 	}
 
 	UPROPERTY(BlueprintReadOnly, Category = "MidiFile")
@@ -63,13 +52,16 @@ struct HARMONIXMIDI_API FMidiFileData
 	FSongMaps SongMaps;
 	UPROPERTY(BlueprintReadOnly, Category = "MidiFile")
 	TArray<FMidiTrack> Tracks;
-	UPROPERTY(BlueprintReadOnly, Category = "MidiFile")
-	int32 LastEventTick;
+
+	int32 GetLastEventTick() const;
 
 	int32 FindTrackIndexByName(const FString& TrackName);
 	
 	void PostSerialize(const FArchive& Ar);
 
+	bool ConformToLength(int32 NewLengthTicks);
+	bool ConformToLengthGivenLastEventTick(int32 NewLastEventTick);
+
 	/**
 	* Adds a tempo change to the midi data at the given tick for the given track idx
 	* 
@@ -86,6 +78,10 @@ struct HARMONIXMIDI_API FMidiFileData
 	* asserts valid TrackIdx
 	*/
 	void AddTimeSigChange(int32 TrackIdx, int32 Tick, int32 TimeSigNum, int32 TimeSigDenom);
+
+	void ScanTracksForSongLengthChange();
+
+	bool LengthIsAPerfectSubdivision() const;
 };
 
 template<>
@@ -179,9 +175,9 @@ public:
 	 * This function must be called if any changes are made to any of the tracks of this midi file. It will 
 	 * do various internal tasks necessary to assure the midi data is consistent and "playable".
 	 */
-	void TracksChanged();
+	void ScanTracksForSongLengthChange();
 
-	int32  GetLastEventTick() const { return TheMidiData.LastEventTick; }
+	int32  GetLastEventTick() const { return TheMidiData.GetLastEventTick(); }
 
 #if WITH_EDITORONLY_DATA
 	// Import data for this MidiFileAsset
@@ -200,37 +196,22 @@ public:
 
 	int32 GetStartBar() { return TheMidiData.SongMaps.GetBarMap().GetStartBar(); };
 
-	//Midi File Conform Option, keep track if the files' length are conformed during import
-	UPROPERTY(BlueprintReadOnly, Category = "MidiFile")
-	bool bLengthRoundedDown = false;
+	// This does a quick and dirty to check to see if the length of the midi is 
+	// on a musical subdivision. It is used primarily during asset importing, and
+	// if it returns true the importer will do further checking to see how far off
+	// it is, and suggest either conforming the length via QuantizeLengthToNearestPerfectSubdivision
+	// (for "off-by-one" type small errors) or QuantizeLengthToSubdivision (when the
+	// error is large and the user should choose a specific subdivision and direction
+	// for conforming the file.
+	bool LengthIsAPerfectSubdivision() const;
 
-	UPROPERTY(BlueprintReadOnly, Category = "MidiFile")
-	bool bLengthRoundedUp = false;
+	void QuantizeLengthToNearestPerfectSubdivision(const EMidiFileQuantizeDirection Direction);
 
-	UPROPERTY(BlueprintReadOnly, Category = "MidiFile")
-	bool bLengthRoundedToNearest = false;
+	void QuantizeLengthToSubdivision(const EMidiFileQuantizeDirection Direction, const EMidiClockSubdivisionQuantization Subdivision);
 
-	/**
-	* Takes a file conform option (currently the last integer bar or the last integer beat),determine whether or not 
-	* the file needs to be conformed/are already conformed.
-	* @param A midi file conform option (last bar or last beat)
-	* @return A bool indicating the file should or should not be conformed
-	*/
-	bool ShouldConformMidiFileLength(EMidiFileLengthConformOption Option);
+	bool ConformToLength(int32 NewLengthTicks);
 
-	/**
-	* Conform Midi File length(bars) to an integer number of bars:
-	* LengthTicks and LengthBars in SongLengthData are rounded up to an integer number of bars
-	* if the midi file has a fractional bar length or ends on a fractional beat, 
-	* this function takes a midi conform option (Round Down, and Round Up,Round to Nearest):
-	* Round Down: modifies/moves the raw midi events in every midi track of a midi file to the previous integer bar
-	*			  and remove excessive note on/note off pairs, control change events, pitch bend events, etc. at the tick they are moved to
-	* Round Up:	since file length is automatically rounded up, this option does NOT do or change anything to the Midi events in file
-	* Round to Nearest: rounding the file length to the nearest integer (either round down or round up), depending on the file's fractional length
-	* e.g. if a midi file has length 5.875 bars, Round Down conforms it to 5 bars, Round Up conforms it to 6 bars,
-	* and Round To Nearest conforms it to 6 bars (5.875 > 5.5)
-	*/
-	void ConformMidiFileLength(EMidiFileLengthConformOption Option, bool Force = false);
+	bool ConformToLengthGivenLastEventTick(int32 NewLastEventTick);
 
 	const FSongMaps* GetSongMaps() const { return &TheMidiData.SongMaps; }
 	FSongMaps* GetSongMaps() { return &TheMidiData.SongMaps; }
diff --git a/Engine/Plugins/Runtime/Harmonix/Source/HarmonixMidi/Public/HarmonixMidi/MidiMsg.h b/Engine/Plugins/Runtime/Harmonix/Source/HarmonixMidi/Public/HarmonixMidi/MidiMsg.h
--- a/Engine/Plugins/Runtime/Harmonix/Source/HarmonixMidi/Public/HarmonixMidi/MidiMsg.h
+++ b/Engine/Plugins/Runtime/Harmonix/Source/HarmonixMidi/Public/HarmonixMidi/MidiMsg.h
@@ -49,6 +49,8 @@ public:
 	bool  IsAllNotesOff() const { return Type == EType::Runtime && Status == Harmonix::Midi::Constants::GRuntimeAllNotesOffStatus; }
 	bool  IsAllNotesKill() const { return Type == EType::Runtime && Status == Harmonix::Midi::Constants::GRuntimeAllNotesKillStatus; }
 	bool  IsControlChange() const { return IsStd() && Harmonix::Midi::Constants::IsControl(Status); }
+	bool  IsTempo() const { return Type == EType::Tempo; }
+	bool  IsTimeSignature() const { return Type == EType::TimeSig; }
 
 	/** Construct a standard (std) midi message */
 	FMidiMsg(uint8 InStatus, uint8 InData1, uint8 InData2);
diff --git a/Engine/Plugins/Runtime/Harmonix/Source/HarmonixMidi/Public/HarmonixMidi/SongMaps.h b/Engine/Plugins/Runtime/Harmonix/Source/HarmonixMidi/Public/HarmonixMidi/SongMaps.h
--- a/Engine/Plugins/Runtime/Harmonix/Source/HarmonixMidi/Public/HarmonixMidi/SongMaps.h
+++ b/Engine/Plugins/Runtime/Harmonix/Source/HarmonixMidi/Public/HarmonixMidi/SongMaps.h
@@ -14,6 +14,14 @@ class FStdMidiFileReader;
 class FSongMapReceiver;
 class UMidiFile;
 
+UENUM()
+enum class EMidiFileQuantizeDirection : uint8
+{
+	Nearest,
+	Up,
+	Down
+};
+
 UENUM(BlueprintType)
 enum class EMidiClockSubdivisionQuantization : uint8
 {
@@ -46,14 +54,14 @@ public:
 	UPROPERTY()
 	int32 LengthTicks = 0;
 	UPROPERTY()
-	int32 LengthBars = 0;
+	float LengthFractionalBars = 0.0f;
 	UPROPERTY()
 	int32 LastTick = 0;
 
 	bool operator==(const FSongLengthData& Other) const
 	{
 		return	LengthTicks == Other.LengthTicks &&
-				LengthBars == Other.LengthBars && 
+				LengthFractionalBars == Other.LengthFractionalBars && 
 				LastTick == Other.LastTick;
 	}
 };
@@ -139,7 +147,8 @@ public:
 	void                  EmptyBarMap() { BarMap.Empty(); }
 	void                  SetLengthTotalBars(int32 Bars);
 	int32                 CalculateMidiTick(const FMusicTimestamp& Timestamp, const EMidiClockSubdivisionQuantization Quantize) const;
-	int32                 SubdivisionToMidiTicks(const EMidiClockSubdivisionQuantization Division, const int32 Tick) const;
+	int32                 SubdivisionToMidiTicks(const EMidiClockSubdivisionQuantization Division, const int32 AtTick) const;
+	int32                 SubdivisionToMidiTicks(const EMidiClockSubdivisionQuantization Division, const FTimeSignature& TimeSignature) const;
 	static float          SubdivisionToBeats(EMidiClockSubdivisionQuantization Subdivision, const FTimeSignature& TimeSignature);
 
 
@@ -177,10 +186,18 @@ public:
 	const FSongLengthData& GetSongLengthData() const { return LengthData; }
 	float GetSongLengthMs() const;
 	int32 GetSongLengthBeats() const;
-	int32 GetSongLengthBars() const;
+	float GetSongLengthFractionalBars() const;
 
 	int32 GetTicksPerQuarterNote() const { return TicksPerQuarterNote; }
 
+	void SetSongLengthTicks(int32 NewLengthTicks);
+	bool LengthIsAPerfectSubdivision() const;
+
+	int32 QuantizeTickToAnyNearestSubdivision(int32 InTick, EMidiFileQuantizeDirection Direction, EMidiClockSubdivisionQuantization& Division) const;
+	int32 QuantizeTickToNearestSubdivision(int32 InTick, EMidiFileQuantizeDirection Direction, EMidiClockSubdivisionQuantization Division) const;
+	void GetTicksForNearestSubdivision(int32 InTick, EMidiClockSubdivisionQuantization Division, int32& LowerTick, int32& UpperTick) const;
+	FString GetSongLengthString() const;
+
 protected:
 	friend class FSongMapReceiver;
 	UPROPERTY()
diff --git a/Engine/Plugins/Runtime/Harmonix/Source/HarmonixMidiEditor/Private/MidiFileDetailCustomization.cpp b/Engine/Plugins/Runtime/Harmonix/Source/HarmonixMidiEditor/Private/MidiFileDetailCustomization.cpp
--- a/Engine/Plugins/Runtime/Harmonix/Source/HarmonixMidiEditor/Private/MidiFileDetailCustomization.cpp
+++ b/Engine/Plugins/Runtime/Harmonix/Source/HarmonixMidiEditor/Private/MidiFileDetailCustomization.cpp
@@ -3,9 +3,16 @@
 #include "MidiFileDetailCustomization.h"
 #include "HarmonixMidi/MidiFile.h"
 #include "Math/UnitConversion.h"
+#include "MidiFileFactory.h"
 
 #define LOCTEXT_NAMESPACE "HarmonixMIDIEditor"
 
+void FMidiFileDetailCustomization::CustomizeDetails(const TSharedPtr<IDetailLayoutBuilder>& InDetailBuilder)
+{
+	DetailBuilder = InDetailBuilder;
+	CustomizeDetails(*DetailBuilder);
+}
+
 void FMidiFileDetailCustomization::CustomizeDetails(IDetailLayoutBuilder& DetailLayout)
 {
 	//Details Panel for the UROPERTY AssetImportData, no actual customization, just reordering its position 
@@ -31,185 +38,16 @@ void FMidiFileDetailCustomization::CustomizeDetails(IDetailLayoutBuilder& Detail
 	MidiFile = Objects.Last();
 	TWeakObjectPtr<UMidiFile> MidiFileBeingEdited = Cast<UMidiFile>(MidiFile);
 	
-	//determine if file length is already conformed during import:
-	//if it's already conformed just display the conformed length otherwise the original fractional length
-	int32 MidiFileLengthConformed = MidiFileBeingEdited->GetSongMaps()->GetSongLengthData().LengthBars;
-	float MidiFileFractionalLength = MidiFileBeingEdited->GetSongMaps()->GetBarIncludingCountInAtTick(MidiFileBeingEdited->GetLastEventTick());
-	if (MidiFileBeingEdited->bLengthRoundedUp)
-	{
-		//if the file length is 1 bar, it CANNOT be rounded down to 0 bar
-		//otherwise still allowing it to be conformed by rounding down by showing a message on a  tooptip 
-		if (MidiFileLengthConformed != 1)
-		{
-			RoundedLengthToolTipText = FText::FromString(FString::Printf(TEXT("*rounded up from file length of %.3f bars, can still round down"), MidiFileFractionalLength));
-			FileLengthText = FText::FromString(FString::Printf(TEXT("%d Bar(s)*"), MidiFileLengthConformed));
-		}
-		else {
-			RoundedLengthToolTipText = FText::FromString(TEXT(""));
-			FileLengthText = FText::FromString(FString::Printf(TEXT("%d Bar(s)"), MidiFileLengthConformed));
-		}
-
-	}
-	else if (MidiFileBeingEdited->bLengthRoundedDown)
-	{
-		RoundedLengthToolTipText = FText::FromString(TEXT(""));
-		FileLengthText = FText::FromString(FString::Printf(TEXT("%d Bar(s)"), MidiFileLengthConformed));
-	}
-	else {
-		RoundedLengthToolTipText = FText::FromString(TEXT(""));
-		//if the file doesn't need to be conformed (and hence never conformed), display the integer bar length, otherwise display fractional bar length
-		if (!MidiFileBeingEdited->ShouldConformMidiFileLength(EMidiFileLengthConformOption::RoundDown) && !MidiFileBeingEdited->ShouldConformMidiFileLength(EMidiFileLengthConformOption::RoundUp))
-		{
-			FileLengthText = FText::FromString(FString::Printf(TEXT("%d Bar(s)"), MidiFileLengthConformed));
-		}
-		else {
-			FileLengthText = FText::FromString(FString::Printf(TEXT("%.3f Bar(s)"), MidiFileFractionalLength));
-		}
-	}
-
-	//determine whether the file length should be conformed, disable buttons if not
-	auto ShouldConformRoundDown = [MidiFileBeingEdited]
-	{
-		return MidiFileBeingEdited->ShouldConformMidiFileLength(EMidiFileLengthConformOption::RoundDown);
-	};
-
-	auto ShouldConformRoundUp = [this,MidiFileBeingEdited,&MidiFileLengthCategory]
-	{
-		return MidiFileBeingEdited->ShouldConformMidiFileLength(EMidiFileLengthConformOption::RoundUp);
-	};
-
-	auto ShouldConformRoundToNearest = [MidiFileBeingEdited]
-	{
-		return MidiFileBeingEdited->ShouldConformMidiFileLength(EMidiFileLengthConformOption::Nearest);
-	};
-
-
-	//conform buttons OnClick callback function
-	auto OnLengthConformedRoundDown = [this, MidiFileBeingEdited, &MidiFileLengthCategory,MidiFileFractionalLength]
-	{
-		if (MidiFileBeingEdited.IsValid())
-			MidiFileBeingEdited->ConformMidiFileLength(EMidiFileLengthConformOption::RoundDown);
-
-		//update display text for file length
-		FileLengthText = FText::FromString(FString::Printf(TEXT("%d Bar(s)"), MidiFileBeingEdited->GetSongMaps()->GetSongLengthData().LengthBars));
-		FileLengthTextBlock->SetText(FileLengthText);
-		RoundedLengthToolTipText = FText::FromString(TEXT(""));
-		FileLengthTextBlock->SetToolTipText(RoundedLengthToolTipText);
-
-		return FReply::Handled();
-	};
-
-	auto OnLengthConformedRoundUp = [this, MidiFileBeingEdited, &MidiFileLengthCategory,MidiFileFractionalLength]
-	{
-		if (MidiFileBeingEdited.IsValid())
-			MidiFileBeingEdited->ConformMidiFileLength(EMidiFileLengthConformOption::RoundUp);
-		
-		//update display text for file length
-		FileLengthText = FText::FromString(FString::Printf(TEXT("%d Bar(s)*"), MidiFileBeingEdited->GetSongMaps()->GetSongLengthData().LengthBars));
-		FileLengthTextBlock->SetText(FileLengthText);
-		RoundedLengthToolTipText = FText::FromString(FString::Printf(TEXT("*rounded up from file length of %.3f bars, can still round down"), MidiFileFractionalLength));
-		FileLengthTextBlock->SetToolTipText(RoundedLengthToolTipText);
-
-		return FReply::Handled();
-	};
-
-	auto OnLengthConformedRoundToNearest = [this, MidiFileBeingEdited, &MidiFileLengthCategory, MidiFileFractionalLength]
-	{
-		if (MidiFileBeingEdited.IsValid())
-			MidiFileBeingEdited->ConformMidiFileLength(EMidiFileLengthConformOption::Nearest);
-
-	
-		if (MidiFileBeingEdited->bLengthRoundedDown && MidiFileBeingEdited->bLengthRoundedToNearest)
-		{
-			FileLengthText = FText::FromString(FString::Printf(TEXT("%d Bar(s)"), MidiFileBeingEdited->GetSongMaps()->GetSongLengthData().LengthBars));
-			RoundedLengthToolTipText = FText::FromString(TEXT(""));
-		}
-
-		if (MidiFileBeingEdited->bLengthRoundedUp && MidiFileBeingEdited->bLengthRoundedToNearest)
-		{
-			FileLengthText = FText::FromString(FString::Printf(TEXT("%d Bar(s)*"), MidiFileBeingEdited->GetSongMaps()->GetSongLengthData().LengthBars));
-			RoundedLengthToolTipText = FText::FromString(FString::Printf(TEXT("*rounded up from file length of %.3f bars, can still round down"), MidiFileFractionalLength));
-		}
-		 
-		FileLengthTextBlock->SetText(FileLengthText);
-		FileLengthTextBlock->SetToolTipText(RoundedLengthToolTipText);
-
-		return FReply::Handled();
-	};
-
 	//Detail Customization for the File Length Row
-	MidiFileLengthCategory.AddCustomRow(FText::FromString("Midi File Length"))
-	.NameContent()
-	[
-		SNew(STextBlock)
-		.Text(FText::FromString("File Length"))
-		.Font(IDetailLayoutBuilder::GetDetailFont())
-	]
-	.ValueContent()
-	[
-		SNew(SHorizontalBox)
-		+ SHorizontalBox::Slot()
-		.HAlign(HAlign_Left)
-		.VAlign(VAlign_Center)
-		.FillWidth(0.9f)
+	LengthRow = &MidiFileLengthCategory.AddCustomRow(FText::FromString("Midi File Length"));
+	LengthRow->NameContent()
 		[
-			SAssignNew(FileLengthTextBlock,STextBlock)
-			.Text(FileLengthText)
+			SNew(STextBlock)
+			.Text(FText::FromString("File Length"))
 			.Font(IDetailLayoutBuilder::GetDetailFont())
-			.ToolTipText(RoundedLengthToolTipText)
-		]
-		+ SHorizontalBox::Slot()
-			.HAlign(HAlign_Right)
-			.VAlign(VAlign_Center)
-			.FillWidth(0.1f)
-		[
-			SNew(SHorizontalBox)
-			+ SHorizontalBox::Slot()
-			.HAlign(HAlign_Right)
-			.VAlign(VAlign_Center)
-			.AutoWidth()
-			[
-				SNew(STextBlock)
-				.Font(IDetailLayoutBuilder::GetDetailFont())
-				.Text(FText::FromString("Conform file length to a whole bar by: "))
-			]
-			+ SHorizontalBox::Slot()
-			.HAlign(HAlign_Left)
-			.VAlign(VAlign_Center)
-			.AutoWidth()
-			[
-				SNew(SButton)
-				.VAlign(VAlign_Center)
-				.Text(FText::FromString("Round to Nearest"))
-				.OnClicked_Lambda(OnLengthConformedRoundToNearest)
-				.IsEnabled_Lambda(ShouldConformRoundToNearest)
-			]
-			+ SHorizontalBox::Slot()
-			.HAlign(HAlign_Left)
-			.VAlign(VAlign_Center)
-			.AutoWidth()
-			.Padding(10, 0)
-			[
-				SNew(SButton)
-				.VAlign(VAlign_Center)
-				.Text(FText::FromString("Round Down"))
-				.OnClicked_Lambda(OnLengthConformedRoundDown)
-				.IsEnabled_Lambda(ShouldConformRoundDown)
-			]
-			+ SHorizontalBox::Slot()
-				.HAlign(HAlign_Right)
-				.VAlign(VAlign_Center)
-				.AutoWidth()
-			[
-				SNew(SButton)
-				.VAlign(VAlign_Center)
-				.Text(FText::FromString("Round Up"))
-				.OnClicked_Lambda(OnLengthConformedRoundUp)
-				.IsEnabled_Lambda(ShouldConformRoundUp)
-			]
+		];
 
-		]
-	];
+	BuildLengthRow(MidiFileBeingEdited.Get());
 
 	IDetailCategoryBuilder& TrackInfoCategory = DetailLayout.EditCategory(TEXT("Track Info"), LOCTEXT("TrackInfo", "Tracks:"));
 	const UMidiFile::FMidiTrackList& Tracks = MidiFileBeingEdited->GetTracks();
@@ -260,4 +98,142 @@ void FMidiFileDetailCustomization::CustomizeDetails(IDetailLayoutBuilder& Detail
 		];
 	}
 }
+
+void FMidiFileDetailCustomization::BuildLengthRow(UMidiFile* TheMidiFile)
+{
+	bool bLengthIsAPerfectSubdivision = TheMidiFile->GetSongMaps()->LengthIsAPerfectSubdivision();
+	if (bLengthIsAPerfectSubdivision)
+	{
+		FileLengthText = FText::FromString(TheMidiFile->GetSongMaps()->GetSongLengthString());
+	}
+	else
+	{
+		FileLengthText = FText::Format(LOCTEXT("NonQualtizedMIDIFileLength", "{0} (Warning: File is not quantized to any standard musical subdivision.)"),
+			FText::FromString(TheMidiFile->GetSongMaps()->GetSongLengthString()));
+	}
+
+	//conform buttons OnClick callback function
+	auto OnLengthConformTrivial = [this, TheMidiFile]
+		{
+			TheMidiFile->QuantizeLengthToNearestPerfectSubdivision(EMidiFileQuantizeDirection::Nearest);
+			TheMidiFile->PostEditChange();
+			BuildLengthRow(TheMidiFile);
+			return FReply::Handled();
+		};
+
+	auto OnLengthConformGross = [this, TheMidiFile]
+		{
+			int32 CurrentLengthTicks = TheMidiFile->GetSongMaps()->GetSongLengthData().LengthTicks;
+			EMidiClockSubdivisionQuantization Subdivision = EMidiClockSubdivisionQuantization::None;
+			int32 QuantizedLengthTick = TheMidiFile->GetSongMaps()->QuantizeTickToAnyNearestSubdivision(CurrentLengthTicks, EMidiFileQuantizeDirection::Nearest, Subdivision);
+			UMidiFileFactory::AskOrDoGrossConform(TheMidiFile, false, nullptr, CurrentLengthTicks, QuantizedLengthTick, Subdivision);
+			TheMidiFile->PostEditChange();
+			BuildLengthRow(TheMidiFile);
+			return FReply::Handled();
+		};
+
+	bool bLengthBoxIsNew = false;
+	if (!LengthBox.IsValid())
+	{
+		bLengthBoxIsNew = true;
+		LengthRow->ValueContent()
+			[
+				SAssignNew(LengthBox, SHorizontalBox)
+			];
+	}
+	LengthBox->ClearChildren();
+
+	if (bLengthIsAPerfectSubdivision)
+	{
+		LengthBox->AddSlot()
+			.HAlign(HAlign_Left)
+			.VAlign(VAlign_Center)
+			.AutoWidth()
+			[
+				SAssignNew(FileLengthTextBlock, STextBlock)
+					.Text(FileLengthText)
+					.Font(IDetailLayoutBuilder::GetDetailFont())
+					.ColorAndOpacity(FLinearColor::White)
+					.ToolTipText(RoundedLengthToolTipText)
+			];
+	}
+	else if (UMidiFileFactory::LengthCanBeTriviallyConformed(TheMidiFile))
+	{
+		LengthBox->AddSlot()
+			.HAlign(HAlign_Left)
+			.VAlign(VAlign_Center)
+			.AutoWidth()
+			[
+				SAssignNew(FileLengthTextBlock, STextBlock)
+					.Text(FileLengthText)
+					.Font(IDetailLayoutBuilder::GetDetailFont())
+					.ColorAndOpacity(FLinearColor::Yellow)
+					.ToolTipText(RoundedLengthToolTipText)
+			];
+		LengthBox->AddSlot()
+			.HAlign(HAlign_Left)
+			.VAlign(VAlign_Center)
+			.AutoWidth()
+			[
+				SNew(SHorizontalBox)
+					+ SHorizontalBox::Slot()
+					.HAlign(HAlign_Right)
+					.VAlign(VAlign_Center)
+					.AutoWidth()
+					.Padding(10.0f, 0.0f, 0.0f, 0.0f)
+					[
+						SNew(STextBlock)
+							.Font(IDetailLayoutBuilder::GetDetailFont())
+							.Text(LOCTEXT("ConformByPrompt", "MIDI file is only a tick or two away from a musical subdivision: "))
+					]
+					+ SHorizontalBox::Slot()
+					.HAlign(HAlign_Left)
+					.VAlign(VAlign_Center)
+					.AutoWidth()
+					[
+						SNew(SButton)
+							.VAlign(VAlign_Center)
+							.Text(LOCTEXT("ConformTrivial", "Conform To Nearest Subdivision"))
+							.OnClicked_Lambda(OnLengthConformTrivial)
+					]
+			];
+	}
+	else
+	{
+		LengthBox->AddSlot()
+			.HAlign(HAlign_Left)
+			.VAlign(VAlign_Center)
+			.AutoWidth()
+			[
+				SAssignNew(FileLengthTextBlock, STextBlock)
+					.Text(FileLengthText)
+					.Font(IDetailLayoutBuilder::GetDetailFont())
+					.ColorAndOpacity(bLengthIsAPerfectSubdivision ? FLinearColor::White : FLinearColor::Yellow)
+					.ToolTipText(RoundedLengthToolTipText)
+			];
+		LengthBox->AddSlot()
+			.HAlign(HAlign_Right)
+			.VAlign(VAlign_Center)
+			.AutoWidth()
+			[
+				SNew(SHorizontalBox)
+					+ SHorizontalBox::Slot()
+					.HAlign(HAlign_Left)
+					.VAlign(VAlign_Center)
+					.Padding(30.0f, 0.0f, 0.0f, 0.0f)
+					.AutoWidth()
+					[
+						SNew(SButton)
+							.VAlign(VAlign_Center)
+							.Text(LOCTEXT("ConformByWorkflow", "Conform..."))
+							.OnClicked_Lambda(OnLengthConformGross)
+					]
+			];
+	}
+	if (!bLengthBoxIsNew)
+	{
+		LengthBox->Invalidate(EInvalidateWidgetReason::Layout);
+	}
+}
+
 #undef LOCTEXT_NAMESPACE
diff --git a/Engine/Plugins/Runtime/Harmonix/Source/HarmonixMidiEditor/Private/MidiFileDetailCustomization.h b/Engine/Plugins/Runtime/Harmonix/Source/HarmonixMidiEditor/Private/MidiFileDetailCustomization.h
--- a/Engine/Plugins/Runtime/Harmonix/Source/HarmonixMidiEditor/Private/MidiFileDetailCustomization.h
+++ b/Engine/Plugins/Runtime/Harmonix/Source/HarmonixMidiEditor/Private/MidiFileDetailCustomization.h
@@ -18,6 +18,8 @@
 Customize details panel for UMidiFile
 */
 
+class UMidiFile;
+
 class FMidiFileDetailCustomization : public IDetailCustomization
 {
 public:
@@ -27,11 +29,18 @@ public:
 	}
 
 	//Text Display for midi file length
+	virtual void CustomizeDetails(const TSharedPtr<IDetailLayoutBuilder>& InDetailBuilder) override;
+	virtual void CustomizeDetails(IDetailLayoutBuilder& InDetailBuilder) override;
+
+private:
+	TSharedPtr<IDetailLayoutBuilder> DetailBuilder;
 	FText FileLengthText;
 	FText RoundedLengthToolTipText;
 	TSharedPtr<STextBlock> FileLengthTextBlock;
-	virtual void CustomizeDetails(IDetailLayoutBuilder& DetailBuilder) override;
+	TSharedPtr<SHorizontalBox> LengthBox;
+	FDetailWidgetRow* LengthRow = nullptr;
 
-private:
 	TWeakObjectPtr<UObject> MidiFile;
+
+	void BuildLengthRow(UMidiFile* TheMidiFile);
 };
\ No newline at end of file
diff --git a/Engine/Plugins/Runtime/Harmonix/Source/HarmonixMidiEditor/Private/MidiFileFactory.cpp b/Engine/Plugins/Runtime/Harmonix/Source/HarmonixMidiEditor/Private/MidiFileFactory.cpp
--- a/Engine/Plugins/Runtime/Harmonix/Source/HarmonixMidiEditor/Private/MidiFileFactory.cpp
+++ b/Engine/Plugins/Runtime/Harmonix/Source/HarmonixMidiEditor/Private/MidiFileFactory.cpp
@@ -6,6 +6,10 @@
 #include "EditorFramework/AssetImportData.h"
 #include "Misc/MessageDialog.h"
 #include "Misc/Paths.h"
+#include "Widgets/Input/STextComboBox.h"
+#include "IPropertyTypeCustomization.h"
+
+#define LOCTEXT_NAMESPACE "HarmonixMIDIEditor"
 
 UMidiFileFactory::UMidiFileFactory(const FObjectInitializer& ObjectInitializer)
 	: Super(ObjectInitializer)
@@ -21,12 +25,12 @@ UMidiFileFactory::UMidiFileFactory(const FObjectInitializer& ObjectInitializer)
 
 FText UMidiFileFactory::GetDisplayName() const
 {
-	return NSLOCTEXT("MIDI", "MIDIImporterFactoryName", "Standard MIDI File");
+	return LOCTEXT("ImporterFactory_Name", "Standard MIDI File");
 }
 
 FText UMidiFileFactory::GetToolTip() const
 {
-	return NSLOCTEXT("MIDI", "MIDIImporterFactoryDescription", "Standard MIDI Files exported from Digital Audio Workstations");
+	return LOCTEXT("ImporterFactory_Description", "Standard MIDI Files exported from Digital Audio Workstations");
 }
 
 bool UMidiFileFactory::FactoryCanImport(const FString& Filename)
@@ -49,6 +53,12 @@ UObject* UMidiFileFactory::FactoryCreateFile(UClass* InClass, UObject* InParent,
 	return MidiFileAsset;
 }
 
+void UMidiFileFactory::CleanUp()
+{
+	CheckImportedFilesForValidLengths();
+	Super::CleanUp();
+}
+
 bool UMidiFileFactory::CanReimport(UObject* Obj, TArray<FString>& OutFilenames)
 {
 	if (UMidiFile* AsMidiFile = Cast<UMidiFile>(Obj))
@@ -73,268 +83,418 @@ void UMidiFileFactory::SetReimportPaths(UObject* Obj, const TArray<FString>& New
 	}
 }
 
-void UMidiFileFactory::CleanUp()
-{ 
+EReimportResult::Type UMidiFileFactory::Reimport(UObject* Obj)
+{
+	UMidiFile* AsMidiFile = Cast<UMidiFile>(Obj);
+	if (!AsMidiFile)
+		return EReimportResult::Failed;
+
+	FString ReimportPath = AsMidiFile->GetImportedSrcFilePath();
+	if (ReimportPath.IsEmpty())
+	{
+		UE_LOG(LogHarmonixMidiEditor, Warning, TEXT("Failed to reimport MIDI file: %s"), *AsMidiFile->GetFullName());
+		return EReimportResult::Failed;
+	}
+
+	AsMidiFile->LoadStdMidiFile(ReimportPath, Harmonix::Midi::Constants::GTicksPerQuarterNoteInt);
+	ImportedFiles.AddUnique(AsMidiFile);
+
+	return EReimportResult::Succeeded;
+}
+
+void UMidiFileFactory::PostImportCleanUp()
+{
+	CheckImportedFilesForValidLengths();
+}
+
+void UMidiFileFactory::CheckImportedFilesForValidLengths()
+{
 	for (int FileIndex = 0; FileIndex < ImportedFiles.Num(); FileIndex++)
 	{
 		UMidiFile* MidiFileAsset = ImportedFiles[FileIndex];
-		//only prompt this dialog if imported file(s) can be conformed by at least one conform option (round up or round down)
-		if (MidiFileAsset->ShouldConformMidiFileLength(EMidiFileLengthConformOption::RoundDown) || MidiFileAsset->ShouldConformMidiFileLength(EMidiFileLengthConformOption::RoundUp))
+
+		if (!MidiFileAsset->GetSongMaps()->LengthIsAPerfectSubdivision())
 		{
-			if (bShouldApplyToAll && FileIndex > 0)
+			// trivially different, or grossly different?
+			int32 CurrentLengthTicks = MidiFileAsset->GetSongMaps()->GetSongLengthData().LengthTicks;
+			EMidiClockSubdivisionQuantization Subdivision = EMidiClockSubdivisionQuantization::None;
+			int32 QuantizedLengthTick = MidiFileAsset->GetSongMaps()->QuantizeTickToAnyNearestSubdivision(CurrentLengthTicks, EMidiFileQuantizeDirection::Nearest, Subdivision);
+			if (LengthCanBeTriviallyConformed(MidiFileAsset))
+			{
+				AskOrDoTrivialConform(MidiFileAsset, FileIndex < (ImportedFiles.Num() - 1), this, CurrentLengthTicks, QuantizedLengthTick);
+			}
+			else
 			{
-				MidiFileAsset->ConformMidiFileLength(ApplyToAllOption);
-				continue;
+				AskOrDoGrossConform(MidiFileAsset, FileIndex < (ImportedFiles.Num() - 1), this, CurrentLengthTicks, QuantizedLengthTick, Subdivision);
 			}
-			ShowConformMidiFileLengthDialog(FileIndex);
 		}
-
 	}
 	ImportedFiles.Empty();
-	bShouldApplyToAll = true;
-	ApplyToAllOption = EMidiFileLengthConformOption::RoundUp;
+	bApplyGrossConformToAll = false;
+	bDontApplyGrossConformToAll = false;
+	bApplyOffByOneConformToAll = false;
+	bDontApplyOffByOneConformToAll = false;
+	ApplyGrossConformDirection = EMidiFileQuantizeDirection::Up;
+	GrossConformSubdivision = EMidiClockSubdivisionQuantization::None;
 }
 
-void UMidiFileFactory::ShowConformMidiFileLengthDialog(int32 MidiFileAssetIndex)
+bool UMidiFileFactory::LengthCanBeTriviallyConformed(UMidiFile* MidiFile)
 {
-	UMidiFile* MidiFileAsset = ImportedFiles[MidiFileAssetIndex];
-	
-	//show the dialog widgets on a window
-	TSharedPtr<SWindow> ParentWindow;
-	SAssignNew(ParentWindow, SWindow)
-		.Title(FText::FromString(TEXT("Conform MIDI File Length")))
-		.ClientSize(FVector2D(490, 250));
-	
-	//the custom pop up dialog that asks for conforming midi file length
-	TSharedRef<SConformMidiFileLengthDialog> ConformMidiFileLengthDialog = SNew(SConformMidiFileLengthDialog);
-	ParentWindow->SetContent(ConformMidiFileLengthDialog);
-
-	// enable the Apply To All checkbox if importing more than 1 file 
-	(ImportedFiles.Num() > 1 && MidiFileAssetIndex != ImportedFiles.Num() - 1) ? ConformMidiFileLengthDialog->ApplyToAllCheckBox->SetVisibility(EVisibility::Visible) : ConformMidiFileLengthDialog->ApplyToAllCheckBox->SetVisibility(EVisibility::Collapsed);
+	int32 CurrentLengthTicks = MidiFile->GetSongMaps()->GetSongLengthData().LengthTicks;
+	EMidiClockSubdivisionQuantization Division = EMidiClockSubdivisionQuantization::None;
+	int32 QuantizedLengthTick = MidiFile->GetSongMaps()->QuantizeTickToAnyNearestSubdivision(CurrentLengthTicks, EMidiFileQuantizeDirection::Nearest, Division);
 
-	// we'll need this for a few options down below...	
-	float MidiFileFractionalLength = MidiFileAsset->GetSongMaps()->GetBarIncludingCountInAtTick(MidiFileAsset->GetLastEventTick() + 1);
+	return (FMath::Abs(CurrentLengthTicks - QuantizedLengthTick) <= kMaxTickErrorForTrivialConform);
+}
 
-	//if the file has length less than 1 bar (e.g. 0.25 bars),it cannot be rounded down to 0 bar, this is checked in UMidiFile::ShouldConformFileLength(EMidiFileLengthConformOption Option)
-	//disable the Round Down checkbox and the Round To Nearest checkbox, the only available option should be Round Up
-	if (!MidiFileAsset->ShouldConformMidiFileLength(EMidiFileLengthConformOption::RoundDown) && MidiFileAsset->ShouldConformMidiFileLength(EMidiFileLengthConformOption::RoundUp))
+void UMidiFileFactory::AskOrDoTrivialConform(UMidiFile* MidiFileAsset, bool bIsOneOfMany, UMidiFileFactory* CallingFactory, int32 CurrentLengthTicks, int32 QuantizedLengthTicks)
+{
+	if (CallingFactory && CallingFactory->bDontApplyOffByOneConformToAll)
 	{
-		ConformMidiFileLengthDialog->RoundDownCheckBox->SetEnabled(false);
-		ConformMidiFileLengthDialog->RoundDownCheckBox->SetIsChecked(ECheckBoxState::Unchecked);
-
-		ConformMidiFileLengthDialog->RoundToNearestCheckBox->SetEnabled(false);
-		ConformMidiFileLengthDialog->RoundToNearestCheckBox->SetIsChecked(ECheckBoxState::Unchecked);
-
-		ConformMidiFileLengthDialog->RoundUpCheckBox->SetIsChecked(ECheckBoxState::Checked);
-		ConformMidiFileLengthDialog->ConformOption = EMidiFileLengthConformOption::RoundUp;
+		// we've been asked to skip all minor conforms...
+		return;
 	}
-	else
+
+	EAppReturnType::Type Result = EAppReturnType::Type::No;
+	if (!CallingFactory ||  !CallingFactory->bApplyOffByOneConformToAll)
 	{
-		ConformMidiFileLengthDialog->RoundDownCheckBox->SetIsChecked(ECheckBoxState::Unchecked);
-		ConformMidiFileLengthDialog->RoundToNearestCheckBox->SetIsChecked(ECheckBoxState::Unchecked);
-		ConformMidiFileLengthDialog->RoundUpCheckBox->SetIsChecked(ECheckBoxState::Unchecked);
-		if (FMath::Frac(MidiFileFractionalLength) < 0.1f || FMath::Frac(MidiFileFractionalLength) > 0.9f)
+		Result = FMessageDialog::Open(
+			EAppMsgCategory::Warning,
+			bIsOneOfMany ? EAppMsgType::YesNoYesAllNoAll : EAppMsgType::YesNo,
+			EAppReturnType::Type::Yes,
+			FText::Format(LOCTEXT("ImporterFactory_OffByOneWarningMsg", "The MIDI File \"{0}\" has a length that is {1} tick(s) off\n"
+				"from of being evenly divisible by any musical subdivision.\n\n"
+				"It is HIGHLY recommended that you conform this length to the nearest musical\n"
+				"subdivision so that looping, sequential playback, etc. stays in sync.\n\n"
+				"Can I conform this length?"), FText::FromString(MidiFileAsset->GetName()), CurrentLengthTicks - QuantizedLengthTicks),
+			LOCTEXT("ImporterFactory_OffByOneWarningTitle", "MIDI Import Warning: Minor Length Issue..."));
+		if (CallingFactory && Result == EAppReturnType::Type::YesAll)
 		{
-			ConformMidiFileLengthDialog->RoundToNearestCheckBox->SetIsChecked(ECheckBoxState::Checked);
-			ConformMidiFileLengthDialog->ConformOption = EMidiFileLengthConformOption::Nearest;
+			CallingFactory->bApplyOffByOneConformToAll = true;
 		}
-		else if (FMath::Frac(MidiFileFractionalLength) < 0.2f)
+		else if (CallingFactory && Result == EAppReturnType::Type::NoAll)
 		{
-			ConformMidiFileLengthDialog->RoundDownCheckBox->SetIsChecked(ECheckBoxState::Checked);
-			ConformMidiFileLengthDialog->ConformOption = EMidiFileLengthConformOption::RoundDown;
-		}
-		else
-		{
-			ConformMidiFileLengthDialog->RoundUpCheckBox->SetIsChecked(ECheckBoxState::Checked);
-			ConformMidiFileLengthDialog->ConformOption = EMidiFileLengthConformOption::RoundUp;
+			CallingFactory->bDontApplyOffByOneConformToAll = true;
 		}
 	}
 
-	//set text for text block that asks user whether or not to conform midi file length
-	FString MidiFileName = MidiFileAsset->GetName();
-	ConformMidiFileLengthDialog->AskConformFileLengthText->SetText(FText::FromString(FString::Printf(
-		TEXT("MIDI File Name: %s.mid, Length: %.5f bars"),
-		*MidiFileName,
-		MidiFileFractionalLength)));
-
-	//Set callback function for Ok button
-	ConformMidiFileLengthDialog->OkButton->SetOnClicked(FOnClicked::CreateLambda([this, ParentWindow, ConformMidiFileLengthDialog, &MidiFileAsset]()-> FReply {
-		MidiFileAsset->ConformMidiFileLength(ConformMidiFileLengthDialog->ConformOption);
-		ParentWindow->RequestDestroyWindow();
-		return FReply::Handled();
-		}));
-
-	//show the dialog
-	FSlateApplication::Get().AddModalWindow(ParentWindow.ToSharedRef(), nullptr);
-		
-	//remember the current conform option and mark that this option should be apply to all other files imported
-	if (ConformMidiFileLengthDialog->ApplyToAll)
+	if (Result == EAppReturnType::Type::Yes || (CallingFactory && CallingFactory->bApplyOffByOneConformToAll))
 	{
-		ApplyToAllOption = ConformMidiFileLengthDialog->ConformOption;
-		bShouldApplyToAll = true;
-	}
-	else {
-		bShouldApplyToAll = false;
+		MidiFileAsset->QuantizeLengthToNearestPerfectSubdivision(EMidiFileQuantizeDirection::Nearest);
 	}
 }
 
-EReimportResult::Type UMidiFileFactory::Reimport(UObject* Obj)
+void UMidiFileFactory::AskOrDoGrossConform(UMidiFile* MidiFileAsset, bool bIsOneOfMany, UMidiFileFactory* CallingFactory, int32 CurrentLengthTicks, int32 QuantizedLengthTick, EMidiClockSubdivisionQuantization BestSubdivision)
 {
-	UMidiFile* AsMidiFile = Cast<UMidiFile>(Obj);
-	if (!AsMidiFile)
-		return EReimportResult::Failed;
-
-	FString ReimportPath = AsMidiFile->GetImportedSrcFilePath();
-	if (ReimportPath.IsEmpty())
+	if (CallingFactory && CallingFactory->bDontApplyGrossConformToAll)
 	{
-		UE_LOG(LogHarmonixMidiEditor, Warning, TEXT("Failed to reimport MIDI file: %s"), *AsMidiFile->GetFullName());
-		return EReimportResult::Failed;
+		// we've been told to skip all gross conforms...
+		return;
 	}
-	
-	AsMidiFile->LoadStdMidiFile(ReimportPath, Harmonix::Midi::Constants::GTicksPerQuarterNoteInt);
 
-	//Allow file length to be reconformed after reimport 
-	AsMidiFile->bLengthRoundedDown = false;
-	AsMidiFile->bLengthRoundedUp = false;
-	AsMidiFile->bLengthRoundedToNearest = false;
+	EMidiFileQuantizeDirection Direction = CallingFactory ? CallingFactory->ApplyGrossConformDirection : EMidiFileQuantizeDirection::Nearest;
+	EMidiClockSubdivisionQuantization Subdivision = CallingFactory ? CallingFactory->GrossConformSubdivision : EMidiClockSubdivisionQuantization::None;
 
-	return EReimportResult::Succeeded;
-	
+	if (!CallingFactory || !CallingFactory->bApplyGrossConformToAll)
+	{
+		// let's figure out the best direction to go an the best unit...
+		EMidiFileQuantizeDirection RecommendedDirection = QuantizedLengthTick > CurrentLengthTicks ? EMidiFileQuantizeDirection::Up : EMidiFileQuantizeDirection::Down;
+
+		// first we need a window to hold the custom dialog...
+		TSharedPtr<SWindow> ParentWindow;
+		SAssignNew(ParentWindow, SWindow)
+			.Title(LOCTEXT("ImporterFactory_GrossConformTitle", "MIDI Import Warning: Significant Length Issue..."))
+			.ClientSize(FVector2D(600, 250));
+
+		// now a detailed description of the file in question...
+		FString MidiFileName = MidiFileAsset->GetName();
+		FText Description = FText::Format(LOCTEXT("ImporterFactory_FileSpecs", "MIDI File Name: {0}.mid\nLength: {1}"),
+			FText::FromString(MidiFileName),
+			FText::FromString(MidiFileAsset->GetSongMaps()->GetSongLengthString()));
+
+		// now we can create the custom dialog widget inside the window we made...
+		TSharedRef<SConformMidiFileLengthDialog> ConformMidiFileLengthDialog = SNew(SConformMidiFileLengthDialog)
+			.ParentWindow(ParentWindow)
+			.bMultipleFiles(bIsOneOfMany)
+			.FileDescription(Description)
+			.RecommendedDirection(RecommendedDirection)
+			.RecommendedSubdivision(BestSubdivision);
+
+		//show the dialog...
+		FSlateApplication::Get().AddModalWindow(ParentWindow.ToSharedRef(), nullptr);
+
+		// now we have the results...
+		if (!ConformMidiFileLengthDialog->bUserTookAction || ConformMidiFileLengthDialog->ConformSubdivision == EMidiClockSubdivisionQuantization::None)
+		{
+			if (CallingFactory && ConformMidiFileLengthDialog->bUserTookAction)
+			{
+				CallingFactory->bDontApplyGrossConformToAll = ConformMidiFileLengthDialog->bDoForAll;
+			}
+			// return without doing anything...
+			return;
+		}
+
+		Direction = ConformMidiFileLengthDialog->ConformDirection;
+		Subdivision = ConformMidiFileLengthDialog->ConformSubdivision;
+
+		if (CallingFactory)
+		{
+			CallingFactory->bApplyGrossConformToAll = ConformMidiFileLengthDialog->bDoForAll;
+			CallingFactory->ApplyGrossConformDirection = ConformMidiFileLengthDialog->ConformDirection;
+			CallingFactory->GrossConformSubdivision = ConformMidiFileLengthDialog->ConformSubdivision;
+		}
+	}
+
+	if (Subdivision != EMidiClockSubdivisionQuantization::None)
+	{
+		MidiFileAsset->QuantizeLengthToSubdivision(Direction, Subdivision);
+		MidiFileAsset->PostEditChange();
+	}
 }
 
 void SConformMidiFileLengthDialog::Construct(const FArguments& InArgs)
 {
+	check(InArgs._ParentWindow);
+
+	DirectionNames.Add(MakeShared<FString>(StaticEnum<EMidiFileQuantizeDirection>()->GetDisplayNameTextByValue(0).ToString()));
+	DirectionNames.Add(MakeShared<FString>(StaticEnum<EMidiFileQuantizeDirection>()->GetDisplayNameTextByValue(1).ToString()));
+	DirectionNames.Add(MakeShared<FString>(StaticEnum<EMidiFileQuantizeDirection>()->GetDisplayNameTextByValue(2).ToString()));
+	SelectedDirection = DirectionNames[(int32)InArgs._RecommendedDirection];
+
+	QuantizationSubdivisions.Empty();
+	QuantizationSubdivisions.Add(MakeShared<FSubdivisionEntry>(EMidiClockSubdivisionQuantization::Bar, StaticEnum<EMidiClockSubdivisionQuantization>()->GetDisplayNameTextByValue((int64)EMidiClockSubdivisionQuantization::Bar).ToString()));
+	QuantizationSubdivisions.Add(MakeShared<FSubdivisionEntry>(EMidiClockSubdivisionQuantization::Beat, StaticEnum<EMidiClockSubdivisionQuantization>()->GetDisplayNameTextByValue((int64)EMidiClockSubdivisionQuantization::Beat).ToString()));
+	QuantizationSubdivisions.Add(MakeShared<FSubdivisionEntry>(EMidiClockSubdivisionQuantization::WholeNote, StaticEnum<EMidiClockSubdivisionQuantization>()->GetDisplayNameTextByValue((int64)EMidiClockSubdivisionQuantization::WholeNote).ToString()));
+	QuantizationSubdivisions.Add(MakeShared<FSubdivisionEntry>(EMidiClockSubdivisionQuantization::HalfNote, StaticEnum<EMidiClockSubdivisionQuantization>()->GetDisplayNameTextByValue((int64)EMidiClockSubdivisionQuantization::HalfNote).ToString()));
+	QuantizationSubdivisions.Add(MakeShared<FSubdivisionEntry>(EMidiClockSubdivisionQuantization::QuarterNote, StaticEnum<EMidiClockSubdivisionQuantization>()->GetDisplayNameTextByValue((int64)EMidiClockSubdivisionQuantization::QuarterNote).ToString()));
+	QuantizationSubdivisions.Add(MakeShared<FSubdivisionEntry>(EMidiClockSubdivisionQuantization::EighthNote, StaticEnum<EMidiClockSubdivisionQuantization>()->GetDisplayNameTextByValue((int64)EMidiClockSubdivisionQuantization::EighthNote).ToString()));
+	QuantizationSubdivisions.Add(MakeShared<FSubdivisionEntry>(EMidiClockSubdivisionQuantization::SixteenthNote, StaticEnum<EMidiClockSubdivisionQuantization>()->GetDisplayNameTextByValue((int64)EMidiClockSubdivisionQuantization::SixteenthNote).ToString()));
+	QuantizationSubdivisions.Add(MakeShared<FSubdivisionEntry>(EMidiClockSubdivisionQuantization::ThirtySecondNote, StaticEnum<EMidiClockSubdivisionQuantization>()->GetDisplayNameTextByValue((int64)EMidiClockSubdivisionQuantization::ThirtySecondNote).ToString()));
+	QuantizationSubdivisions.Add(MakeShared<FSubdivisionEntry>(EMidiClockSubdivisionQuantization::DottedWholeNote, StaticEnum<EMidiClockSubdivisionQuantization>()->GetDisplayNameTextByValue((int64)EMidiClockSubdivisionQuantization::DottedWholeNote).ToString()));
+	QuantizationSubdivisions.Add(MakeShared<FSubdivisionEntry>(EMidiClockSubdivisionQuantization::DottedHalfNote, StaticEnum<EMidiClockSubdivisionQuantization>()->GetDisplayNameTextByValue((int64)EMidiClockSubdivisionQuantization::DottedHalfNote).ToString()));
+	QuantizationSubdivisions.Add(MakeShared<FSubdivisionEntry>(EMidiClockSubdivisionQuantization::DottedQuarterNote, StaticEnum<EMidiClockSubdivisionQuantization>()->GetDisplayNameTextByValue((int64)EMidiClockSubdivisionQuantization::DottedQuarterNote).ToString()));
+	QuantizationSubdivisions.Add(MakeShared<FSubdivisionEntry>(EMidiClockSubdivisionQuantization::DottedEighthNote, StaticEnum<EMidiClockSubdivisionQuantization>()->GetDisplayNameTextByValue((int64)EMidiClockSubdivisionQuantization::DottedEighthNote).ToString()));
+	QuantizationSubdivisions.Add(MakeShared<FSubdivisionEntry>(EMidiClockSubdivisionQuantization::DottedSixteenthNote, StaticEnum<EMidiClockSubdivisionQuantization>()->GetDisplayNameTextByValue((int64)EMidiClockSubdivisionQuantization::DottedSixteenthNote).ToString()));
+	QuantizationSubdivisions.Add(MakeShared<FSubdivisionEntry>(EMidiClockSubdivisionQuantization::HalfNoteTriplet, StaticEnum<EMidiClockSubdivisionQuantization>()->GetDisplayNameTextByValue((int64)EMidiClockSubdivisionQuantization::HalfNoteTriplet).ToString()));
+	QuantizationSubdivisions.Add(MakeShared<FSubdivisionEntry>(EMidiClockSubdivisionQuantization::QuarterNoteTriplet, StaticEnum<EMidiClockSubdivisionQuantization>()->GetDisplayNameTextByValue((int64)EMidiClockSubdivisionQuantization::QuarterNoteTriplet).ToString()));
+	QuantizationSubdivisions.Add(MakeShared<FSubdivisionEntry>(EMidiClockSubdivisionQuantization::EighthNoteTriplet, StaticEnum<EMidiClockSubdivisionQuantization>()->GetDisplayNameTextByValue((int64)EMidiClockSubdivisionQuantization::EighthNoteTriplet).ToString()));
+	QuantizationSubdivisions.Add(MakeShared<FSubdivisionEntry>(EMidiClockSubdivisionQuantization::SixteenthNoteTriplet, StaticEnum<EMidiClockSubdivisionQuantization>()->GetDisplayNameTextByValue((int64)EMidiClockSubdivisionQuantization::SixteenthNoteTriplet).ToString()));
+
+	if (TSharedPtr<FSubdivisionEntry>* Match = QuantizationSubdivisions.FindByPredicate([Recommended = InArgs._RecommendedSubdivision](const TSharedPtr<FSubdivisionEntry>& Item) { return Item->Value == Recommended; }))
+	{
+		SelectedQuantizationSubdivision = *Match;
+
+	}
+	else
+	{
+		SelectedQuantizationSubdivision = QuantizationSubdivisions[0];
+	}
+
+	TSharedPtr<SHorizontalBox> ActionButtonPanel;
+
 	ChildSlot
 	[
 		SNew(SBorder)
-		.HAlign(HAlign_Center)
-		.VAlign(VAlign_Center)
-		.BorderImage(FCoreStyle::Get().GetBrush("ToolPanel.GroupBorder"))
+		.HAlign(HAlign_Center).VAlign(VAlign_Center).BorderImage(FCoreStyle::Get().GetBrush("ToolPanel.GroupBorder"))
 		[
 			SNew(SVerticalBox)
-			// Text for the conform option
-			+ SVerticalBox::Slot()
-			.AutoHeight()
-			.HAlign(HAlign_Center)
-			.VAlign(VAlign_Center)
-			.Padding(10)
-			[
-				SAssignNew(AskConformFileLengthText,STextBlock)
-				.TextStyle(&MidiFileInformationStyle)
-			]
-			+ SVerticalBox::Slot()
-			.AutoHeight()
-			.HAlign(HAlign_Center)
-			.VAlign(VAlign_Center)
-			.Padding(5)
-			[
-				SNew(STextBlock)
-				.Text(FText::FromString(TEXT(
-					"Currently, all MIDI asset lengths must be conformed to some full number of bars."
-					"\n\n      Some or all of the MIDI files you are importing have fractional bar counts."
-					"\n\n                                        How would you like them conformed?\n")))
-			]
 			+ SVerticalBox::Slot()
-			.AutoHeight()
-			.Padding(5)
-			.VAlign(VAlign_Center)
-			.HAlign(HAlign_Center)
-			[
-				// Round Up checkbox
-				SNew(SHorizontalBox)
-				+ SHorizontalBox::Slot()
-				.Padding(5, 5)
-				.HAlign(HAlign_Center)
-				.VAlign(VAlign_Center)
-				.AutoWidth()
+				.AutoHeight().HAlign(HAlign_Left).VAlign(VAlign_Center).Padding(20.0f, 10.0f)
 				[
-					SAssignNew(RoundUpCheckBox, SCheckBox)
-					.ToolTipText(FText::FromString(TEXT("this does NOT change/add events in file because file length is automatically rounded up to the next integer bar after import")))
-					.OnCheckStateChanged(this, &SConformMidiFileLengthDialog::HandleConformOptionCheckboxChanged, EMidiFileLengthConformOption::RoundUp)
-					.IsChecked_Lambda([this]() { return ConformOption == EMidiFileLengthConformOption::RoundUp ? ECheckBoxState::Checked : ECheckBoxState::Unchecked; })
-					.Content()
-					[
-						SNew(STextBlock)
-						.Text(FText::FromString(TEXT("Round Up")))
-					]
+					SNew(STextBlock)
+					.TextStyle(&MidiFileInformationStyle)
+					.Text(InArgs._FileDescription)
 				]
-				// Round Down checkbox
-				+ SHorizontalBox::Slot()
-				.HAlign(HAlign_Center)
-				.VAlign(VAlign_Center)
-				.AutoWidth()
-				.Padding(5, 5)
+			+ SVerticalBox::Slot()
+				.AutoHeight().HAlign(HAlign_Center).VAlign(VAlign_Center).Padding(5)
 				[
-					SAssignNew(RoundDownCheckBox, SCheckBox)
-					.OnCheckStateChanged(this, &SConformMidiFileLengthDialog::HandleConformOptionCheckboxChanged, EMidiFileLengthConformOption::RoundDown)
-					.IsChecked_Lambda([this]() { return ConformOption == EMidiFileLengthConformOption::RoundDown ? ECheckBoxState::Checked : ECheckBoxState::Unchecked; })
-					.Content()
-					[
-						SNew(STextBlock)
-						.ToolTipText(FText::FromString(TEXT("move MIDI events in file that exceed the last integer bar to the last tick of that bar and remove excessive events on that tick")))
-						.Text(FText::FromString(TEXT("Round Down")))
-					]
+					SNew(STextBlock)
+					.Text(LOCTEXT("ImporterFactory_GrossConformExplanation", "Some or all of the MIDI files you are importing have a length\n"
+															   "That is not evenly divisible by any standard musical subdivision.\n\n"
+															   "It is highly recommended that you conform MIDI file lengths to a\n"
+															   "musical subdivision so that synchronization doesn't drift when looping\n"
+															   "or sequencing (stiching) midi files end-to-end."))
 				]
-				//Round To Nearest checkbox
-				+ SHorizontalBox::Slot()
-				.HAlign(HAlign_Center)
-				.VAlign(VAlign_Center)
-				.AutoWidth()
-				.Padding(5, 5)
+			+ SVerticalBox::Slot()
+				.AutoHeight().Padding(5).VAlign(VAlign_Center).HAlign(HAlign_Center)
 				[
-					SAssignNew(RoundToNearestCheckBox, SCheckBox)
-					.OnCheckStateChanged(this, &SConformMidiFileLengthDialog::HandleConformOptionCheckboxChanged, EMidiFileLengthConformOption::Nearest)
-					.IsChecked_Lambda([this]() { return ConformOption == EMidiFileLengthConformOption::Nearest ? ECheckBoxState::Checked : ECheckBoxState::Unchecked; })
-					.Content()
-					[
-						SNew(STextBlock)
-						.ToolTipText(FText::FromString(TEXT("round MIDI File length to the nearest integer bar (either rounding up or rounding down)")))
-						.Text(FText::FromString(TEXT("Round to Nearest")))
-					]
+					// Direction
+					SNew(SHorizontalBox)
+					+ SHorizontalBox::Slot()
+						.Padding(5, 5).HAlign(HAlign_Center).VAlign(VAlign_Center).AutoWidth()
+						[
+							SNew(STextBlock)
+							.Text(LOCTEXT("ImporterFactory_Label_Round", "Round:"))
+						]
+					+ SHorizontalBox::Slot()
+						.Padding(5, 5).HAlign(HAlign_Center).VAlign(VAlign_Center).AutoWidth()
+						[
+							SNew(STextComboBox)
+ 							.OptionsSource(&DirectionNames)
+							.InitiallySelectedItem(SelectedDirection)
+							//.Font(IDetailLayoutBuilder::GetDetailFont())
+							.OnSelectionChanged_Lambda([this](TSharedPtr<FString> InDirection, ESelectInfo::Type InSelectType)
+								{
+									SelectedDirection = InDirection;
+								})
+						]
+					+ SHorizontalBox::Slot()
+						.Padding(5, 5).HAlign(HAlign_Center).VAlign(VAlign_Center).AutoWidth()
+						[
+							SNew(STextBlock)
+							.Text(LOCTEXT("ImporterFactory_Label_Subdivision", "Subdivision:"))
+						]
+					+ SHorizontalBox::Slot()
+						.Padding(5, 5).HAlign(HAlign_Center).VAlign(VAlign_Center).AutoWidth()
+						[
+							SNew(SComboBox<TSharedPtr<FSubdivisionEntry>>)
+							.OptionsSource(&QuantizationSubdivisions)
+							.OnGenerateWidget_Lambda([](TSharedPtr<FSubdivisionEntry> InQuantizationUnit)
+								{
+									return SNew(STextBlock)
+										.Text(FText::FromString(InQuantizationUnit->DisplayString))
+										.Font(IPropertyTypeCustomizationUtils::GetRegularFont());
+								})
+							.InitiallySelectedItem(SelectedQuantizationSubdivision)
+							.OnSelectionChanged_Lambda([this](TSharedPtr<FSubdivisionEntry> InQuantizationUnit, ESelectInfo::Type InSelectType)
+								{
+									SelectedQuantizationSubdivision = InQuantizationUnit;
+								})
+							.Content()
+							[
+								SNew(STextBlock)
+								.Font(IPropertyTypeCustomizationUtils::GetRegularFont())
+								.Text_Lambda([this]()
+									{
+										return FText::FromString(SelectedQuantizationSubdivision->DisplayString);
+									})
+							]
+						]
 				]
-			]
-			// "Apply to All" checkbox (visible only if ApplyToAll is true)
 			+ SVerticalBox::Slot()
-			.AutoHeight()
-			.VAlign(VAlign_Center)
-			.HAlign(HAlign_Center)
-			.Padding(5)
-			[
-				SNew(SBox)
+				.AutoHeight().Padding(5).VAlign(VAlign_Center).HAlign(HAlign_Center)
 				[
-					SAssignNew(ApplyToAllCheckBox, SCheckBox)
-					.OnCheckStateChanged(this, &SConformMidiFileLengthDialog::HandleApplyToAllCheckboxChanged)
-					.IsChecked_Lambda([this]() { return ApplyToAll ? ECheckBoxState::Checked : ECheckBoxState::Unchecked; })
-					.Content()
-					[
-						SNew(STextBlock)
-						.Text(FText::FromString(TEXT("Apply to All")))
-					]
+					SAssignNew(ActionButtonPanel, SHorizontalBox)
 				]
-			]
-			// "Ok" button
-			+ SVerticalBox::Slot()
-			.AutoHeight()
-			.Padding(15)
-			.VAlign(VAlign_Bottom)
-			.HAlign(HAlign_Center)
-			[
-				SAssignNew(OkButton, SButton)
-				.Text(FText::FromString(TEXT("Ok")))
-			]
 		]	
 	];
-}
 
-void SConformMidiFileLengthDialog::HandleConformOptionCheckboxChanged(ECheckBoxState NewState, EMidiFileLengthConformOption Option)
-{
-	if (NewState == ECheckBoxState::Checked)
+	if (!InArgs._bMultipleFiles)
 	{
-		ConformOption = Option;
+		ActionButtonPanel->AddSlot()
+			.Padding(5, 5).HAlign(HAlign_Center).VAlign(VAlign_Center).AutoWidth()
+			[
+				SNew(SButton)
+					.Text(LOCTEXT("ImporterFactory_Button_Conform", "Conform"))
+					.OnClicked_Lambda([this, ParentWindow = InArgs._ParentWindow]()->FReply
+						{
+							bDoForAll = false;
+							bUserTookAction = true;
+							ConformDirection = (EMidiFileQuantizeDirection)StaticEnum<EMidiFileQuantizeDirection>()->GetValueByNameString(*SelectedDirection);
+							if (SelectedQuantizationSubdivision)
+							{
+								ConformSubdivision = SelectedQuantizationSubdivision->Value;
+							}
+							else
+							{
+								ConformSubdivision = EMidiClockSubdivisionQuantization::None;
+							}
+							ParentWindow->RequestDestroyWindow();
+							return FReply::Handled();
+						})
+			];
+		ActionButtonPanel->AddSlot()
+			.Padding(5, 5).HAlign(HAlign_Center).VAlign(VAlign_Center).AutoWidth()
+			[
+				SNew(SButton)
+					.Text(LOCTEXT("ImporterFactory_Button_Dont_Conform", "Don't Conform"))
+					.OnClicked_Lambda([this, ParentWindow = InArgs._ParentWindow]()->FReply
+						{
+							bDoForAll = false;
+							bUserTookAction = true;
+							ConformSubdivision = EMidiClockSubdivisionQuantization::None;
+							ParentWindow->RequestDestroyWindow();
+							return FReply::Handled();
+						})
+			];
+	}
+	else
+	{
+		ActionButtonPanel->AddSlot()
+			.Padding(5, 5).HAlign(HAlign_Center).VAlign(VAlign_Center).AutoWidth()
+			[
+				SNew(SButton)
+					.Text(LOCTEXT("ImporterFactory_Button_Conform", "Conform"))
+					.OnClicked_Lambda([this, ParentWindow = InArgs._ParentWindow]()->FReply
+						{
+							bDoForAll = false;
+							bUserTookAction = true;
+							ConformDirection = (EMidiFileQuantizeDirection)StaticEnum<EMidiFileQuantizeDirection>()->GetValueByNameString(*SelectedDirection);
+							if (SelectedQuantizationSubdivision)
+							{
+								ConformSubdivision = SelectedQuantizationSubdivision->Value;
+							}
+							else
+							{
+								ConformSubdivision = EMidiClockSubdivisionQuantization::None;
+							}
+							ParentWindow->RequestDestroyWindow();
+							return FReply::Handled();
+						})
+			];
+		ActionButtonPanel->AddSlot()
+			.Padding(5, 5).HAlign(HAlign_Center).VAlign(VAlign_Center).AutoWidth()
+			[
+				SNew(SButton)
+					.Text(LOCTEXT("ImporterFactory_Button_Conform_All", "Conform All"))
+					.OnClicked_Lambda([this, ParentWindow = InArgs._ParentWindow]()->FReply
+						{
+							bDoForAll = true;
+							bUserTookAction = true;
+							ConformDirection = (EMidiFileQuantizeDirection)StaticEnum<EMidiFileQuantizeDirection>()->GetValueByNameString(*SelectedDirection);
+							if (SelectedQuantizationSubdivision)
+							{
+								ConformSubdivision = SelectedQuantizationSubdivision->Value;
+							}
+							else
+							{
+								ConformSubdivision = EMidiClockSubdivisionQuantization::None;
+							}
+							ParentWindow->RequestDestroyWindow();
+							return FReply::Handled();
+						})
+			];
+		ActionButtonPanel->AddSlot()
+			.Padding(5, 5).HAlign(HAlign_Center).VAlign(VAlign_Center).AutoWidth()
+			[
+				SNew(SButton)
+					.Text(LOCTEXT("ImporterFactory_Button_Dont_Conform", "Don't Conform"))
+					.OnClicked_Lambda([this, ParentWindow = InArgs._ParentWindow]()->FReply
+						{
+							bDoForAll = false;
+							bUserTookAction = true;
+							ConformSubdivision = EMidiClockSubdivisionQuantization::None;
+							ParentWindow->RequestDestroyWindow();
+							return FReply::Handled();
+						})
+			];
+		ActionButtonPanel->AddSlot()
+			.Padding(5, 5).HAlign(HAlign_Center).VAlign(VAlign_Center).AutoWidth()
+			[
+				SNew(SButton)
+					.Text(LOCTEXT("ImporterFactory_Button_Dont_Conform_Any", "Don't Conform Any"))
+					.OnClicked_Lambda([this, ParentWindow = InArgs._ParentWindow]()->FReply
+						{
+							bDoForAll = true;
+							bUserTookAction = true;
+							ConformSubdivision = EMidiClockSubdivisionQuantization::None;
+							ParentWindow->RequestDestroyWindow();
+							return FReply::Handled();
+						})
+			];
 	}
-}
 
-void SConformMidiFileLengthDialog::HandleApplyToAllCheckboxChanged(ECheckBoxState NewState)
-{
-	ApplyToAll = (NewState == ECheckBoxState::Checked);
+	InArgs._ParentWindow->SetContent(this->AsShared());
 }
+
+#undef LOCTEXT_NAMESPACE
diff --git a/Engine/Plugins/Runtime/Harmonix/Source/HarmonixMidiEditor/Private/MidiFileFactory.h b/Engine/Plugins/Runtime/Harmonix/Source/HarmonixMidiEditor/Private/MidiFileFactory.h
--- a/Engine/Plugins/Runtime/Harmonix/Source/HarmonixMidiEditor/Private/MidiFileFactory.h
+++ b/Engine/Plugins/Runtime/Harmonix/Source/HarmonixMidiEditor/Private/MidiFileFactory.h
@@ -10,6 +10,7 @@
 #include "Widgets/SWindow.h"
 #include "Widgets/Input/SButton.h"
 #include "Widgets/Input/SCheckBox.h"
+#include "Widgets/Input/SComboBox.h"
 #include "Styling/SlateTypes.h"
 #include "Widgets/SBoxPanel.h"
 #include "Layout/Visibility.h"
@@ -38,15 +39,27 @@ class UMidiFileFactory : public UFactory, public FReimportHandler
 	virtual bool CanReimport(UObject* Obj, TArray<FString>& OutFilenames) override;
 	virtual void SetReimportPaths(UObject* Obj, const TArray<FString>& NewReimportPaths) override;
 	virtual EReimportResult::Type Reimport(UObject* Obj) override;
+	virtual void PostImportCleanUp() override;
 	//~ END FReimportHandler interface --
 
-	//Show a custom dialog upon importing midi files, prompt the user the option to conform midi file length
-	void ShowConformMidiFileLengthDialog(int32 MidiFileAssetIndex);
-	
+	static constexpr int32 kMaxTickErrorForTrivialConform = 2;
+	static bool LengthCanBeTriviallyConformed(UMidiFile* MidiFile);
+
 	//keep track of files that are imported for the pop-up dialog (multi batch)
 	TArray<UMidiFile*> ImportedFiles;
-	bool bShouldApplyToAll = true;
-	EMidiFileLengthConformOption ApplyToAllOption;
+	bool bApplyGrossConformToAll = false;
+	bool bDontApplyGrossConformToAll = false;
+	bool bApplyOffByOneConformToAll = false;
+	bool bDontApplyOffByOneConformToAll = false;
+	EMidiFileQuantizeDirection ApplyGrossConformDirection = EMidiFileQuantizeDirection::Up;
+	EMidiClockSubdivisionQuantization GrossConformSubdivision = EMidiClockSubdivisionQuantization::None;
+
+	// making these static means we can call them from other places in the editor UI...
+	static void AskOrDoTrivialConform(UMidiFile* MidiFileAsset, bool bIsOneOfMany, UMidiFileFactory* CallingFactory, int32 CurrentLengthTicks, int32 QuantizedLengthTick);
+	static void AskOrDoGrossConform(UMidiFile* MidiFileAsset, bool bIsOneOfMany, UMidiFileFactory* CallingFactory, int32 CurrentLengthTicks, int32 QuantizedLengthTick, EMidiClockSubdivisionQuantization BestSubdivision);
+
+private:
+	void CheckImportedFilesForValidLengths();
 };
 
 /* 
@@ -59,30 +72,42 @@ class SConformMidiFileLengthDialog : public SCompoundWidget
 {
 public:
 	SLATE_BEGIN_ARGS(SConformMidiFileLengthDialog)
-		: _ConformOption(EMidiFileLengthConformOption::RoundUp)
-		, _ApplyToAll(true) 
+		: _ParentWindow()
+		, _bMultipleFiles(false)
+		, _FileDescription()
+		, _RecommendedDirection(EMidiFileQuantizeDirection::Nearest)
+		, _RecommendedSubdivision(EMidiClockSubdivisionQuantization::Bar)
 	{}
 	
-	SLATE_ARGUMENT(EMidiFileLengthConformOption,ConformOption)
-	SLATE_ARGUMENT(bool,ApplyToAll)
-	
+	SLATE_ARGUMENT(TSharedPtr<SWindow>, ParentWindow)
+	SLATE_ARGUMENT(bool, bMultipleFiles)
+	SLATE_ARGUMENT(FText, FileDescription)
+	SLATE_ARGUMENT(EMidiFileQuantizeDirection, RecommendedDirection)
+	SLATE_ARGUMENT(EMidiClockSubdivisionQuantization, RecommendedSubdivision)
 	SLATE_END_ARGS()
 
 	void Construct(const FArguments& InArgs);
 
 	//custom widgets 
 	const FTextBlockStyle MidiFileInformationStyle = FTextBlockStyle().SetFont(FCoreStyle::GetDefaultFontStyle("Bold", 10)).SetColorAndOpacity(FLinearColor::White);
-	EMidiFileLengthConformOption ConformOption = EMidiFileLengthConformOption::RoundUp;
-	bool ApplyToAll = true;
-	TSharedPtr<STextBlock> AskConformFileLengthText;
-	TSharedPtr<SCheckBox> RoundToNearestCheckBox;
-	TSharedPtr<SCheckBox> RoundDownCheckBox;
-	TSharedPtr<SCheckBox> RoundUpCheckBox;
-	TSharedPtr<SCheckBox> ApplyToAllCheckBox;
-	TSharedPtr<SButton> OkButton;
+	EMidiFileQuantizeDirection ConformDirection = EMidiFileQuantizeDirection::Up;
+	EMidiClockSubdivisionQuantization ConformSubdivision = EMidiClockSubdivisionQuantization::Bar;
+	bool bDoForAll = false;
+	bool bUserTookAction = false;
 
 private:
-	void HandleConformOptionCheckboxChanged(ECheckBoxState NewState, EMidiFileLengthConformOption Option);
-	void HandleApplyToAllCheckboxChanged(ECheckBoxState NewState);
+	TArray<TSharedPtr<FString>> DirectionNames;
+	TSharedPtr<FString> SelectedDirection;
 
+	struct FSubdivisionEntry
+	{
+		EMidiClockSubdivisionQuantization Value;
+		FString DisplayString;
+		FSubdivisionEntry(EMidiClockSubdivisionQuantization InValue, const FString& InString)
+			: Value(InValue)
+			, DisplayString(InString)
+		{}
+	};
+	TArray<TSharedPtr<FSubdivisionEntry>> QuantizationSubdivisions;
+	TSharedPtr<FSubdivisionEntry> SelectedQuantizationSubdivision;
 };
\ No newline at end of file
diff --git a/Engine/Plugins/Runtime/Harmonix/Source/HarmonixMidiTests/Private/ConformMidiFileLengthTests.cpp b/Engine/Plugins/Runtime/Harmonix/Source/HarmonixMidiTests/Private/ConformMidiFileLengthTests.cpp
--- a/Engine/Plugins/Runtime/Harmonix/Source/HarmonixMidiTests/Private/ConformMidiFileLengthTests.cpp
+++ b/Engine/Plugins/Runtime/Harmonix/Source/HarmonixMidiTests/Private/ConformMidiFileLengthTests.cpp
@@ -2,6 +2,7 @@
 
 #include "MidiTestUtility.h"
 #include "Misc/AutomationTest.h"
+#include "Internationalization/Regex.h"
 
 #if WITH_DEV_AUTOMATION_TESTS
 
@@ -26,66 +27,96 @@ namespace HarmonixMidiTests::ConformMidiFileLength
 	*	1 Note On / Note Off pair and 1 Text event will be added to the tick position of that fractional bar length
 	*	(Note On and Note Off events are 1 tick away so they don't end up on the same tick)
 	*/
-	void AddEventsToTestMidiFile(UMidiFile* InMidiFile, int32 InNumTracks, int32 InNumChannels, float InFileLengthBars)
+	void AddEventsToTestMidiFile(UMidiFile* InMidiFile, float InFileLengthBars, int32 InNumChannels, int32 InNumTracksExcludingConductorTrack, int32& RemovableNotes, int32& RemoveableControlChanges, int32& RemovablePitchAndPoly)
 	{
 		constexpr int32 DefaultNoteNumber = 60;//C4
 		constexpr int32 DefaultNoteVelocity = 90;
 		constexpr uint8 DefaultControllerID = 69; //Hold Pedal
-		constexpr uint8 DefaultControlValue = 127;
+		constexpr uint8 DefaultControlValue = 120;
 		constexpr uint8 DefaultPitchBendValueLSB = 64;
-		constexpr uint8 DefaultPitchBendValueMSB = 127;
+		constexpr uint8 DefaultPitchBendValueMSB = 120;
 		constexpr uint8 DefaultPolyPressNoteNumber = 62;
 		constexpr uint8 DefaultPolyPresValue = 127;
 
-		//Add events to track 1 - InNumTracks, Track 0 is the conductor track
-		for (int32 TrackIndex = 1; TrackIndex < InNumTracks; ++TrackIndex)
+		// These will be useful below...
+		const FSongMaps* SongMaps = InMidiFile->GetSongMaps();
+		const FBarMap& BarMap = SongMaps->GetBarMap();
+
+		RemovableNotes = 0;
+		RemoveableControlChanges = 0;
+		RemovablePitchAndPoly = 0;
+
+		// Add events to track 1 - InNumTracks, Track 0 is the conductor track
+		for (int32 TrackIndex = 1; TrackIndex < (InNumTracksExcludingConductorTrack + 1); ++TrackIndex)
 		{
 			FMidiTrack* CurrentTrack = InMidiFile->GetTrack(TrackIndex);
 			for (int32 Channel = 0; Channel < InNumChannels; ++Channel)
 			{
-				for (int32 Bar = 0; Bar < FMath::CeilToInt32(InFileLengthBars); ++Bar)
+				for (int32 Bar = 0; Bar < FMath::FloorToInt32(InFileLengthBars); ++Bar)
 				{
-					AddNoteOnNoteOffPairsToFile(InMidiFile, DefaultNoteNumber, DefaultNoteVelocity, TrackIndex, Channel, (float)Bar);
+					int32 DestinationTick = BarMap.BarBeatTickIncludingCountInToTick(Bar, 1, 0);
+					int32 Duration = SongMaps->SubdivisionToMidiTicks(EMidiClockSubdivisionQuantization::Beat, DestinationTick) - 1;
+					AddNoteOnNoteOffPairToFile(InMidiFile, DefaultNoteNumber, DefaultNoteVelocity, TrackIndex, Channel, DestinationTick, Duration);
 
-					//Currently, Midi CC events & Midi Text events are added to bars with EVEN bar numbers,
-					//Midi Poly Press & Pitch Bend event are added to bars with ODD bar numbers
+					// Currently, Midi CC events & Midi Text events are added to bars with EVEN bar numbers,
+					// Midi Poly Press & Pitch Bend event are added to bars with ODD bar numbers
 					if (Bar % 2 == 0)
 					{
-						//Add Control Events to each channels/tracks
-						AddCCEventsToFile(InMidiFile, DefaultControllerID, DefaultNoteNumber, TrackIndex, Channel, (float)Bar);
-						//Add Text events to tracks
-						AddTextEventsToFile(InMidiFile, TEXT("TextInMidiFile"), TrackIndex, Bar);
+						// Add Control Events to each channels/tracks
+						AddCCEventToFile(InMidiFile, DefaultControllerID, DefaultNoteNumber, TrackIndex, Channel, DestinationTick);
+						// Add Text events to tracks
+						if (Channel == 0)
+						{
+							AddTextEventToFile(InMidiFile, TEXT("TextInMidiFile"), TrackIndex, DestinationTick);
+						}
 					}
 					else
 					{
 						//Add pitch bend events to channels/tracks
-						AddPitchEventsToFile(InMidiFile, DefaultPitchBendValueLSB, DefaultPitchBendValueMSB, TrackIndex, Channel, (float)Bar);
+						AddPitchEventToFile(InMidiFile, DefaultPitchBendValueLSB, DefaultPitchBendValueMSB, TrackIndex, Channel, DestinationTick);
 						//Add Poly Pres events to channels/tracks
-						AddPolyPresEventsToFile(InMidiFile, DefaultPolyPressNoteNumber, DefaultPolyPresValue, TrackIndex, Channel, (float)Bar);
+						AddPolyPresEventToFile(InMidiFile, DefaultPolyPressNoteNumber, DefaultPolyPresValue, TrackIndex, Channel, DestinationTick);
 					}
 				}
-				//if bar length is a fractional number, add additional midi events after the last integer bar
+				// If bar length is a fractional number, add additional midi events after the last integer bar
 				if (InFileLengthBars != (int32)InFileLengthBars)
 				{
-					//This will end up getting removed
-					AddNoteOnNoteOffPairsToFile(InMidiFile, DefaultNoteNumber, DefaultNoteVelocity, TrackIndex, Channel, InFileLengthBars);
-					//Add 2 CC Events (same controller ID) to test the conform function where it should remove events with the same type on the last tick
-					AddCCEventsToFile(InMidiFile, DefaultControllerID, DefaultControlValue, TrackIndex, Channel, InFileLengthBars);
-					AddCCEventsToFile(InMidiFile, DefaultControllerID, DefaultControlValue + 1, TrackIndex, Channel, InFileLengthBars);
-
-					//Add 2 Pitch Bend Events to test the conform function where it should remove events with the same type on the last tick
-					AddPitchEventsToFile(InMidiFile, DefaultPitchBendValueLSB, DefaultPitchBendValueMSB, TrackIndex, Channel, InFileLengthBars);
-					AddPitchEventsToFile(InMidiFile, DefaultPitchBendValueLSB + 1, DefaultPitchBendValueMSB + 1, TrackIndex, Channel, InFileLengthBars);
-
-					//Add 2 Pitch Bend Events to test the conform function where it should remove events with the same type on the last tick
-					AddPolyPresEventsToFile(InMidiFile, DefaultPolyPressNoteNumber, DefaultPolyPresValue, TrackIndex, Channel, InFileLengthBars);
-					AddPolyPresEventsToFile(InMidiFile, DefaultPolyPressNoteNumber, DefaultPolyPresValue + 1, TrackIndex, Channel, InFileLengthBars);
+					int32 DestinationTick = BarMap.FractionalBarIncludingCountInToTick(InFileLengthBars);
+					int32 Duration = SongMaps->SubdivisionToMidiTicks(EMidiClockSubdivisionQuantization::Beat, DestinationTick) - 1;
+					int32 NoteDestinationTick = DestinationTick - (Duration + 1);
+					DestinationTick--;
+
+					// This will end up getting removed
+					AddNoteOnNoteOffPairToFile(InMidiFile, DefaultNoteNumber, DefaultNoteVelocity, TrackIndex, Channel, NoteDestinationTick, Duration);
+					RemovableNotes++;
+
+					// Add 2 CC Events (same controller ID) to test the conform function where it should remove events with the same type on the last tick
+					// 1 of these should be removed by the conform...
+					AddCCEventToFile(InMidiFile, DefaultControllerID, DefaultControlValue, TrackIndex, Channel, DestinationTick);
+					RemoveableControlChanges++;
+					AddCCEventToFile(InMidiFile, DefaultControllerID, DefaultControlValue + 1, TrackIndex, Channel, DestinationTick);
+
+					// Add 3 Pitch Bend Events to test the conform function where it should remove events with the same type on the last tick
+					// 2 of these should be removed by the conform...
+					AddPitchEventToFile(InMidiFile, DefaultPitchBendValueLSB, DefaultPitchBendValueMSB, TrackIndex, Channel, DestinationTick);
+					AddPitchEventToFile(InMidiFile, DefaultPitchBendValueLSB + 1, DefaultPitchBendValueMSB + 1, TrackIndex, Channel, DestinationTick);
+					RemovablePitchAndPoly += 2;
+					AddPitchEventToFile(InMidiFile, DefaultPitchBendValueLSB - 1, DefaultPitchBendValueMSB - 1, TrackIndex, Channel, DestinationTick);
+
+					// Add 4 Poly pressure Events to test the conform function where it should remove events with the same type on the last tick
+					// 3 of these should be removed by the conform...
+					AddPolyPresEventToFile(InMidiFile, DefaultPolyPressNoteNumber, DefaultPolyPresValue, TrackIndex, Channel, DestinationTick);
+					AddPolyPresEventToFile(InMidiFile, DefaultPolyPressNoteNumber, DefaultPolyPresValue - 1, TrackIndex, Channel, DestinationTick);
+					AddPolyPresEventToFile(InMidiFile, DefaultPolyPressNoteNumber, DefaultPolyPresValue - 2, TrackIndex, Channel, DestinationTick);
+					RemovablePitchAndPoly += 3;
+					AddPolyPresEventToFile(InMidiFile, DefaultPolyPressNoteNumber, DefaultPolyPresValue + 1, TrackIndex, Channel, DestinationTick);
 				}
 			}
+			InMidiFile->GetTrack(TrackIndex)->Sort();
 		}
 
 		//update file information in SongLengthData
-		InMidiFile->TracksChanged();
+		InMidiFile->ScanTracksForSongLengthChange();
 	}
 
 	/**
@@ -105,46 +136,55 @@ namespace HarmonixMidiTests::ConformMidiFileLength
 
 			for (int EventIndex = Events.Num() - 1; EventIndex >= 0 && Events[EventIndex].GetTick() == ConformedLastEventTick; --EventIndex)
 			{
-				const FMidiMsg& CurrentMsg = Events[EventIndex].GetMsg();
-				uint8 MsgStatus = CurrentMsg.Status;
+				using namespace Harmonix::Midi::Constants;
+
+				const FMidiEvent& Event = Events[EventIndex];
+				const FMidiMsg& Msg = Event.GetMsg();
 
-				//check for Note On/Note Off pairs on the last tick
-				if (CurrentMsg.IsNoteOff())
+				if (Msg.IsStd())
 				{
-					for (int32 i = EventIndex - 1; i > 0 && Events[i].GetTick() == ConformedLastEventTick; --i)
+					uint8 MsgType = Msg.GetStdStatusType();
+
+					//filter Note On/Note off pairs on the last tick
+					if (MsgType == GNoteOff)
 					{
-						//check equality of note on/note off events' midi note number (data1) 
-						if (Events[i].GetMsg().IsNoteOn() && Events[i].GetMsg().GetStdData1() == CurrentMsg.GetStdData1())
+						for (int32 i = EventIndex - 1; i > 0 && Events[i].GetTick() == ConformedLastEventTick; --i)
 						{
-							return true;
+							//check equality of note on/note off events' midi note number (data1) 
+							if (Events[i].GetMsg().IsNoteOn() && Events[i].GetMsg().GetStdData1() == Msg.GetStdData1())
+							{
+								return true;
+							}
 						}
 					}
-				}
 
-				//check for pitch bend events on the last tick
-				if (MsgStatus == Harmonix::Midi::Constants::GPitch)
-				{
-					//check if there exist multiple events with same status on the last tick
-					for (int32 i = EventIndex - 1; i >= 0 && Events[i].GetTick() == ConformedLastEventTick; --i)
+					//filter chan press/pitch bend/program change events
+					else if (MsgType == GChanPres || MsgType == GPitch || MsgType == GProgram)
 					{
-						if (Events[i].GetMsg().Status == MsgStatus)
+						//check if there exist multiple events with same status on the last tick
+						for (int32 i = EventIndex - 1; i >= 0 && Events[i].GetTick() == ConformedLastEventTick; --i)
 						{
-							return true;
+							if (Events[i].GetMsg().IsStd() &&
+								Events[i].GetMsg().Status == Msg.Status)
+							{
+								return true;
+							}
 						}
 					}
-				}
 
-				//check for CC events and Poly Press on the last tick with the same controller ID
-				if (MsgStatus == Harmonix::Midi::Constants::GControl || MsgStatus == Harmonix::Midi::Constants::GPolyPres)
-				{
-					uint8 CurrentEventControllerId = CurrentMsg.Data1;
-					for (int32 i = EventIndex - 1; i >= 0 && Events[i].GetTick() == ConformedLastEventTick; --i)
+					//filter Control Change events and Poly Pres events
+					else if (MsgType == GControl || MsgType == GPolyPres)
 					{
-						//check control change events for identical controller ID (data1)
-						//check poly pres events for the same note number (data1)
-						if (Events[i].GetMsg().Data1 == CurrentEventControllerId)
+						for (int32 i = EventIndex - 1; i >= 0 && Events[i].GetTick() == ConformedLastEventTick; --i)
 						{
-							return true;
+							//check control change events for identical controller ID (data1) on the same tick and remove the later one
+							//check poly pres events for the same note number (data1) on the same tick and remove the later one
+							if (Events[i].GetMsg().IsStd() &&
+								Events[i].GetMsg().GetStdStatus() == Msg.GetStdStatus() &&
+								Events[i].GetMsg().Data1 == Msg.Data1)
+							{
+								return true;
+							}
 						}
 					}
 				}
@@ -154,100 +194,6 @@ namespace HarmonixMidiTests::ConformMidiFileLength
 		return false;
 	}
 
-	/**
-	 * Test the results of the Midi file length conform options after calling the ConformMidiFileLength() function on the input midi file
-	 */
-	static bool RunTestConformed(FAutomationTestBase& InTest, EMidiFileLengthConformOption Option, UMidiFile* InConformedMidiFile, float InFileLengthBars, int32 OriginalLastEventTick)
-	{
-		//Get the conformed Last Event Tick (should be the very last tick of the last integer bar)
-		int32 LastEventTickConformed = InConformedMidiFile->GetLastEventTick();
-		//Get the conformed Last Tick from the Song Length Data
-		int32 SongLengthLastTickConformed = InConformedMidiFile->GetSongMaps()->GetSongLengthData().LastTick;
-		//Expected last event tick after rounding up/down(nearest)
-		int32 ExpectedLastTickRoundUp = InConformedMidiFile->GetSongMaps()->GetBarMap().BarIncludingCountInToTick(FMath::CeilToInt32(InFileLengthBars)) - 1;
-		int32 ExpectedLastTickRoundDown = InConformedMidiFile->GetSongMaps()->GetBarMap().BarIncludingCountInToTick((int32)InFileLengthBars) - 1;
-
-		//tests for each conform options:
-		switch (Option) {
-		case EMidiFileLengthConformOption::RoundUp:
-			//round up doesn't do anything to the events, meaning the last event tick will still be the same as before, but song length 
-			//data will change to the last tick of the next integer bar
-			if (InFileLengthBars < 1.0)
-			{
-				//after conforming by rounding up, all rounding options should be disabled if file length is less than 1 bar
-				if (!InTest.TestFalse("Cannot conform by rounding up down, or nearest", (InConformedMidiFile->ShouldConformMidiFileLength(EMidiFileLengthConformOption::RoundUp) ||
-					InConformedMidiFile->ShouldConformMidiFileLength(EMidiFileLengthConformOption::RoundDown) || InConformedMidiFile->ShouldConformMidiFileLength(EMidiFileLengthConformOption::Nearest))))
-				{
-					return false;
-				}
-			}
-			else {
-				//if file length is greater than 1 bar, round down should still be available after rounding up, nearest shouldn't be 
-				if (!InTest.TestTrue("Can conform by rounding down", InConformedMidiFile->ShouldConformMidiFileLength(EMidiFileLengthConformOption::RoundDown)))
-				{
-					return false;
-				}
-				if (!InTest.TestFalse("Cannot conform by rounding up or nearest", (InConformedMidiFile->ShouldConformMidiFileLength(EMidiFileLengthConformOption::RoundUp) ||
-					InConformedMidiFile->ShouldConformMidiFileLength(EMidiFileLengthConformOption::Nearest))))
-				{
-					return false;
-				}
-			}
-			//events should have no changes
-			if (!InTest.TestEqual("Last event tick should be the same after rounding up", LastEventTickConformed, OriginalLastEventTick))
-			{
-				return false;
-			}
-
-			if (!InTest.TestEqual("Last tick in song length data should be rounded up to the nearest integer bar", SongLengthLastTickConformed, ExpectedLastTickRoundUp))
-			{
-				return false;
-			}
-			break;
-
-		case EMidiFileLengthConformOption::RoundDown:
-			//Last tick in song length data should be at the end of the last integer bar
-			if (!InTest.TestEqual("Last tick in song length data after conforming should be at the end of the last integer bar of the file", SongLengthLastTickConformed,ExpectedLastTickRoundDown))
-			{
-				return false;
-			}
-
-			//conformed last event tick should be less than or equal to the last tick of song length data in the last integer bar
-			//NOTE: we check 'less than or equal to' for a special case:
-			//if a Midi file has a fractional bar length,
-			//and the only Midi events exceeding the last integer bar are Note On/Note Off pairs,
-			//while there is no Midi event on the very last tick of the last integer bar ---------
-			//they'll first be moved to the last tick of the last integer bar and then REMOVED since we don't want Note On/Note Off pairs on 
-			//the same tick. After this, the last event tick in file might end up being less than expected (the end of the last integer bar)
-			//(The last tick in song length data will still be at the end of the last integer bar after calling TracksChanged())
-			if (!InTest.TestTrue("Conformed last event tick should be less than or equal to the last tick in song length data", LastEventTickConformed <= SongLengthLastTickConformed))
-			{
-				return false;
-			}
-
-			//after conforming by rounding down, all rounding options should be disabled 
-			if (!InTest.TestFalse("Cannot conform by rounding up down, or nearest", (InConformedMidiFile->ShouldConformMidiFileLength(EMidiFileLengthConformOption::RoundUp) ||
-				InConformedMidiFile->ShouldConformMidiFileLength(EMidiFileLengthConformOption::RoundDown) || InConformedMidiFile->ShouldConformMidiFileLength(EMidiFileLengthConformOption::Nearest))))
-			{
-				return false;
-			}
-
-			//check for excessive events on the last tick after conforming 
-			if (!InTest.TestFalse("There should not be excessive Pitch Bend/CC/Poly Press or Note On/Note Off pairs on the last tick after conforming", ContainsExcessiveEventsAtLastEventTick(InConformedMidiFile)))
-			{
-				return false;
-			}
-			break;
-		case EMidiFileLengthConformOption::Nearest:
-			//round to nearest is either rounding down or rounding up
-			;
-				break;
-		}
-		return true;
-	}
-
-
-
 	IMPLEMENT_SIMPLE_AUTOMATION_TEST(
 		FTestConformMidiFileLength,
 		"Harmonix.Midi.ConformMidiFileLength",
@@ -255,141 +201,93 @@ namespace HarmonixMidiTests::ConformMidiFileLength
 		bool FTestConformMidiFileLength::RunTest(const FString&)
 	{
 		//input test values for creating a midi file
-		constexpr float InFileLengthBars = 5.5f;
-		constexpr int32 InNumChannels = 2;
-		constexpr int32 InNumTracks = 4;
-		constexpr int32 InTimeSigNum6 = 6;
-		constexpr int32 InTimeSigDenum8 = 8;
-		constexpr int32 InTimeSigNum4 = 4;
-		constexpr int32 InTimeSigDenum4 = 4;
-		constexpr int32 InTempo = 120;
+		constexpr float FileLengthBars = 5.5f;
+		constexpr int32 NumChannels = 2;
+		constexpr int32 NumTracksIncludingConductor = 4;
+		constexpr int32 NumTracksExcludingConductor = NumTracksIncludingConductor - 1;
+		constexpr int32 TimeSigNum5 = 5;
+		constexpr int32 TimeSigDenum8 = 8;
+		constexpr int32 TimeSigNum4 = 4;
+		constexpr int32 TimeSigDenum4 = 4;
+		constexpr int32 Tempo = 120;
 		
 
-		//if a Midi file only has events on tick 0 (the very beginning), it should still be able to be conformed by rounding up
+		// We should get a warning and no quantization if we ask to round down because that would result in a zero length file. 
+		AddExpectedMessage(TEXT("QuantizeLengthToNearestPerfectSubdivision: Asked to Quantize file length DOWN, but that would result in a zero length midi file. NOT ALLOWED! Skipping quantization!"), ELogVerbosity::Warning);
 		constexpr float FileLength0Bar = 0.0;
-		UMidiFile* MidiFile0Bar = BuildMidiFile(FileLength0Bar, InNumChannels, InNumTracks, InTimeSigNum4, InTimeSigDenum4, InTempo);
-		//no need to round up again since we did it upon creating the file
-		int32 ExpectedLastEventTick0BarRoundedUp = MidiFile0Bar->GetSongMaps()->GetBarMap().BarIncludingCountInToTick(1) - 1;
-		int32 SongLengthDataLastEventTick0BarRoundedUp = MidiFile0Bar->GetSongMaps()->GetSongLengthData().LastTick;
-		//last tick in song length data should be at the end of bar 1, but the last event tick should still be 0
-		UTEST_TRUE("A Midi file that only has events on tick 0 should still be rounded up to 1 Bar (in song length data)",
-			ExpectedLastEventTick0BarRoundedUp == SongLengthDataLastEventTick0BarRoundedUp);
-		//verify the last event tick doesn't change
-		UTEST_TRUE("Verify the last event tick is still 0",
-			MidiFile0Bar->GetLastEventTick() == 0
-		);
-
-		//create an empty midi file to test for rounding up
-		UMidiFile* MidiFileToTestRoundUp = BuildMidiFile(InFileLengthBars, InNumChannels, InNumTracks, InTimeSigNum6, InTimeSigDenum8, InTempo);
-		//Add some events to the midi file for testing 
-		AddEventsToTestMidiFile(MidiFileToTestRoundUp, InNumTracks, InNumChannels, InFileLengthBars);
-		//get the last event tick before conforming
-		int32 LastEventTickNotConformed = MidiFileToTestRoundUp->GetLastEventTick();
-		//conform by rounding up
-		MidiFileToTestRoundUp->ConformMidiFileLength(EMidiFileLengthConformOption::RoundUp);
-		//check results
-		if (!RunTestConformed(*this, EMidiFileLengthConformOption::RoundUp, MidiFileToTestRoundUp, InFileLengthBars,LastEventTickNotConformed))
-		{
-			return false;
-		}
+		UMidiFile* MidiFile = CreateAndInitializaMidiFile(FileLength0Bar, NumTracksIncludingConductor, TimeSigNum4, TimeSigDenum4, Tempo, true);
+		FSongLengthData PreQuantizeLengthData = MidiFile->GetSongMaps()->GetSongLengthData();
+		MidiFile->QuantizeLengthToSubdivision(EMidiFileQuantizeDirection::Down, EMidiClockSubdivisionQuantization::Bar);
+		FSongLengthData PostQuantizeLengthData = MidiFile->GetSongMaps()->GetSongLengthData();
 
-		//create an empty midi file to test for rounding down
-		UMidiFile* MidiFileToTestRoundDown = BuildMidiFile(InFileLengthBars, InNumChannels, InNumTracks, InTimeSigNum4, InTimeSigDenum4, InTempo);
-		//Add some events to the midi file for testing 
-		AddEventsToTestMidiFile(MidiFileToTestRoundDown, InNumTracks, InNumChannels, InFileLengthBars);
-		//conform by rounding down
-		MidiFileToTestRoundDown->ConformMidiFileLength(EMidiFileLengthConformOption::RoundDown);
-		//check results
-		if (!RunTestConformed(*this, EMidiFileLengthConformOption::RoundDown, MidiFileToTestRoundDown, InFileLengthBars, LastEventTickNotConformed))
-		{
-			return false;
-		}
+		int32 ExpectedLengthTickPostQuantize = 1;
+		int32 ExpectedLastTickPostQuantize = 0;
 
-		//round to nearest, if the fractional part is greater than 0.5, results will be the same as round up
-		//otherwise equivalent to round down
+		UTEST_TRUE("Last ticks are correct pre/post quantization.",
+			PreQuantizeLengthData.LastTick == 0 &&
+			PostQuantizeLengthData.LastTick == ExpectedLastTickPostQuantize);
 
-		constexpr float TestFileLengthNearest = 5.25;
-		//create an empty midi file to test for rounding to nearest (down)
-		UMidiFile* MidiFileToTestRoundToNearestDown = BuildMidiFile(TestFileLengthNearest, InNumChannels, InNumTracks, InTimeSigNum4, InTimeSigDenum4, InTempo);
-		//Add some events to the midi file for testing 
-		AddEventsToTestMidiFile(MidiFileToTestRoundToNearestDown, InNumTracks, InNumChannels, TestFileLengthNearest);
-		UTEST_FALSE("After rounding up upon creating the file, it cannot be conformed by rounding up again or to the nearest.",
-			MidiFileToTestRoundToNearestDown->ShouldConformMidiFileLength(EMidiFileLengthConformOption::RoundUp) ||
-			MidiFileToTestRoundToNearestDown->ShouldConformMidiFileLength(EMidiFileLengthConformOption::Nearest)
-		);
-		//Ideally, this file can still be rounded down, but calling ConformMidiFileLength(EMidiFileLengthConformOption::Nearest) won't do it,
-		//because we automatically round the file up upon creating it, and it shouldn't be able to 'round to nearest anymore'
-		UTEST_TRUE("A file with fractional portion of its length (bars) less than 0.5 can be rounded down, but not by any other conform options",
-			MidiFileToTestRoundToNearestDown->ShouldConformMidiFileLength(EMidiFileLengthConformOption::RoundDown)
-			);
-		//since the fractional portion is greater than 0.5, this is equivalent to rounding up
-		constexpr float TestFileLengthNearest2 = 5.7;
-		//create an empty midi file to test for rounding to nearest (up)
-		UMidiFile* MidiFileToTestRoundToNearestUp = BuildMidiFile(TestFileLengthNearest2, InNumChannels, InNumTracks, InTimeSigNum4, InTimeSigDenum4, InTempo);
-		//Add some events to the midi file for testing 
-		AddEventsToTestMidiFile(MidiFileToTestRoundToNearestUp, InNumTracks, InNumChannels, TestFileLengthNearest2);
-		//get the last event tick before conforming
-		int32 LastEventTickNotConformedNearestUp = MidiFileToTestRoundToNearestUp->GetLastEventTick();
-		//conform by rounding to nearest (up)
-		MidiFileToTestRoundToNearestUp->ConformMidiFileLength(EMidiFileLengthConformOption::Nearest);
-		//since we know this result will be the same as round up, just test for round up
-		if (!RunTestConformed(*this, EMidiFileLengthConformOption::RoundUp, MidiFileToTestRoundToNearestUp, TestFileLengthNearest2, LastEventTickNotConformedNearestUp))
-		{
-			return false;
-		}
+		UTEST_TRUE("Length ticks are correct pre/post quantization.",
+			PreQuantizeLengthData.LengthTicks == 1 &&
+			PostQuantizeLengthData.LengthTicks == ExpectedLengthTickPostQuantize);
 
-		//special case: bar length is less than 1 bar
-		constexpr float TestFileLength1Bar = 0.6;
-		//create an empty midi file for testing bar length less than 1 bar
-		UMidiFile* MidiFileToTest1Bar = BuildMidiFile(TestFileLength1Bar, InNumChannels, InNumTracks, InTimeSigNum4, InTimeSigDenum4, InTempo);
-		//Add some events to the midi file for testing 
-		AddEventsToTestMidiFile(MidiFileToTest1Bar, InNumTracks, InNumChannels, TestFileLength1Bar);
-		//get the last event tick before conforming
-		int32 LastEventTickNotConformed1Bar = MidiFileToTest1Bar->GetLastEventTick();
-		//in this case, this file's length cannot round down or nearest, only can round up
-		//(we rounded up already when creating the file, so no option should be conformable
-		UTEST_FALSE("A Midi file with length less than 1 bar can only be conformed by rounding up",
-			MidiFileToTest1Bar->ShouldConformMidiFileLength(EMidiFileLengthConformOption::RoundUp) || 
-			MidiFileToTest1Bar->ShouldConformMidiFileLength(EMidiFileLengthConformOption::RoundDown) ||
-			MidiFileToTest1Bar->ShouldConformMidiFileLength(EMidiFileLengthConformOption::Nearest)
-			);
-
-		//conform by rounding up
-		MidiFileToTest1Bar->ConformMidiFileLength(EMidiFileLengthConformOption::RoundUp);
-		//check results
-		if (!RunTestConformed(*this, EMidiFileLengthConformOption::RoundUp, MidiFileToTest1Bar, TestFileLength1Bar, LastEventTickNotConformed1Bar))
-		{
-			return false;
-		}
+		// "Nearest" should always be up since down would result in a zero length midi file...
+		PreQuantizeLengthData = MidiFile->GetSongMaps()->GetSongLengthData();
+		MidiFile->QuantizeLengthToSubdivision(EMidiFileQuantizeDirection::Nearest, EMidiClockSubdivisionQuantization::Bar);
+		PostQuantizeLengthData = MidiFile->GetSongMaps()->GetSongLengthData();
 
-		//test a midi file that should not be conformed (having an integer bar length)
-		constexpr float IntegerFileLength = 5.0;
-		//create a midi file with an integer bar length
-		UMidiFile* MidiFileToTestIntegerLength = BuildMidiFile(IntegerFileLength, InNumChannels, InNumTracks, InTimeSigNum4, InTimeSigDenum4, InTempo);
+		ExpectedLengthTickPostQuantize = (MidiFile->GetSongMaps()->GetTicksPerQuarterNote() * 4); // <-- because we know this test has a time signature of 4/4
+		ExpectedLastTickPostQuantize = ExpectedLengthTickPostQuantize - 1;
+
+		UTEST_TRUE("Last ticks are correct pre/post quantization.",
+			PreQuantizeLengthData.LastTick == 0 &&
+			PostQuantizeLengthData.LastTick == ExpectedLastTickPostQuantize);
+
+		UTEST_TRUE("Length ticks are correct pre/post quantization.",
+			PreQuantizeLengthData.LengthTicks == 1 &&
+			PostQuantizeLengthData.LengthTicks == ExpectedLengthTickPostQuantize);
+
+		// "Up" should definitely work...
+		MidiFile = CreateAndInitializaMidiFile(FileLength0Bar, NumTracksIncludingConductor, TimeSigNum4, TimeSigDenum4, Tempo, true);
+		PreQuantizeLengthData = MidiFile->GetSongMaps()->GetSongLengthData();
+		MidiFile->QuantizeLengthToSubdivision(EMidiFileQuantizeDirection::Up, EMidiClockSubdivisionQuantization::Bar);
+		PostQuantizeLengthData = MidiFile->GetSongMaps()->GetSongLengthData();
+
+		ExpectedLengthTickPostQuantize = (MidiFile->GetSongMaps()->GetTicksPerQuarterNote() * 4); // <-- because we know this test has a time signature of 4/4
+		ExpectedLastTickPostQuantize = ExpectedLengthTickPostQuantize - 1;
+
+		UTEST_TRUE("Last ticks are correct pre/post quantization.",
+			PreQuantizeLengthData.LastTick == 0 &&
+			PostQuantizeLengthData.LastTick == ExpectedLastTickPostQuantize); 
+
+		UTEST_TRUE("Length ticks are correct pre/post quantization.",
+			PreQuantizeLengthData.LengthTicks == 1 &&
+			PostQuantizeLengthData.LengthTicks == ExpectedLengthTickPostQuantize);
+
+		// Let's try a longer length in 5/8...
+		constexpr float TestFileLengthNearest = 5.25;
+		MidiFile = CreateAndInitializaMidiFile(TestFileLengthNearest, NumTracksIncludingConductor, TimeSigNum5, TimeSigDenum8, Tempo, true);
 		//Add some events to the midi file for testing 
-		AddEventsToTestMidiFile(MidiFileToTestIntegerLength, InNumTracks, InNumChannels, IntegerFileLength);
-		//get the last event tick
-		int32 LastEventTickIntegerLength = MidiFileToTestIntegerLength->GetLastEventTick();
-
-		//No conform option should be enabled
-		UTEST_FALSE("A Midi file with an integer bar length cannot be conformed",
-			MidiFileToTestIntegerLength->ShouldConformMidiFileLength(EMidiFileLengthConformOption::RoundUp) ||
-			MidiFileToTestIntegerLength->ShouldConformMidiFileLength(EMidiFileLengthConformOption::RoundDown) || 
-			MidiFileToTestIntegerLength->ShouldConformMidiFileLength(EMidiFileLengthConformOption::Nearest)
-		);
-
-		//"conform" file length (expects nothing to change in file)
-		//(we can still call the ConformMidiFileLength() function but nothing will happen to the file)
-		MidiFileToTestIntegerLength->ConformMidiFileLength(EMidiFileLengthConformOption::RoundUp);
-		MidiFileToTestIntegerLength->ConformMidiFileLength(EMidiFileLengthConformOption::RoundDown);
-		MidiFileToTestIntegerLength->ConformMidiFileLength(EMidiFileLengthConformOption::Nearest);
-		//check results 
-		UTEST_EQUAL("Nothing should change to a midi file that cannot be conformed despite calling the ConformMidiFileLength function on it",
-			LastEventTickIntegerLength,
-			MidiFileToTestIntegerLength->GetSongMaps()->GetBarMap().BarIncludingCountInToTick(IntegerFileLength) - 1);
-		UTEST_EQUAL("Last tick in song length data should also equal to last event tick",
-			LastEventTickIntegerLength,
-			MidiFileToTestIntegerLength->GetSongMaps()->GetSongLengthData().LastTick);
+		int32 RemoveableNotes = 0;
+		int32 RemoveablePicthAndPoly = 0;
+		int32 RemoveableControl = 0;
+		AddEventsToTestMidiFile(MidiFile, TestFileLengthNearest, NumChannels, NumTracksExcludingConductor, RemoveableNotes, RemoveableControl, RemoveablePicthAndPoly);
+		PreQuantizeLengthData = MidiFile->GetSongMaps()->GetSongLengthData();
+		MidiFile->QuantizeLengthToSubdivision(EMidiFileQuantizeDirection::Down, EMidiClockSubdivisionQuantization::Bar);
+		PostQuantizeLengthData = MidiFile->GetSongMaps()->GetSongLengthData();
+
+		ExpectedLengthTickPostQuantize = ((MidiFile->GetSongMaps()->GetTicksPerQuarterNote() / 2) * 5) * 5; // <-- because we know this test has a time signature of 4/4
+		ExpectedLastTickPostQuantize = ExpectedLengthTickPostQuantize - 1;
+
+		UTEST_TRUE("Last ticks are correct pre/post quantization.",
+			PreQuantizeLengthData.LastTick == ((((MidiFile->GetSongMaps()->GetTicksPerQuarterNote() / 2) * 5) * TestFileLengthNearest) - 1) &&
+			PostQuantizeLengthData.LastTick == ExpectedLastTickPostQuantize);
+
+		UTEST_TRUE("Length ticks are corrent pre/post quantization.",
+			PreQuantizeLengthData.LengthTicks == (((MidiFile->GetSongMaps()->GetTicksPerQuarterNote() / 2) * 5) * TestFileLengthNearest) &&
+			PostQuantizeLengthData.LengthTicks == ExpectedLengthTickPostQuantize);
+
+		UTEST_FALSE("No excessive events on conformed last tick.", ContainsExcessiveEventsAtLastEventTick(MidiFile));
 
 		return true;
 	}
diff --git a/Engine/Plugins/Runtime/Harmonix/Source/HarmonixMidiTests/Private/MidiFileImportExportTests.cpp b/Engine/Plugins/Runtime/Harmonix/Source/HarmonixMidiTests/Private/MidiFileImportExportTests.cpp
--- a/Engine/Plugins/Runtime/Harmonix/Source/HarmonixMidiTests/Private/MidiFileImportExportTests.cpp
+++ b/Engine/Plugins/Runtime/Harmonix/Source/HarmonixMidiTests/Private/MidiFileImportExportTests.cpp
@@ -33,7 +33,7 @@ namespace HarmonixMidiTests::ImportExportTests
 	*	1 Note On / Note Off pair and 1 Text event will be added to the tick position of that fractional bar length
 	*	(Note On and Note Off events are 1 tick away so they don't end up on the same tick)
 	*/
-	void AddEventsToTestMidiFile(UMidiFile* InMidiFile, int32 InNumTracks, int32 InNumChannels, float InFileLengthBars)
+	void AddEventsToTestMidiFile(UMidiFile* InMidiFile, int32 InNumTracksExcludingConductorTrack, int32 InNumChannels, float InFileLengthBars)
 	{
 		constexpr int32 DefaultNoteNumber = 60;//C4
 		constexpr int32 DefaultNoteVelocity = 90;
@@ -44,56 +44,66 @@ namespace HarmonixMidiTests::ImportExportTests
 		constexpr uint8 DefaultPolyPressNoteNumber = 62;
 		constexpr uint8 DefaultPolyPresValue = 127;
 
+		// these will be useful below...
+		const FSongMaps* SongMaps =  InMidiFile->GetSongMaps();
+		const FBarMap& BarMap =SongMaps->GetBarMap();
+		
 		//Add events to track 1 - InNumTracks, Track 0 is the conductor track
-		for (int32 TrackIndex = 1; TrackIndex < InNumTracks; ++TrackIndex)
+		for (int32 TrackIndex = 1; TrackIndex < (InNumTracksExcludingConductorTrack + 1); ++TrackIndex)
 		{
 			FMidiTrack* CurrentTrack = InMidiFile->GetTrack(TrackIndex);
 			for (int32 Channel = 0; Channel < InNumChannels; ++Channel)
 			{
 				for (int32 Bar = 0; Bar < FMath::CeilToInt32(InFileLengthBars); ++Bar)
 				{
-					AddNoteOnNoteOffPairsToFile(InMidiFile, DefaultNoteNumber, DefaultNoteVelocity, TrackIndex, Channel, (float)Bar);
+					int32 DestinationTick = BarMap.BarBeatTickIncludingCountInToTick(Bar, 1, 0);
+					int32 Duration = SongMaps->SubdivisionToMidiTicks(EMidiClockSubdivisionQuantization::Beat, DestinationTick) - 1;
+					AddNoteOnNoteOffPairToFile(InMidiFile, DefaultNoteNumber, DefaultNoteVelocity, TrackIndex, Channel, DestinationTick, Duration);
 
 					//Currently, Midi CC events & Midi Text events are added to bars with EVEN bar numbers,
 					//Midi Poly Press & Pitch Bend event are added to bars with ODD bar numbers
 					if (Bar % 2 == 0)
 					{
 						//Add Control Events to each channels/tracks
-						AddCCEventsToFile(InMidiFile, DefaultControllerID, DefaultNoteNumber, TrackIndex, Channel, (float)Bar);
+						AddCCEventToFile(InMidiFile, DefaultControllerID, DefaultNoteNumber, TrackIndex, Channel, DestinationTick);
 						//Add Text events to tracks
-						AddTextEventsToFile(InMidiFile, TEXT("TextInMidiFile"), TrackIndex, Bar);
+						AddTextEventToFile(InMidiFile, TEXT("TextInMidiFile"), TrackIndex, DestinationTick);
 					}
 					else
 					{
 						//Add pitch bend events to channels/tracks
-						AddPitchEventsToFile(InMidiFile, DefaultPitchBendValueLSB, DefaultPitchBendValueMSB, TrackIndex, Channel, (float)Bar);
+						AddPitchEventToFile(InMidiFile, DefaultPitchBendValueLSB, DefaultPitchBendValueMSB, TrackIndex, Channel, DestinationTick);
 						//Add Poly Pres events to channels/tracks
-						AddPolyPresEventsToFile(InMidiFile, DefaultPolyPressNoteNumber, DefaultPolyPresValue, TrackIndex, Channel, (float)Bar);
+						AddPolyPresEventToFile(InMidiFile, DefaultPolyPressNoteNumber, DefaultPolyPresValue, TrackIndex, Channel, DestinationTick);
 					}
 				}
 				//if bar length is a fractional number, add additional midi events after the last integer bar
 				if (InFileLengthBars != (int32)InFileLengthBars)
 				{
+					int32 DestinationTick = BarMap.FractionalBarIncludingCountInToTick(InFileLengthBars) - 1;
+					int32 Duration = SongMaps->SubdivisionToMidiTicks(EMidiClockSubdivisionQuantization::Beat, DestinationTick) - 1;
+					int32 NoteDestinationTick = DestinationTick - (Duration + 1);
+
 					//This will end up getting removed
-					AddNoteOnNoteOffPairsToFile(InMidiFile, DefaultNoteNumber, DefaultNoteVelocity, TrackIndex, Channel, InFileLengthBars);
+					AddNoteOnNoteOffPairToFile(InMidiFile, DefaultNoteNumber, DefaultNoteVelocity, TrackIndex, Channel, NoteDestinationTick, Duration);
 					//Add 2 CC Events (same controller ID) to test the conform function where it should remove events with the same type on the last tick
-					AddCCEventsToFile(InMidiFile, DefaultControllerID, DefaultControlValue, TrackIndex, Channel, InFileLengthBars);
-					AddCCEventsToFile(InMidiFile, DefaultControllerID, DefaultControlValue + 1, TrackIndex, Channel, InFileLengthBars);
+					AddCCEventToFile(InMidiFile, DefaultControllerID, DefaultControlValue, TrackIndex, Channel, DestinationTick);
+					AddCCEventToFile(InMidiFile, DefaultControllerID, DefaultControlValue + 1, TrackIndex, Channel, DestinationTick);
 
 					//Add 2 Pitch Bend Events to test the conform function where it should remove events with the same type on the last tick
-					AddPitchEventsToFile(InMidiFile, DefaultPitchBendValueLSB, DefaultPitchBendValueMSB, TrackIndex, Channel, InFileLengthBars);
-					AddPitchEventsToFile(InMidiFile, DefaultPitchBendValueLSB + 1, DefaultPitchBendValueMSB + 1, TrackIndex, Channel, InFileLengthBars);
+					AddPitchEventToFile(InMidiFile, DefaultPitchBendValueLSB, DefaultPitchBendValueMSB, TrackIndex, Channel, DestinationTick);
+					AddPitchEventToFile(InMidiFile, DefaultPitchBendValueLSB + 1, DefaultPitchBendValueMSB + 1, TrackIndex, Channel, DestinationTick);
 
 					//Add 2 Pitch Bend Events to test the conform function where it should remove events with the same type on the last tick
-					AddPolyPresEventsToFile(InMidiFile, DefaultPolyPressNoteNumber, DefaultPolyPresValue, TrackIndex, Channel, InFileLengthBars);
-					AddPolyPresEventsToFile(InMidiFile, DefaultPolyPressNoteNumber, DefaultPolyPresValue + 1, TrackIndex, Channel, InFileLengthBars);
+					AddPolyPresEventToFile(InMidiFile, DefaultPolyPressNoteNumber, DefaultPolyPresValue, TrackIndex, Channel, DestinationTick);
+					AddPolyPresEventToFile(InMidiFile, DefaultPolyPressNoteNumber, DefaultPolyPresValue + 1, TrackIndex, Channel, DestinationTick);
 				}
 			}
 			InMidiFile->GetTrack(TrackIndex)->Sort();
 		}
 
 		//update file information in SongLengthData
-		InMidiFile->TracksChanged();
+		InMidiFile->ScanTracksForSongLengthChange();
 	}
 
 	IMPLEMENT_SIMPLE_AUTOMATION_TEST(
@@ -105,7 +115,8 @@ namespace HarmonixMidiTests::ImportExportTests
 		//input test values for creating a midi file
 		const float FileLengthBars = 5.5f;
 		const int32 NumChannels = 2;
-		const int32 NumTracks = 4;
+		const int32 NumTracksIncludingConductorTrack = 4;
+		const int32 NumTracksExcludingConductorTrack = NumTracksIncludingConductorTrack - 1;
 		const int32 TimeSigNum6 = 6;
 		const int32 TimeSigDenum8 = 8;
 		const int32 TimeSigNum4 = 4;
@@ -113,9 +124,9 @@ namespace HarmonixMidiTests::ImportExportTests
 		const int32 Tempo = 120;
 		
 		//create an empty midi file to test for rounding to nearest (down)
-		UMidiFile* GeneratedMidiFile = BuildMidiFile(FileLengthBars, NumChannels, NumTracks, TimeSigNum4, TimeSigDenum4, Tempo);
+		UMidiFile* GeneratedMidiFile = CreateAndInitializaMidiFile(FileLengthBars, NumTracksIncludingConductorTrack, TimeSigNum4, TimeSigDenum4, Tempo, true);
 		//Add some events to the midi file for testing 
-		AddEventsToTestMidiFile(GeneratedMidiFile, NumTracks, NumChannels, FileLengthBars);
+		AddEventsToTestMidiFile(GeneratedMidiFile, NumTracksExcludingConductorTrack, NumChannels, FileLengthBars);
 
 		TArray<uint8> StdMidiFileBytes;
 		TSharedPtr<FMemoryWriter> StdMidiFileOut = MakeShared<FMemoryWriter>(StdMidiFileBytes, true);
@@ -142,7 +153,8 @@ namespace HarmonixMidiTests::ImportExportTests
 		//input test values for creating a midi file
 		const float FileLengthBars = 5.5f;
 		const int32 NumChannels = 2;
-		const int32 NumTracks = 4;
+		const int32 NumTracksIncludingConductorTrack = 4;
+		const int32 NumTracksExcludingConductorTrack = NumTracksIncludingConductorTrack - 1;
 		const int32 TimeSigNum6 = 6;
 		const int32 TimeSigDenum8 = 8;
 		const int32 TimeSigNum4 = 4;
@@ -150,9 +162,9 @@ namespace HarmonixMidiTests::ImportExportTests
 		const int32 Tempo = 120;
 
 		//create an empty midi file to test for rounding to nearest (down)
-		UMidiFile* GeneratedMidiFile = BuildMidiFile(FileLengthBars, NumChannels, NumTracks, TimeSigNum4, TimeSigDenum4, Tempo);
+		UMidiFile* GeneratedMidiFile = CreateAndInitializaMidiFile(FileLengthBars, NumTracksIncludingConductorTrack, TimeSigNum4, TimeSigDenum4, Tempo, true);
 		//Add some events to the midi file for testing 
-		AddEventsToTestMidiFile(GeneratedMidiFile, NumTracks, NumChannels, FileLengthBars);
+		AddEventsToTestMidiFile(GeneratedMidiFile, NumTracksExcludingConductorTrack, NumChannels, FileLengthBars);
 
 		TArray<uint8> StdMidiFileBytes;
 		TSharedPtr<FMemoryWriter> StdMidiFileOut = MakeShared<FMemoryWriter>(StdMidiFileBytes, true);
diff --git a/Engine/Plugins/Runtime/Harmonix/Source/HarmonixMidiTests/Private/MidiTestUtility.cpp b/Engine/Plugins/Runtime/Harmonix/Source/HarmonixMidiTests/Private/MidiTestUtility.cpp
--- a/Engine/Plugins/Runtime/Harmonix/Source/HarmonixMidiTests/Private/MidiTestUtility.cpp
+++ b/Engine/Plugins/Runtime/Harmonix/Source/HarmonixMidiTests/Private/MidiTestUtility.cpp
@@ -3,7 +3,7 @@
 
 namespace Harmonix::Testing::Utility::MidiTestUtility
 {
-	UMidiFile* BuildMidiFile(float InFileLengthBars, int32 InNumChannels, int32 InNumTracks, int32 InTimeSigNum, int32 InTimeSigDenom, int32 InTempo)
+	UMidiFile* CreateAndInitializaMidiFile(float FileLengthBars, int32 InNumTracksIncludingConductor, int32 InTimeSigNum, int32 InTimeSigDenom, int32 InTempo, bool PutTextEventOnLastTick)
 	{
 		UMidiFile* TheMidiFile = NewObject<UMidiFile>();
 
@@ -21,105 +21,111 @@ namespace Harmonix::Testing::Utility::MidiTestUtility
 		TheMidiFile->GetSongMaps()->GetTempoMap().AddTempoInfoPoint(Midi::Constants::BPMToMidiTempo(InTempo), 0);
 		TheMidiFile->GetSongMaps()->GetBarMap().AddTimeSignatureAtBarIncludingCountIn(0, InTimeSigNum, InTimeSigDenom);
 		TheMidiFile->GetSongMaps()->GetBarMap().SetTicksPerQuarterNote(DefaultTicksPerQuarter);
-		// Whenever you change the song maps you need to build the corresponding midi conductor track...
+		// Whenever you change the song maps you need to either:
+		// a) Add corresponding midi events to the conductor track. Or...
+		// b) Build the conductor track from scratch.
+		// (b) is easier...
 		TheMidiFile->BuildConductorTrack();
 
+		// now we can calculate the tick length since the midi file now has a conductor track...
+		int32 TickLength = TheMidiFile->GetSongMaps()->GetBarMap().FractionalBarIncludingCountInToTick(FileLengthBars);
+		if (TickLength <= 0)
+		{
+			// MIDI Files must be at least one tick long. Even the initial tempo and time signature are specified on tick 0!
+			TickLength = 1;
+		}
+
 		//Add tracks to Midi file according to input argument
-		for (int32 Track = 1; Track < InNumTracks; ++Track)
+		for (int32 Track = 1; Track < InNumTracksIncludingConductor; ++Track)
 		{
 			//Add midi tracks 1 - InNumTracks to file, Track 0 is the conductor track
-			FMidiTrack* CurrentTrack = TheMidiFile->AddTrack(FString::Printf(TEXT("TestMidiTrack%i"), Track));
-
-			//Add a text message at the end of track 1, marking the file's last event tick
-			if (Track == 1)
+			FMidiTrack* TheNewTrack = TheMidiFile->AddTrack(FString::Printf(TEXT("TestMidiTrack%i"), Track));
+			if (TheNewTrack && PutTextEventOnLastTick)
 			{
-				uint16 TextIndex = CurrentTrack->AddText("TheEndOfTheFile");
-				int32 LastEventTickInFile = TheMidiFile->GetSongMaps()->GetBarMap().FractionalBarIncludingCountInToTick(InFileLengthBars);
-				if (InFileLengthBars == (int32)InFileLengthBars && InFileLengthBars != 0.f)
-				{
-					//if the input bar length is an integer number, the last event tick is 1 tick before the next bar
-					//(InFileLengthBars will end up being the very beginning of the next bar)
-					LastEventTickInFile -= 1;
-				}
-				FMidiEvent TextEventAtTheEnd(LastEventTickInFile, FMidiMsg::CreateText(TextIndex, Midi::Constants::GMeta_Text));
-				CurrentTrack->AddEvent(TextEventAtTheEnd);
+				uint16 TextIndex = TheNewTrack->AddText(FString::Printf(TEXT("TheEndOfTrack%i"),Track));
+				FMidiEvent TextEventAtTheEnd(TickLength - 1, FMidiMsg::CreateText(TextIndex, Midi::Constants::GMeta_Text));
+				TheNewTrack->AddEvent(TextEventAtTheEnd);
 			}
-
 		}
-		//update song length information
-		TheMidiFile->TracksChanged();
 
-		//To keep up with the results of imported Midi files, we ALWAYS round the file's length up to an
-		//interger number of bars
-		TheMidiFile->ConformMidiFileLength(EMidiFileLengthConformOption::RoundUp);
-		//if the fractional part of the file's length is greater than 0.5,
-		//rounding up and rounding to the nearest will get the same results
-		if(FGenericPlatformMath::Fractional(InFileLengthBars) >= 0.5f) TheMidiFile->bLengthRoundedToNearest = true;
+		//update song length information
+		TheMidiFile->ConformToLength(TickLength);
 
 		return TheMidiFile;
 	}
 
-	void AddNoteOnNoteOffPairsToFile(UMidiFile* InFile, int32 InNoteNumber, int32 InNoteVelocity, int32 InTrackIndex, int32 InChannel, float InBarIndex)
+	void AddNoteOnNoteOffPairToFile(UMidiFile* InFile, int32 InNoteNumber, int32 InNoteVelocity, int32 InTrackIndex, int32 InChannel, int32 AtTick, int32 DurationTicks)
 	{
+		if (InTrackIndex >= InFile->GetTracks().Num() || InTrackIndex <= 0)
+		{
+			UE_LOG(LogMIDI, Error, TEXT("AddNoteOnNoteOffPairToFile: Bad Track Index: %d"), InTrackIndex);
+			return;
+		}
+
 		FMidiTrack* CurrentTrack = InFile->GetTrack(InTrackIndex);
-		//tick number at the start of every bar for adding Note On/Note Off midi events
-		int32 CurrentBarTick = InFile->GetSongMaps()->GetBarMap().BarIncludingCountInToTick(InBarIndex);
-		int32 DefaultTicksPerQuarter = 960;
 		//Note On / Note Off pair
-		FMidiEvent CurrentNoteOnEvent(CurrentBarTick,
+		FMidiEvent CurrentNoteOnEvent(AtTick,
 			FMidiMsg::CreateNoteOn(InChannel, InNoteNumber, InNoteVelocity));
 		//Note Offs are 1 beat away from Note Ons (ticks per quarter note)
-		FMidiEvent CurrentNoteOffEvent(CurrentBarTick + DefaultTicksPerQuarter,
+		FMidiEvent CurrentNoteOffEvent(AtTick + DurationTicks,
 			FMidiMsg::CreateNoteOff(InChannel, InNoteNumber));
 		CurrentTrack->AddEvent(CurrentNoteOnEvent);
 		CurrentTrack->AddEvent(CurrentNoteOffEvent);
 	}
 
-	void AddCCEventsToFile(UMidiFile* InFile, uint8 InControllerID, uint8 InControlValue, int32 InTrackIndex, int32 InChannel, float InBarIndex)
+	void AddCCEventToFile(UMidiFile* InFile, uint8 InControllerID, uint8 InControlValue, int32 InTrackIndex, int32 InChannel, int32 AtTick)
 	{
+		if (InTrackIndex >= InFile->GetTracks().Num() || InTrackIndex <= 0)
+		{
+			UE_LOG(LogMIDI, Error, TEXT("AddCCEventToFile: Bad Track Index: %d"), InTrackIndex);
+			return;
+		}
+
 		FMidiTrack* CurrentTrack = InFile->GetTrack(InTrackIndex);
-		//tick number at the start of every bar for adding CC events
-		int32 CurrentBarTick = InFile->GetSongMaps()->GetBarMap().BarIncludingCountInToTick(InBarIndex);
-		//Add Control Events to each channels/tracks
-		FMidiEvent CurrentCCEvent(CurrentBarTick,
-			FMidiMsg(Midi::Constants::GControl + InChannel, InControllerID, InControlValue)
-		);
+		FMidiEvent CurrentCCEvent(AtTick,
+			FMidiMsg(Midi::Constants::GControl | InChannel, InControllerID, InControlValue));
 		CurrentTrack->AddEvent(CurrentCCEvent);
 	}
 
-	void AddTextEventsToFile(UMidiFile* InFile, FString InText, int32 InTrackIndex, float InBarIndex)
+	void AddTextEventToFile(UMidiFile* InFile, FString InText, int32 InTrackIndex, int32 AtTick)
 	{
+		if (InTrackIndex >= InFile->GetTracks().Num() || InTrackIndex <= 0)
+		{
+			UE_LOG(LogMIDI, Error, TEXT("AddTextEventToFile: Bad Track Index: %d"), InTrackIndex);
+			return;
+		}
+
 		FMidiTrack* CurrentTrack = InFile->GetTrack(InTrackIndex);
-		//tick number at the start of every bar for adding text events
-		int32 CurrentBarTick = InFile->GetSongMaps()->GetBarMap().BarIncludingCountInToTick(InBarIndex);
-		//text events
 		uint16 TextIndex = CurrentTrack->AddText(InText);
-		FMidiEvent CurrentTextEvent(CurrentBarTick, FMidiMsg::CreateText(TextIndex, Midi::Constants::GMeta_Text));
+		FMidiEvent CurrentTextEvent(AtTick, FMidiMsg::CreateText(TextIndex, Midi::Constants::GMeta_Text));
 		CurrentTrack->AddEvent(CurrentTextEvent);
 	}
 
-	void AddPitchEventsToFile(UMidiFile* InFile, uint8 InPitchValueLSB, uint8 InPitchValueMSB, int32 InTrackIndex, int32 InChannel, float InBarIndex)
+	void AddPitchEventToFile(UMidiFile* InFile, uint8 InPitchValueLSB, uint8 InPitchValueMSB, int32 InTrackIndex, int32 InChannel, int32 AtTick)
 	{
+		if (InTrackIndex >= InFile->GetTracks().Num() || InTrackIndex <= 0)
+		{
+			UE_LOG(LogMIDI, Error, TEXT("AddPitchEventToFile: Bad Track Index: %d"), InTrackIndex);
+			return;
+		}
+
 		FMidiTrack* CurrentTrack = InFile->GetTrack(InTrackIndex);
-		//tick number at the start of every bar for adding text events
-		int32 CurrentBarTick = InFile->GetSongMaps()->GetBarMap().BarIncludingCountInToTick(InBarIndex);
-		//Add pitch bend events to channels/tracks
-		FMidiEvent CurrentPitchEvent(CurrentBarTick,
-			FMidiMsg(Midi::Constants::GPitch + InChannel, InPitchValueLSB, InPitchValueMSB)
-		);
+		FMidiEvent CurrentPitchEvent(AtTick,
+			FMidiMsg(Midi::Constants::GPitch | InChannel, InPitchValueLSB, InPitchValueMSB));
 		CurrentTrack->AddEvent(CurrentPitchEvent);
 	}
 
-	void AddPolyPresEventsToFile(UMidiFile* InFile, uint8 InNoteNumber, uint8 InPolyPresValue, int32 InTrackIndex, int32 InChannel, float InBarIndex)
+	void AddPolyPresEventToFile(UMidiFile* InFile, uint8 InNoteNumber, uint8 InPolyPresValue, int32 InTrackIndex, int32 InChannel, int32 AtTick)
 	{
+		if (InTrackIndex >= InFile->GetTracks().Num() || InTrackIndex <= 0)
+		{
+			UE_LOG(LogMIDI, Error, TEXT("AddPolyPresEventToFile: Bad Track Index: %d"), InTrackIndex);
+			return;
+		}
+
 		FMidiTrack* CurrentTrack = InFile->GetTrack(InTrackIndex);
-		//tick number at the start of every bar for adding text events
-		//TODO: BarIndex is a whole number but 
-		int32 CurrentBarTick = InFile->GetSongMaps()->GetBarMap().BarIncludingCountInToTick(InBarIndex);
-		//Add Poly Pres events to channels/tracks
-		FMidiEvent CurrentPolyPresEvent(CurrentBarTick,
-			FMidiMsg(Midi::Constants::GPolyPres + InChannel, InNoteNumber, InPolyPresValue)
-		);
+		FMidiEvent CurrentPolyPresEvent(AtTick,
+			FMidiMsg(Midi::Constants::GPolyPres | InChannel, InNoteNumber, InPolyPresValue));
 		CurrentTrack->AddEvent(CurrentPolyPresEvent);
 	}
 }
diff --git a/Engine/Plugins/Runtime/Harmonix/Source/HarmonixMidiTests/Private/MidiTestUtility.h b/Engine/Plugins/Runtime/Harmonix/Source/HarmonixMidiTests/Private/MidiTestUtility.h
--- a/Engine/Plugins/Runtime/Harmonix/Source/HarmonixMidiTests/Private/MidiTestUtility.h
+++ b/Engine/Plugins/Runtime/Harmonix/Source/HarmonixMidiTests/Private/MidiTestUtility.h
@@ -12,37 +12,37 @@ namespace Harmonix::Testing::Utility::MidiTestUtility
 	 * input file length in bars (fractional), input number of midi channels, input number of midi tracks, input time signature, and input tempo (bpm)
 	 * 1 text event is added at the end of the file according to input file length to mark the last event tick
 	 */
-	UMidiFile* BuildMidiFile(float FileLengthBars, int32 NumChannels, int32 NumTracks,int32 InTimeSigNum, int32 InTimeSigDenom, int32 InTempo);
+	UMidiFile* CreateAndInitializaMidiFile(float FileLengthBars, int32 NumTracksIncludingConductor, int32 InTimeSigNum, int32 InTimeSigDenom, int32 InTempo, bool PutTextEventOnLastTick = false);
 	
 	/**
 	 * Add 1 Note On/Note Off pair to the input midi file given the input Midi note number, note velocity,
 	 * track index, channel, and fractional bar position
 	 * 
 	 */
-	void AddNoteOnNoteOffPairsToFile(UMidiFile* InFile, int32 InNoteNumber, int32 InNoteVelocity,int32 InTrackIndex, int32 InChannel,float InBarIndex);
+	void AddNoteOnNoteOffPairToFile(UMidiFile* InFile, int32 InNoteNumber, int32 InNoteVelocity,int32 InTrackIndex, int32 InChannel, int32 AtTick, int32 DurationTicks);
 
 	/**
 	* Add 1 CC Event to the input midi file given the input Midi Controller ID, control value,
 	* track index, channel, and fractional bar position
 	*/
-	void AddCCEventsToFile(UMidiFile* InFile, uint8 InControllerID, uint8 InControlValue,int32 InTrackIndex, int32 InChannel, float InBarIndex);
+	void AddCCEventToFile(UMidiFile* InFile, uint8 InControllerID, uint8 InControlValue,int32 InTrackIndex, int32 InChannel, int32 AtTick);
 
 	/**
 	* Add 1 CC Event to the input midi file given the input text,
 	* track index, and fractional bar position
 	*/
-	void AddTextEventsToFile(UMidiFile* InFile, FString InText, int32 InTrackIndex, float InBarIndex);
+	void AddTextEventToFile(UMidiFile* InFile, FString InText, int32 InTrackIndex, int32 AtTick);
 
 	/**
 	* Add 1 Pitch Bend Event to the input midi file given the input Pitch Bend value (LSB), Pitch Bend value (MSB),
 	* track index, channel, and fractional bar position
 	*/
-	void AddPitchEventsToFile(UMidiFile* InFile, uint8 InPitchValueLSB, uint8 InPitchValueMSB, int32 InTrackIndex, int32 InChannel,float InBarIndex);
+	void AddPitchEventToFile(UMidiFile* InFile, uint8 InPitchValueLSB, uint8 InPitchValueMSB, int32 InTrackIndex, int32 InChannel, int32 AtTick);
 
 	/**
 	* Add 1 Pitch Bend Event to the input midi file given the input Midi note number, polyphonic pressure value,
 	* track index, channel, and fractional bar position
 	*/
-	void AddPolyPresEventsToFile(UMidiFile* InFile, uint8 InNoteNumber, uint8 InPolyPresValue, int32 InTrackIndex, int32 InChannel, float InBarIndex);
+	void AddPolyPresEventToFile(UMidiFile* InFile, uint8 InNoteNumber, uint8 InPolyPresValue, int32 InTrackIndex, int32 InChannel, int32 AtTick);
 
 }
diff --git a/Engine/Plugins/Runtime/Harmonix/Source/HarmonixMidiTests/Private/SongMapsTests.cpp b/Engine/Plugins/Runtime/Harmonix/Source/HarmonixMidiTests/Private/SongMapsTests.cpp
new file mode 100644
--- /dev/null
+++ b/Engine/Plugins/Runtime/Harmonix/Source/HarmonixMidiTests/Private/SongMapsTests.cpp
@@ -0,0 +1,332 @@
+// Copyright Epic Games, Inc. All Rights Reserved.
+
+#include "Logging/LogMacros.h"
+#include "Misc/AutomationTest.h"
+#include "HarmonixMidi/MidiFile.h"
+
+#if WITH_DEV_AUTOMATION_TESTS
+
+DEFINE_LOG_CATEGORY_STATIC(LogSongMapsTest, Log, All);
+
+namespace HarmonixMidiTests::SongMaps
+{
+
+	static constexpr bool GLogQualtizationDetails = true;
+
+	void LogQuantizationDetails(const FBarMap& BarMap, int32 OriginalTick, int32 QuantizedTick, EMidiClockSubdivisionQuantization ResultDivision)
+	{
+		int32 BarIndex = 0;
+		int32 BeatInBar = 0;
+		int32 TickIndexInBeat = 0;
+		BarMap.TickToBarBeatTickIncludingCountIn(QuantizedTick, BarIndex, BeatInBar, TickIndexInBeat);
+		UE_LOG(LogSongMapsTest, Log, TEXT("Tick -> %d: Quantized %s to %d, Division = %s -> %d | %d | %d"), OriginalTick, (QuantizedTick < OriginalTick ? TEXT("^") : TEXT("v")), QuantizedTick, *StaticEnum<EQuartzCommandQuantization>()->GetDisplayNameTextByIndex((int32)ResultDivision).ToString(), BarIndex, BeatInBar, TickIndexInBeat);
+	}
+
+	TSharedPtr<FMidiFileData> BuildMidiWithOneTimeSigature(int32 Numerator, int32 Denominator)
+	{
+		TSharedPtr<FMidiFileData> MidiData = MakeShared<FMidiFileData>();
+		check(MidiData);
+
+		// make 97bpm, 4/4  tempo map...
+		FTempoMap& TempoMap = MidiData->SongMaps.GetTempoMap();
+		TempoMap.Empty();
+		FBarMap& BarMap = MidiData->SongMaps.GetBarMap();
+		BarMap.Empty();
+		MidiData->Tracks.Empty();
+
+		MidiData->Tracks.Add(FMidiTrack(TEXT("conductor")));
+		MidiData->Tracks[0].AddEvent(FMidiEvent(0, FMidiMsg(Numerator, Denominator)));
+		BarMap.AddTimeSignatureAtBarIncludingCountIn(0, Numerator, Denominator);
+		const int32 MidiTempo = Harmonix::Midi::Constants::BPMToMidiTempo(97.0f);
+		MidiData->Tracks[0].AddEvent(FMidiEvent(0, FMidiMsg(MidiTempo)));
+		TempoMap.AddTempoInfoPoint(MidiTempo, 0);
+		MidiData->Tracks[0].Sort();
+		MidiData->ConformToLength(std::numeric_limits<int32>::max());
+
+		return MidiData;
+	}
+
+
+	IMPLEMENT_SIMPLE_AUTOMATION_TEST(
+		FTestSongMapsFourFour,
+		"Harmonix.Midi.SongMaps.QuantizeToAny.4/4",
+		EAutomationTestFlags::EditorContext | EAutomationTestFlags::EngineFilter)
+		bool FTestSongMapsFourFour::RunTest(const FString&)
+	{
+		const TSharedPtr<FMidiFileData> MidiData = BuildMidiWithOneTimeSigature(4,4);
+		const FSongMaps& SongMaps = MidiData->SongMaps;
+		for (int32 i = 0; i < SongMaps.GetTicksPerQuarterNote() * 7; i+=10)
+		{
+			EMidiClockSubdivisionQuantization ResultDivision = EMidiClockSubdivisionQuantization::Bar;
+			int32 QuantizedTick = SongMaps.QuantizeTickToAnyNearestSubdivision(i, EMidiFileQuantizeDirection::Nearest, ResultDivision);
+			UTEST_TRUE(TEXT("Got good quantization."), QuantizedTick % SongMaps.SubdivisionToMidiTicks(ResultDivision, QuantizedTick) == 0);
+			if (GLogQualtizationDetails)
+			{
+				LogQuantizationDetails(SongMaps.GetBarMap(), i, QuantizedTick, ResultDivision);
+			}
+			QuantizedTick = SongMaps.QuantizeTickToAnyNearestSubdivision(i, EMidiFileQuantizeDirection::Down, ResultDivision);
+			UTEST_TRUE(TEXT("Got good quantization."), QuantizedTick % SongMaps.SubdivisionToMidiTicks(ResultDivision, QuantizedTick) == 0);
+			if (GLogQualtizationDetails)
+			{
+				LogQuantizationDetails(SongMaps.GetBarMap(), i, QuantizedTick, ResultDivision);
+			}
+			QuantizedTick = SongMaps.QuantizeTickToAnyNearestSubdivision(i, EMidiFileQuantizeDirection::Up, ResultDivision);
+			UTEST_TRUE(TEXT("Got good quantization."), QuantizedTick % SongMaps.SubdivisionToMidiTicks(ResultDivision, QuantizedTick) == 0);
+			if (GLogQualtizationDetails)
+			{
+				LogQuantizationDetails(SongMaps.GetBarMap(), i, QuantizedTick, ResultDivision);
+			}
+		}
+		return true;
+	}
+
+	IMPLEMENT_SIMPLE_AUTOMATION_TEST(
+		FTestSongMapsFiveEight,
+		"Harmonix.Midi.SongMaps.QuantizeToAny.5/8",
+		EAutomationTestFlags::EditorContext | EAutomationTestFlags::EngineFilter)
+		bool FTestSongMapsFiveEight::RunTest(const FString&)
+	{
+		const TSharedPtr<FMidiFileData> MidiData = BuildMidiWithOneTimeSigature(5, 8);
+		const FSongMaps& SongMaps = MidiData->SongMaps;
+		for (int32 i = 0; i < SongMaps.GetTicksPerQuarterNote() * 7; i += 10)
+		{
+			EMidiClockSubdivisionQuantization ResultDivision = EMidiClockSubdivisionQuantization::Bar;
+			int32 QuantizedTick = SongMaps.QuantizeTickToAnyNearestSubdivision(i, EMidiFileQuantizeDirection::Nearest, ResultDivision);
+			UTEST_TRUE(TEXT("Got good quantization."), QuantizedTick % SongMaps.SubdivisionToMidiTicks(ResultDivision, QuantizedTick) == 0);
+			if (GLogQualtizationDetails)
+			{
+				LogQuantizationDetails(SongMaps.GetBarMap(), i, QuantizedTick, ResultDivision);
+			}
+			QuantizedTick = SongMaps.QuantizeTickToAnyNearestSubdivision(i, EMidiFileQuantizeDirection::Down, ResultDivision);
+			UTEST_TRUE(TEXT("Got good quantization."), QuantizedTick % SongMaps.SubdivisionToMidiTicks(ResultDivision, QuantizedTick) == 0);
+			if (GLogQualtizationDetails)
+			{
+				LogQuantizationDetails(SongMaps.GetBarMap(), i, QuantizedTick, ResultDivision);
+			}
+			QuantizedTick = SongMaps.QuantizeTickToAnyNearestSubdivision(i, EMidiFileQuantizeDirection::Up, ResultDivision);
+			UTEST_TRUE(TEXT("Got good quantization."), QuantizedTick % SongMaps.SubdivisionToMidiTicks(ResultDivision, QuantizedTick) == 0);
+			if (GLogQualtizationDetails)
+			{
+				LogQuantizationDetails(SongMaps.GetBarMap(), i, QuantizedTick, ResultDivision);
+			}
+		}
+		return true;
+	}
+
+	IMPLEMENT_SIMPLE_AUTOMATION_TEST(
+		FTestSongMapsThreeTwo,
+		"Harmonix.Midi.SongMaps.QuantizeToAny.3/2",
+		EAutomationTestFlags::EditorContext | EAutomationTestFlags::EngineFilter)
+		bool FTestSongMapsThreeTwo::RunTest(const FString&)
+	{
+		const TSharedPtr<FMidiFileData> MidiData = BuildMidiWithOneTimeSigature(3, 2);
+		const FSongMaps& SongMaps = MidiData->SongMaps;
+		for (int32 i = 0; i < SongMaps.GetTicksPerQuarterNote() * 7; i += 10)
+		{
+			EMidiClockSubdivisionQuantization ResultDivision = EMidiClockSubdivisionQuantization::Bar;
+			int32 QuantizedTick = SongMaps.QuantizeTickToAnyNearestSubdivision(i, EMidiFileQuantizeDirection::Nearest, ResultDivision);
+			UTEST_TRUE(TEXT("Got good quantization."), QuantizedTick % SongMaps.SubdivisionToMidiTicks(ResultDivision, QuantizedTick) == 0);
+			if (GLogQualtizationDetails)
+			{
+				LogQuantizationDetails(SongMaps.GetBarMap(), i, QuantizedTick, ResultDivision);
+			}
+			QuantizedTick = SongMaps.QuantizeTickToAnyNearestSubdivision(i, EMidiFileQuantizeDirection::Down, ResultDivision);
+			UTEST_TRUE(TEXT("Got good quantization."), QuantizedTick % SongMaps.SubdivisionToMidiTicks(ResultDivision, QuantizedTick) == 0);
+			if (GLogQualtizationDetails)
+			{
+				LogQuantizationDetails(SongMaps.GetBarMap(), i, QuantizedTick, ResultDivision);
+			}
+			QuantizedTick = SongMaps.QuantizeTickToAnyNearestSubdivision(i, EMidiFileQuantizeDirection::Up, ResultDivision);
+			UTEST_TRUE(TEXT("Got good quantization."), QuantizedTick % SongMaps.SubdivisionToMidiTicks(ResultDivision, QuantizedTick) == 0);
+			if (GLogQualtizationDetails)
+			{
+				LogQuantizationDetails(SongMaps.GetBarMap(), i, QuantizedTick, ResultDivision);
+			}
+		}
+		return true;
+	};
+
+	IMPLEMENT_SIMPLE_AUTOMATION_TEST(
+		FTestSongMapsMixedTimeSig,
+		"Harmonix.Midi.SongMaps.QuantizeToAny.MixedTimeSignatures",
+		EAutomationTestFlags::EditorContext | EAutomationTestFlags::EngineFilter)
+		bool FTestSongMapsMixedTimeSig::RunTest(const FString&)
+	{
+		TSharedPtr<FMidiFileData> MidiData = BuildMidiWithOneTimeSigature(4, 4);
+		FBarMap& BarMap = MidiData->SongMaps.GetBarMap();
+		int32 TickAfterFirstBar = MidiData->TicksPerQuarterNote * 4;
+		MidiData->Tracks[0].AddEvent(FMidiEvent(TickAfterFirstBar, FMidiMsg(5, 8)));
+		BarMap.AddTimeSignatureAtBarIncludingCountIn(1, 5, 8);
+		int32 TickAfterSecondBar = MidiData->TicksPerQuarterNote / 2 * 5 + TickAfterFirstBar;
+		MidiData->Tracks[0].AddEvent(FMidiEvent(TickAfterSecondBar, FMidiMsg(3, 2)));
+		BarMap.AddTimeSignatureAtBarIncludingCountIn(2, 3, 2);
+
+		MidiData->Tracks[0].Sort();
+
+		const FSongMaps& SongMaps = MidiData->SongMaps;
+
+		// Check First Bar...
+		for (int32 i = 0; i < TickAfterFirstBar; i += 10)
+		{
+			EMidiClockSubdivisionQuantization ResultDivision = EMidiClockSubdivisionQuantization::Bar;
+			int32 QuantizedTick = SongMaps.QuantizeTickToAnyNearestSubdivision(i, EMidiFileQuantizeDirection::Nearest, ResultDivision);
+			UTEST_TRUE(TEXT("Got good quantization."), QuantizedTick % SongMaps.SubdivisionToMidiTicks(ResultDivision, i) == 0);
+			if (GLogQualtizationDetails)
+			{
+				LogQuantizationDetails(SongMaps.GetBarMap(), i, QuantizedTick, ResultDivision);
+			}
+			QuantizedTick = SongMaps.QuantizeTickToAnyNearestSubdivision(i, EMidiFileQuantizeDirection::Down, ResultDivision);
+			UTEST_TRUE(TEXT("Got good quantization."), QuantizedTick % SongMaps.SubdivisionToMidiTicks(ResultDivision, i) == 0);
+			if (GLogQualtizationDetails)
+			{
+				LogQuantizationDetails(SongMaps.GetBarMap(), i, QuantizedTick, ResultDivision);
+			}
+			QuantizedTick = SongMaps.QuantizeTickToAnyNearestSubdivision(i, EMidiFileQuantizeDirection::Up, ResultDivision);
+			UTEST_TRUE(TEXT("Got good quantization."), QuantizedTick % SongMaps.SubdivisionToMidiTicks(ResultDivision, i) == 0);
+			if (GLogQualtizationDetails)
+			{
+				LogQuantizationDetails(SongMaps.GetBarMap(), i, QuantizedTick, ResultDivision);
+			}
+		}
+		// Check Second Bar...
+		for (int32 i = TickAfterFirstBar; i < TickAfterSecondBar; i += 10)
+		{
+			EMidiClockSubdivisionQuantization ResultDivision = EMidiClockSubdivisionQuantization::Bar;
+			int32 QuantizedTick = SongMaps.QuantizeTickToAnyNearestSubdivision(i, EMidiFileQuantizeDirection::Nearest, ResultDivision);
+			UTEST_TRUE(TEXT("Got good quantization."), (QuantizedTick - TickAfterFirstBar) % SongMaps.SubdivisionToMidiTicks(ResultDivision, i) == 0);
+			if (GLogQualtizationDetails)
+			{
+				LogQuantizationDetails(SongMaps.GetBarMap(), i, QuantizedTick, ResultDivision);
+			}
+			QuantizedTick = SongMaps.QuantizeTickToAnyNearestSubdivision(i, EMidiFileQuantizeDirection::Down, ResultDivision);
+			UTEST_TRUE(TEXT("Got good quantization."), (QuantizedTick - TickAfterFirstBar) % SongMaps.SubdivisionToMidiTicks(ResultDivision, i) == 0);
+			if (GLogQualtizationDetails)
+			{
+				LogQuantizationDetails(SongMaps.GetBarMap(), i, QuantizedTick, ResultDivision);
+			}
+			QuantizedTick = SongMaps.QuantizeTickToAnyNearestSubdivision(i, EMidiFileQuantizeDirection::Up, ResultDivision);
+			UTEST_TRUE(TEXT("Got good quantization."), (QuantizedTick - TickAfterFirstBar) % SongMaps.SubdivisionToMidiTicks(ResultDivision, i) == 0);
+			if (GLogQualtizationDetails)
+			{
+				LogQuantizationDetails(SongMaps.GetBarMap(), i, QuantizedTick, ResultDivision);
+			}
+		}
+		// Check Third Bar...
+		for (int32 i = TickAfterSecondBar; i < (TickAfterSecondBar + MidiData->TicksPerQuarterNote * 2 * 3); i += 10)
+		{
+			EMidiClockSubdivisionQuantization ResultDivision = EMidiClockSubdivisionQuantization::Bar;
+			int32 QuantizedTick = SongMaps.QuantizeTickToAnyNearestSubdivision(i, EMidiFileQuantizeDirection::Nearest, ResultDivision);
+			UTEST_TRUE(TEXT("Got good quantization."), (QuantizedTick-TickAfterSecondBar) % SongMaps.SubdivisionToMidiTicks(ResultDivision, i) == 0);
+			if (GLogQualtizationDetails)
+			{
+				LogQuantizationDetails(SongMaps.GetBarMap(), i, QuantizedTick, ResultDivision);
+			}
+			QuantizedTick = SongMaps.QuantizeTickToAnyNearestSubdivision(i, EMidiFileQuantizeDirection::Down, ResultDivision);
+			UTEST_TRUE(TEXT("Got good quantization."), (QuantizedTick - TickAfterSecondBar) % SongMaps.SubdivisionToMidiTicks(ResultDivision, i) == 0);
+			if (GLogQualtizationDetails)
+			{
+				LogQuantizationDetails(SongMaps.GetBarMap(), i, QuantizedTick, ResultDivision);
+			}
+			QuantizedTick = SongMaps.QuantizeTickToAnyNearestSubdivision(i, EMidiFileQuantizeDirection::Up, ResultDivision);
+			UTEST_TRUE(TEXT("Got good quantization."), (QuantizedTick - TickAfterSecondBar) % SongMaps.SubdivisionToMidiTicks(ResultDivision, i) == 0);
+			if (GLogQualtizationDetails)
+			{
+				LogQuantizationDetails(SongMaps.GetBarMap(), i, QuantizedTick, ResultDivision);
+			}
+		}
+
+		return true;
+	}
+
+	IMPLEMENT_SIMPLE_AUTOMATION_TEST(
+		FTestSongMapsLengthIsSubdivision,
+		"Harmonix.Midi.SongMaps.LengthIsPerfectSubdivision",
+		EAutomationTestFlags::EditorContext | EAutomationTestFlags::EngineFilter)
+		bool FTestSongMapsLengthIsSubdivision::RunTest(const FString&)
+	{
+		TSharedPtr<FMidiFileData> MidiData = BuildMidiWithOneTimeSigature(4, 4);
+		int32 TickOnDownbeatOfBar2 = MidiData->TicksPerQuarterNote * 4;
+		MidiData->Tracks.Add(FMidiTrack("note-data"));
+		{
+			FMidiTrack& NoteTrack = MidiData->Tracks[1];
+			NoteTrack.AddEvent(FMidiEvent(0, FMidiMsg::CreateNoteOn(0, 60, 127)));
+			NoteTrack.AddEvent(FMidiEvent(TickOnDownbeatOfBar2, FMidiMsg::CreateNoteOff(0, 60)));
+			NoteTrack.Sort();
+			MidiData->ScanTracksForSongLengthChange();
+		}
+
+		UTEST_FALSE("MIDI file is NOT a perfect subdivision", MidiData->LengthIsAPerfectSubdivision());
+
+		EMidiClockSubdivisionQuantization QuantizedDivision = EMidiClockSubdivisionQuantization::None;
+		int32 TargetFixTick = MidiData->SongMaps.QuantizeTickToAnyNearestSubdivision(MidiData->SongMaps.GetSongLengthData().LengthTicks, EMidiFileQuantizeDirection::Nearest, QuantizedDivision);
+	
+		UTEST_TRUE("Quantized Length Is Correct", TargetFixTick == TickOnDownbeatOfBar2);
+		UTEST_TRUE("Quantized Subdivision Is Correct", QuantizedDivision == EMidiClockSubdivisionQuantization::Bar);
+
+		MidiData->Tracks[1] = FMidiTrack("note-data");
+		{
+			FMidiTrack& NoteTrack = MidiData->Tracks[1];
+			NoteTrack.AddEvent(FMidiEvent(0, FMidiMsg::CreateNoteOn(0, 60, 127)));
+			NoteTrack.AddEvent(FMidiEvent(TickOnDownbeatOfBar2 - 1, FMidiMsg::CreateNoteOff(0, 60)));
+			NoteTrack.Sort();
+			MidiData->ScanTracksForSongLengthChange();
+		}
+
+		UTEST_TRUE("MIDI file IS a perfect subdivision", MidiData->LengthIsAPerfectSubdivision());
+
+		MidiData->Tracks[1] = FMidiTrack("note-data");
+		{
+			FMidiTrack& NoteTrack = MidiData->Tracks[1];
+			NoteTrack.AddEvent(FMidiEvent(0, FMidiMsg::CreateNoteOn(0, 60, 127)));
+			NoteTrack.AddEvent(FMidiEvent(TickOnDownbeatOfBar2 - MidiData->SongMaps.SubdivisionToMidiTicks(EMidiClockSubdivisionQuantization::ThirtySecondNote, 0), FMidiMsg::CreateNoteOff(0, 60)));
+			NoteTrack.Sort();
+			MidiData->ScanTracksForSongLengthChange();
+		}
+
+		UTEST_FALSE("MIDI file is NOT a perfect subdivision", MidiData->LengthIsAPerfectSubdivision());
+
+		MidiData->Tracks[1] = FMidiTrack("note-data");
+		{
+			FMidiTrack& NoteTrack = MidiData->Tracks[1];
+			NoteTrack.AddEvent(FMidiEvent(0, FMidiMsg::CreateNoteOn(0, 60, 127)));
+			NoteTrack.AddEvent(FMidiEvent(TickOnDownbeatOfBar2 - 1 - MidiData->SongMaps.SubdivisionToMidiTicks(EMidiClockSubdivisionQuantization::ThirtySecondNote, 0), FMidiMsg::CreateNoteOff(0, 60)));
+			NoteTrack.Sort();
+			MidiData->ScanTracksForSongLengthChange();
+		}
+
+		UTEST_TRUE("MIDI file IS a perfect subdivision", MidiData->LengthIsAPerfectSubdivision());
+
+		MidiData = BuildMidiWithOneTimeSigature(7, 8);
+		TickOnDownbeatOfBar2 = MidiData->TicksPerQuarterNote / 2 * 7;
+		MidiData->Tracks.Add(FMidiTrack("note-data"));
+		{
+			FMidiTrack& NoteTrack = MidiData->Tracks[1];
+			NoteTrack.AddEvent(FMidiEvent(0, FMidiMsg::CreateNoteOn(0, 60, 127)));
+			NoteTrack.AddEvent(FMidiEvent(TickOnDownbeatOfBar2, FMidiMsg::CreateNoteOff(0, 60)));
+			NoteTrack.Sort();
+			MidiData->ScanTracksForSongLengthChange();
+		}
+
+		UTEST_FALSE("MIDI file is NOT a perfect subdivision", MidiData->LengthIsAPerfectSubdivision());
+
+		QuantizedDivision = EMidiClockSubdivisionQuantization::None;
+		TargetFixTick = MidiData->SongMaps.QuantizeTickToAnyNearestSubdivision(MidiData->SongMaps.GetSongLengthData().LengthTicks, EMidiFileQuantizeDirection::Nearest, QuantizedDivision);
+
+		UTEST_TRUE("Quantized Length Is Correct", TargetFixTick == TickOnDownbeatOfBar2);
+		UTEST_TRUE("Quantized Subdivision Is Correct", QuantizedDivision == EMidiClockSubdivisionQuantization::Bar);
+
+		MidiData->Tracks[1] = FMidiTrack("note-data");
+		{
+			FMidiTrack& NoteTrack = MidiData->Tracks[1];
+			NoteTrack.AddEvent(FMidiEvent(0, FMidiMsg::CreateNoteOn(0, 60, 127)));
+			NoteTrack.AddEvent(FMidiEvent(TickOnDownbeatOfBar2 - 1, FMidiMsg::CreateNoteOff(0, 60)));
+			NoteTrack.Sort();
+			MidiData->ScanTracksForSongLengthChange();
+		}
+
+		UTEST_TRUE("MIDI file IS a perfect subdivision", MidiData->LengthIsAPerfectSubdivision());
+
+		return true;
+	}
+}
+#endif
\ No newline at end of file
