[Harmonix] Node categorization pass (public nodes)

diff --git a/Engine/Plugins/Runtime/Harmonix/Source/HarmonixMetasound/Private/HarmonixMetasound/Nodes/MidiClockOffsetNode.cpp b/Engine/Plugins/Runtime/Harmonix/Source/HarmonixMetasound/Private/HarmonixMetasound/Nodes/MidiClockOffsetNode.cpp
--- a/Engine/Plugins/Runtime/Harmonix/Source/HarmonixMetasound/Private/HarmonixMetasound/Nodes/MidiClockOffsetNode.cpp
+++ b/Engine/Plugins/Runtime/Harmonix/Source/HarmonixMetasound/Private/HarmonixMetasound/Nodes/MidiClockOffsetNode.cpp
@@ -63,7 +63,7 @@ namespace HarmonixMetasound::Nodes::MidiClockOffset
 				Info.Author				= PluginAuthor;
 				Info.PromptIfMissing	= PluginNodeMissingPrompt;
 				Info.DefaultInterface	= GetVertexInterface();
-				Info.CategoryHierarchy.Emplace(NodeCategories::Music);
+				Info.CategoryHierarchy = { MetasoundNodeCategories::Harmonix, NodeCategories::Music };
 
 				return Info;
 			};
diff --git a/Engine/Plugins/Runtime/Harmonix/Source/HarmonixMetasound/Private/HarmonixMetasound/Nodes/MusicSeekTargetBuilder.cpp b/Engine/Plugins/Runtime/Harmonix/Source/HarmonixMetasound/Private/HarmonixMetasound/Nodes/MusicSeekTargetBuilder.cpp
--- a/Engine/Plugins/Runtime/Harmonix/Source/HarmonixMetasound/Private/HarmonixMetasound/Nodes/MusicSeekTargetBuilder.cpp
+++ b/Engine/Plugins/Runtime/Harmonix/Source/HarmonixMetasound/Private/HarmonixMetasound/Nodes/MusicSeekTargetBuilder.cpp
@@ -55,7 +55,7 @@ namespace HarmonixMetasound::Nodes::BarBeatToSeekTarget
 				Info.Author           = PluginAuthor;
 				Info.PromptIfMissing  = PluginNodeMissingPrompt;
 				Info.DefaultInterface = GetVertexInterface();
-				Info.CategoryHierarchy.Emplace(NodeCategories::Music);
+				Info.CategoryHierarchy = { MetasoundNodeCategories::Harmonix, NodeCategories::Music };
 				return Info;
 			};
 
