[Harmonix] Use shared/weak pointer for clock in pulse generator, unregister play cursor before ~FMidiPlayCursor does so we don't get warnings

#tests automated tests

diff --git a/Engine/Plugins/Runtime/Harmonix/Source/HarmonixMetasound/Private/HarmonixMetasound/MidiOps/PulseGenerator.cpp b/Engine/Plugins/Runtime/Harmonix/Source/HarmonixMetasound/Private/HarmonixMetasound/MidiOps/PulseGenerator.cpp
--- a/Engine/Plugins/Runtime/Harmonix/Source/HarmonixMetasound/Private/HarmonixMetasound/MidiOps/PulseGenerator.cpp
+++ b/Engine/Plugins/Runtime/Harmonix/Source/HarmonixMetasound/Private/HarmonixMetasound/MidiOps/PulseGenerator.cpp
@@ -9,7 +9,7 @@ namespace Harmonix::Midi::Ops
 		Enabled = bEnable;
 	}
 
-	void FPulseGenerator::SetClock(const HarmonixMetasound::FMidiClock* Clock)
+	void FPulseGenerator::SetClock(const TSharedPtr<const HarmonixMetasound::FMidiClock, ESPMode::NotThreadSafe>& Clock)
 	{
 		Cursor.SetClock(Clock);
 	}
@@ -74,21 +74,27 @@ namespace Harmonix::Midi::Ops
 		SetMessageFilter(EFilterPassFlags::TimeSig);
 	}
 
-	void FPulseGenerator::FCursor::SetClock(const HarmonixMetasound::FMidiClock* NewClock)
+	FPulseGenerator::FCursor::~FCursor()
+	{
+		// NB: If we let ~FMidiPlayCursor do this, we get warnings for potentially bad access
+		SetClock(nullptr);
+	}
+
+	void FPulseGenerator::FCursor::SetClock(const TSharedPtr<const HarmonixMetasound::FMidiClock, ESPMode::NotThreadSafe>& NewClock)
 	{
 		if (Clock == NewClock)
 		{
 			return;
 		}
-		
-		if (nullptr == NewClock)
+
+		// unregister if we already have a clock
+		if (const auto PinnedClock = Clock.Pin())
 		{
-			if (nullptr != Owner)
-			{
-				Owner->UnregisterPlayCursor(this);
-			}
+			PinnedClock->UnregisterPlayCursor(this);
 		}
-		else
+
+		// register the new clock if we were given one
+		if (NewClock.IsValid())
 		{
 			NewClock->RegisterHiResPlayCursor(this);
 		}
@@ -133,17 +139,18 @@ namespace Harmonix::Midi::Ops
 		{
 			return;
 		}
-
-		check(nullptr != Clock);
-		int32 NextPulseTick = Clock->GetBarMap().MusicTimestampToTick(NextPulseTimestamp);
+		
+		check(Clock.IsValid()); // if we're here and we don't have a clock, we have problems
+		const auto PinnedClock = Clock.Pin();
+		int32 NextPulseTick = PinnedClock->GetBarMap().MusicTimestampToTick(NextPulseTimestamp);
 
 		while (CurrentTick >= NextPulseTick)
 		{
-			Push({ Clock->GetCurrentBlockFrameIndex(), NextPulseTick });
+			Push({ PinnedClock->GetCurrentBlockFrameIndex(), NextPulseTick });
 			
 			IncrementTimestampByInterval(NextPulseTimestamp, Interval, CurrentTimeSignature);
 
-			NextPulseTick = Clock->GetBarMap().MusicTimestampToTick(NextPulseTimestamp);
+			NextPulseTick = PinnedClock->GetBarMap().MusicTimestampToTick(NextPulseTimestamp);
 		}
 	}
 
@@ -152,10 +159,11 @@ namespace Harmonix::Midi::Ops
 		CurrentTimeSignature.Numerator = Numerator;
 		CurrentTimeSignature.Denominator = Denominator;
 		
-		check(nullptr != Owner);
+		check(Clock.IsValid()); // if we're here and we don't have a clock, we have problems
+		const auto PinnedClock = Clock.Pin();
 		// Time sig changes will come on the downbeat, and if we change time signature,
 		// we want to reset the pulse, so the next pulse is now plus the offset
-		NextPulseTimestamp = Owner->GetBarMap().TickToMusicTimestamp(Tick);
+		NextPulseTimestamp = PinnedClock->GetBarMap().TickToMusicTimestamp(Tick);
 		IncrementTimestampByOffset(NextPulseTimestamp, Interval, CurrentTimeSignature);
 	}
 }
diff --git a/Engine/Plugins/Runtime/Harmonix/Source/HarmonixMetasound/Private/HarmonixMetasound/Nodes/MidiPulseGeneratorNode.cpp b/Engine/Plugins/Runtime/Harmonix/Source/HarmonixMetasound/Private/HarmonixMetasound/Nodes/MidiPulseGeneratorNode.cpp
--- a/Engine/Plugins/Runtime/Harmonix/Source/HarmonixMetasound/Private/HarmonixMetasound/Nodes/MidiPulseGeneratorNode.cpp
+++ b/Engine/Plugins/Runtime/Harmonix/Source/HarmonixMetasound/Private/HarmonixMetasound/Nodes/MidiPulseGeneratorNode.cpp
@@ -159,7 +159,7 @@ namespace HarmonixMetasound::Nodes::MidiPulseGeneratorNode
 
 		void Reset(const FResetParams&)
 		{
-			PulseGenerator.SetClock(Inputs.Clock.Get());
+			PulseGenerator.SetClock(Inputs.Clock->AsShared());
 			ApplyParameters();
 		}
 
@@ -175,7 +175,7 @@ namespace HarmonixMetasound::Nodes::MidiPulseGeneratorNode
 			InVertexData.BindReadVertex(Inputs::MidiNoteNumberName, Inputs.NoteNumber);
 			InVertexData.BindReadVertex(Inputs::MidiVelocityName, Inputs.Velocity);
 
-			PulseGenerator.SetClock(Inputs.Clock.Get());
+			PulseGenerator.SetClock(Inputs.Clock->AsShared());
 			ApplyParameters();
 		}
 
diff --git a/Engine/Plugins/Runtime/Harmonix/Source/HarmonixMetasound/Public/HarmonixMetasound/MidiOps/PulseGenerator.h b/Engine/Plugins/Runtime/Harmonix/Source/HarmonixMetasound/Public/HarmonixMetasound/MidiOps/PulseGenerator.h
--- a/Engine/Plugins/Runtime/Harmonix/Source/HarmonixMetasound/Public/HarmonixMetasound/MidiOps/PulseGenerator.h
+++ b/Engine/Plugins/Runtime/Harmonix/Source/HarmonixMetasound/Public/HarmonixMetasound/MidiOps/PulseGenerator.h
@@ -19,7 +19,7 @@ namespace Harmonix::Midi::Ops
 	public:
 		void Enable(bool bEnable);
 		
-		void SetClock(const HarmonixMetasound::FMidiClock* Clock);
+		void SetClock(const TSharedPtr<const HarmonixMetasound::FMidiClock, ESPMode::NotThreadSafe>& Clock);
 
 		Dsp::Parameters::TParameter<uint8> Channel{ 1, 16, 1 };
 
@@ -45,7 +45,9 @@ namespace Harmonix::Midi::Ops
 		public:
 			FCursor();
 
-			void SetClock(const HarmonixMetasound::FMidiClock* NewClock);
+			virtual ~FCursor() override;
+
+			void SetClock(const TSharedPtr<const HarmonixMetasound::FMidiClock, ESPMode::NotThreadSafe>& NewClock);
 
 			void SetInterval(const FMusicTimeInterval& NewInterval);
 
@@ -66,7 +68,7 @@ namespace Harmonix::Midi::Ops
 			
 			virtual void OnTimeSig(int32 TrackIndex, int32 Tick, int32 Numerator, int32 Denominator, bool IsPreroll) override;
 
-			const HarmonixMetasound::FMidiClock* Clock{ nullptr };
+			TWeakPtr<const HarmonixMetasound::FMidiClock, ESPMode::NotThreadSafe> Clock;
 			FMusicTimeInterval Interval{};
 			TSpscQueue<FPulseTime> PulsesSinceLastProcess;
 			int32 QueueSize{0};
diff --git a/Engine/Plugins/Runtime/Harmonix/Source/HarmonixMetasoundTests/Private/MidiOps/PulseGeneratorTests.cpp b/Engine/Plugins/Runtime/Harmonix/Source/HarmonixMetasoundTests/Private/MidiOps/PulseGeneratorTests.cpp
--- a/Engine/Plugins/Runtime/Harmonix/Source/HarmonixMetasoundTests/Private/MidiOps/PulseGeneratorTests.cpp
+++ b/Engine/Plugins/Runtime/Harmonix/Source/HarmonixMetasoundTests/Private/MidiOps/PulseGeneratorTests.cpp
@@ -20,21 +20,21 @@ namespace Harmonix::Midi::Ops::Tests
 		constexpr float Tempo = 123;
 		constexpr uint8 TimeSigNumerator = 4;
 		constexpr uint8 TimeSigDenominator = 4;
-		HarmonixMetasound::FMidiClock Clock{ OperatorSettings };
-		Clock.AttachToMidiResource(HarmonixMetasound::FMidiClock::MakeClockConductorMidiData(Tempo, TimeSigNumerator, TimeSigDenominator));
+		const auto Clock = MakeShared<HarmonixMetasound::FMidiClock, ESPMode::NotThreadSafe>(OperatorSettings);
+		Clock->AttachToMidiResource(HarmonixMetasound::FMidiClock::MakeClockConductorMidiData(Tempo, TimeSigNumerator, TimeSigDenominator));
 		
 		HarmonixMetasound::FMidiStream OutputStream;
 
-		PulseGenerator.SetClock(&Clock);
+		PulseGenerator.SetClock(Clock);
 
 		// Default: a pulse every beat
 		{
 			constexpr int32 NotesUntilWeAreSatisfiedThisWorks = 23;
 			int32 NumNotesReceived = 0;
 
-			Clock.ResetAndStart(0);
+			Clock->ResetAndStart(0);
 
-			FTimeSignature TimeSignature = Clock.GetBarMap().GetTimeSignatureAtTick(0);
+			FTimeSignature TimeSignature = Clock->GetBarMap().GetTimeSignatureAtTick(0);
 			FMusicTimeInterval Interval = PulseGenerator.GetInterval();
 			FMusicTimestamp NextPulse{ 1, 1 };
 			IncrementTimestampByOffset(NextPulse, Interval, TimeSignature);
@@ -46,21 +46,21 @@ namespace Harmonix::Midi::Ops::Tests
 				{
 					IncrementTimestampByInterval(EndTimestamp, Interval, TimeSignature);
 				}
-				EndTick = Clock.GetBarMap().MusicTimestampToTick(EndTimestamp);
+				EndTick = Clock->GetBarMap().MusicTimestampToTick(EndTimestamp);
 			}
 			
-			while (Clock.GetCurrentMidiTick() < EndTick)
+			while (Clock->GetCurrentMidiTick() < EndTick)
 			{
 				// Advance the clock, which will advance the play cursor in the pulse generator
-				Clock.PrepareBlock();
-				Clock.WriteAdvance(0, OperatorSettings.GetNumFramesPerBlock());
+				Clock->PrepareBlock();
+				Clock->WriteAdvance(0, OperatorSettings.GetNumFramesPerBlock());
 
 				// Process, which will pop the next notes
 				OutputStream.PrepareBlock();
 				PulseGenerator.Process(OutputStream);
 				
 				// If this is a block where we should get a pulse, check that we got it
-				if (Clock.GetCurrentMidiTick() >= Clock.GetBarMap().MusicTimestampToTick(NextPulse))
+				if (Clock->GetCurrentMidiTick() >= Clock->GetBarMap().MusicTimestampToTick(NextPulse))
 				{
 					const bool ShouldGetNoteOff = NumNotesReceived > 0;
 
diff --git a/Engine/Plugins/Runtime/Harmonix/Source/HarmonixMetasoundTests/Private/Nodes/MidiPulseGeneratorNodeTests.cpp b/Engine/Plugins/Runtime/Harmonix/Source/HarmonixMetasoundTests/Private/Nodes/MidiPulseGeneratorNodeTests.cpp
--- a/Engine/Plugins/Runtime/Harmonix/Source/HarmonixMetasoundTests/Private/Nodes/MidiPulseGeneratorNodeTests.cpp
+++ b/Engine/Plugins/Runtime/Harmonix/Source/HarmonixMetasoundTests/Private/Nodes/MidiPulseGeneratorNodeTests.cpp
@@ -25,7 +25,7 @@ namespace HarmonixMetasound::Nodes::MidiPulseGeneratorNode::Tests
 
 		Harmonix::Midi::Ops::FPulseGenerator PulseGenerator;
 		FMidiStream PulseGeneratorMidiOutput;
-		PulseGenerator.SetClock(Clock->Get());
+		PulseGenerator.SetClock((*Clock)->AsShared());
 
 		// Render for a bit and expect the same output from both the node and the raw processor
 		constexpr int32 NumBlocks = 1000;
