FMidiPlayCursor::Owner is private instead of protected, so that derived classes are forced to call GetOwner.

#tests PIE

diff --git a/Engine/Plugins/Runtime/Harmonix/Source/HarmonixMetasound/Private/HarmonixMetasound/Nodes/MetronomeNode.cpp b/Engine/Plugins/Runtime/Harmonix/Source/HarmonixMetasound/Private/HarmonixMetasound/Nodes/MetronomeNode.cpp
--- a/Engine/Plugins/Runtime/Harmonix/Source/HarmonixMetasound/Private/HarmonixMetasound/Nodes/MetronomeNode.cpp
+++ b/Engine/Plugins/Runtime/Harmonix/Source/HarmonixMetasound/Private/HarmonixMetasound/Nodes/MetronomeNode.cpp
@@ -294,7 +294,7 @@ namespace HarmonixMetasound
 		int32 NewTick = MidiClockOutPin->GetCurrentMidiTick() + (InTick - PrevTick);
 		float Ms = MidiClockOutPin->GetSongMaps().TickToMs(NewTick);
 		float CurrentSpeed = MetronomeClock.GetSpeedAtBlockSampleFrame(MetronomeClock.GetCurrentBlockFrameIndex());
-		float AdvanceRatio = Owner->GetSongMaps().GetTempoAtTick(PrevTick)
+		float AdvanceRatio = GetOwner()->GetSongMaps().GetTempoAtTick(PrevTick)
 		                   / MidiClockOutPin->GetSongMaps().GetTempoAtTick(MidiClockOutPin->GetCurrentMidiTick());
 		MidiClockOutPin->InformOfCurrentAdvanceRate(CurrentSpeed * AdvanceRatio);
 		MidiClockOutPin->AdvanceHiResToMs(MetronomeClock.GetCurrentBlockFrameIndex(), Ms, true);
diff --git a/Engine/Plugins/Runtime/Harmonix/Source/HarmonixMidi/Private/HarmonixMidi/SmoothedMidiPlayCursor.cpp b/Engine/Plugins/Runtime/Harmonix/Source/HarmonixMidi/Private/HarmonixMidi/SmoothedMidiPlayCursor.cpp
--- a/Engine/Plugins/Runtime/Harmonix/Source/HarmonixMidi/Private/HarmonixMidi/SmoothedMidiPlayCursor.cpp
+++ b/Engine/Plugins/Runtime/Harmonix/Source/HarmonixMidi/Private/HarmonixMidi/SmoothedMidiPlayCursor.cpp
@@ -78,7 +78,8 @@ void FSmoothedMidiPlayCursor::OnLoop(int loopStartTick, int loopEndTick)
 
 void FSmoothedMidiPlayCursor::SyncSmoothingTimer(bool bEnableErrorCorrection)
 {
-	if (!Owner)
+	FMidiPlayCursorMgr* CursorOwner = GetOwner();
+	if (!CursorOwner)
 	{
 		return;
 	}
@@ -88,7 +89,7 @@ void FSmoothedMidiPlayCursor::SyncSmoothingTimer(bool bEnableErrorCorrection)
 		SmoothingTimer.Start();
 	}
 
-	SetSpeed(Owner->GetCurrentAdvanceRate(Tracker->IsLowRes));
+	SetSpeed(CursorOwner->GetCurrentAdvanceRate(Tracker->IsLowRes));
 
 	float RawMs = CurrentMs;
 	double SmoothMs = SmoothingTimer.Ms();
@@ -108,7 +109,7 @@ void FSmoothedMidiPlayCursor::SyncSmoothingTimer(bool bEnableErrorCorrection)
 			SmoothingTimer.Reset(RawMs);
 			SmoothedMs = RawMs;
 			SmoothedMsDelta = 0.0f;
-			SmoothedTick = Owner->GetTempoMap().MsToTick(RawMs);
+			SmoothedTick = CursorOwner->GetTempoMap().MsToTick(RawMs);
 			ErrorTracker.Reset();
 			return;
 		}
@@ -146,8 +147,8 @@ void FSmoothedMidiPlayCursor::SyncSmoothingTimer(bool bEnableErrorCorrection)
 	SmoothedMsDelta = float(SmoothMs - RawMs);
 
 	// we need to see if the smoothed position is a position that we have never played!
-	float LoopStartMs = Owner->GetLoopStartMs(Tracker->IsLowRes);
-	float LoopEndMs = Owner->GetLoopEndMs(Tracker->IsLowRes);
+	float LoopStartMs = CursorOwner->GetLoopStartMs(Tracker->IsLowRes);
+	float LoopEndMs = CursorOwner->GetLoopEndMs(Tracker->IsLowRes);
 	if (LoopStartMs < LoopEndMs)
 	{
 		// possible loop back...
@@ -167,16 +168,17 @@ void FSmoothedMidiPlayCursor::SyncSmoothingTimer(bool bEnableErrorCorrection)
 		}
 	}
 	SmoothedMs = float(SmoothMs);
-	SmoothedTick = Owner->GetTempoMap().MsToTick(SmoothedMs);
+	SmoothedTick = CursorOwner->GetTempoMap().MsToTick(SmoothedMs);
 	LoopedThisPass = false;
 }
 
 void FSmoothedMidiPlayCursor::UpdateSongPosition()
 {
-	if (!Owner)
+	FMidiPlayCursorMgr* CursorOwner = GetOwner();
+	if (!CursorOwner)
 	{
 		return;
 	}
 
-	CurrentSongPos = Owner->CalculateSongPosWithOffsetMs(SmoothedMsDelta, Tracker->IsLowRes);
+	CurrentSongPos = CursorOwner->CalculateSongPosWithOffsetMs(SmoothedMsDelta, Tracker->IsLowRes);
 }
diff --git a/Engine/Plugins/Runtime/Harmonix/Source/HarmonixMidi/Public/HarmonixMidi/MidiPlayCursor.h b/Engine/Plugins/Runtime/Harmonix/Source/HarmonixMidi/Public/HarmonixMidi/MidiPlayCursor.h
--- a/Engine/Plugins/Runtime/Harmonix/Source/HarmonixMidi/Public/HarmonixMidi/MidiPlayCursor.h
+++ b/Engine/Plugins/Runtime/Harmonix/Source/HarmonixMidi/Public/HarmonixMidi/MidiPlayCursor.h
@@ -111,7 +111,9 @@ protected:
 	int32 CurrentTick = 0;  // We've broadcast all events up to and including this tick
 	int32 LoopCount = 0;
 	bool  UnregisterASAP = false;
+private:
 	FMidiPlayCursorMgr* Owner = nullptr;
+protected:
 	FMidiPlayCursorTracker* Tracker = nullptr;
 
 protected:
