Remove Drop down for keyzones from Fusion Patch Settings Detail Panel

Replace with flat list of keyzones, but display the Sound Wave and note/velocity zone inline.

#tests PIE

diff --git a/Engine/Plugins/Runtime/Harmonix/Source/HarmonixDspEditor/Private/Customization/FusionPatchDetailCustomization.cpp b/Engine/Plugins/Runtime/Harmonix/Source/HarmonixDspEditor/Private/Customization/FusionPatchDetailCustomization.cpp
--- a/Engine/Plugins/Runtime/Harmonix/Source/HarmonixDspEditor/Private/Customization/FusionPatchDetailCustomization.cpp
+++ b/Engine/Plugins/Runtime/Harmonix/Source/HarmonixDspEditor/Private/Customization/FusionPatchDetailCustomization.cpp
@@ -3,265 +3,238 @@
 
 #include "DetailLayoutBuilder.h"
 #include "DetailWidgetRow.h"
+#include "IDetailChildrenBuilder.h"
 #include "PropertyCustomizationHelpers.h"
 #include "HarmonixDsp/FusionSampler/Settings/FusionPatchSettings.h"
 #include "HarmonixDsp/FusionSampler/Settings/KeyzoneSettings.h"
 #include "HarmonixDsp/FusionSampler/FusionPatch.h"
-
-#include "Widgets/Input/SComboBox.h"
 #include "Widgets/Text/STextBlock.h"
+#include "Widgets/Layout/SUniformGridPanel.h"
 
-#include "Widgets/SMinMaxSlider.h"
 
-void FFusionPatchDetailCustomization::CustomizeDetails(IDetailLayoutBuilder& DetailLayout)
+class FKeyzoneDetailArrayBuilder : public FDetailArrayBuilder
 {
-	IDetailCategoryBuilder& FusionPatchDataCategory = DetailLayout.EditCategory(TEXT("Fusion Patch Data"));
-	//IDetailCategoryBuilder& FusionPatchDataCategory = DetailLayout.EditCategory(TEXT("Fusion Patch Data"));
-	IDetailCategoryBuilder& FilePathCategory = DetailLayout.EditCategory(TEXT("File Path"));
+public:
+	static const int32 MinDesiredSlotWidth = 80;
 
-	//get a handle of Fusion Patch Data
-	TSharedPtr<IPropertyHandle> FusionPatchDataHandle = DetailLayout.GetProperty("FusionPatchData");
-	check(FusionPatchDataHandle);
-	TSharedPtr<IPropertyHandle> SettingsHandle = FusionPatchDataHandle->GetChildHandle("Settings");
-	check(SettingsHandle);
+	FKeyzoneDetailArrayBuilder(TSharedRef<IPropertyHandle> InBaseProperty, const TArray<FName>& PropertyHeaderNames, bool DisplayHeaderNames = true)
+		: FDetailArrayBuilder(InBaseProperty)
+		, PropertyHeaderNames(PropertyHeaderNames)
+		, DisplayHeaderNames(DisplayHeaderNames)
+	{}
 
-	//get a handle of the keyzones as array
-	TSharedPtr<IPropertyHandleArray> KeyzonesHandle = FusionPatchDataHandle->GetChildHandle("Keyzones")->AsArray();
-	check(KeyzonesHandle.IsValid());
+	virtual void RefreshChildren() override
+	{
+		OptionsPropertyHandles.Reset();
+		TSharedPtr<IPropertyHandle> PropertyHandle = GetPropertyHandle();
+		check(PropertyHandle);
 
-	//get the number of keyzones in the current patch
-	uint32 NumKeyzones;
-	KeyzonesHandle->GetNumElements(NumKeyzones);
+		TSharedPtr<IPropertyHandleArray> ArrayHandle = PropertyHandle->AsArray();
+		check(ArrayHandle);
 
-	//add all keyzones' name (sample path string) to an array for displaying in the dropdown menu
-	AddKeyzonesNameToMenuArray(KeyzonesHandle, static_cast<int32>(NumKeyzones));
-	if (!KeyzonesNameMenu.IsEmpty()) CurrentKeyzoneName = KeyzonesNameMenu[0];
+		uint32 Num;
+		ArrayHandle->GetNumElements(Num);
+		for (uint32 Idx = 0; Idx < Num; ++Idx)
+		{
+			OptionsPropertyHandles.Add(ArrayHandle->GetElement(Idx));
+		}
+		
+		
+		FDetailArrayBuilder::RefreshChildren();
+	}
 
-	//get the current fusion patch being edited
-	TArray<TWeakObjectPtr<UObject>> Objects;
-	DetailLayout.GetObjectsBeingCustomized(Objects);
-	if (Objects.Num() != 1)
+	virtual void GenerateChildContent(IDetailChildrenBuilder& ChildrenBuilder) override
 	{
-		return;
+		if (DisplayHeaderNames)
+		{
+			GenerateChildHeaderContent(ChildrenBuilder);
+		}
+		FDetailArrayBuilder::GenerateChildContent(ChildrenBuilder);
 	}
-	TWeakObjectPtr<UObject> FusionPatch = Objects.Last();
-	TWeakObjectPtr<UFusionPatch> FusionPatchBeingEdited = Cast<UFusionPatch>(FusionPatch);
-	
-	//add the dropdown menu displaying the sample path/name of keyzones to the name content of this row
-	FusionPatchDataCategory.AddCustomRow(FText::FromString("Fusion Patch"))
-		.NameContent()
-		[
-			SNew(STextBlock)
-			.Text(FText::FromString("Selected Keyzone"))
-			.Font(IDetailLayoutBuilder::GetDetailFontBold())
-		]
-		.ValueContent()
-		.MinDesiredWidth(350.f)
-		[
-			//create the dropdown combo box displaying all the keyzones name/sample path for selection
-			SNew(SComboBox<TSharedPtr<FString>>)
-				.OptionsSource(&KeyzonesNameMenu)
-				.OnSelectionChanged_Lambda([this](TSharedPtr<FString> NewKeyzoneName, ESelectInfo::Type)
-				{
-					CurrentKeyzoneName = NewKeyzoneName;
-					CurrentKeyzoneIndex = KeyzonesNameMenu.Find(NewKeyzoneName);
-				})
-				.OnGenerateWidget_Lambda([](TSharedPtr<FString> InKeyzoneName) -> TSharedRef<SWidget>
-				{
-					return SNew(STextBlock).Text(FText::FromString(*InKeyzoneName));
-				})
-				.InitiallySelectedItem(CurrentKeyzoneName)
-				.Content()
-				[
-					SNew(STextBlock).Text_Lambda([this]() -> FText
-					{
-						if (CurrentKeyzoneName.IsValid())
-						{
-							return FText::FromString(*CurrentKeyzoneName);
-						}
 
-						return FText::FromString(TEXT("InvalidComboEntryText"));
-					})
-				]
-		];
-
-	DrawSelectedKeyzoneProperties(FusionPatchDataCategory, KeyzonesHandle, NumKeyzones);
-
-	//add the Settings properties back so they're visible
-	FusionPatchDataCategory.AddProperty(SettingsHandle);
+	void GenerateChildHeaderContent(IDetailChildrenBuilder& ChildrenBuilder)
+	{
+		TSharedPtr<IPropertyHandle> PropertyHandle = GetPropertyHandle();
+		check(PropertyHandle);
 
-	//hide the non-customized versions of these properties
-	DetailLayout.HideProperty(FusionPatchDataHandle);
-}
+		TSharedPtr<IPropertyHandleArray> ArrayHandle = PropertyHandle->AsArray();
+		check(ArrayHandle);
 
+		TSharedPtr<SUniformGridPanel> GridPanel;
+		SAssignNew(GridPanel, SUniformGridPanel).MinDesiredSlotWidth(MinDesiredSlotWidth);
 
+		uint32 NumElements;
+		ArrayHandle->GetNumElements(NumElements);
+		if (NumElements > 0)
+		{
+			TSharedRef<IPropertyHandle> Element = ArrayHandle->GetElement(0);
 
-TSharedRef<SWidget> FFusionPatchDetailCustomization::CreateMinMaxSliderWidget(TSharedPtr<IPropertyHandle> MinValuePropertyHandle, TSharedPtr<IPropertyHandle> MaxValuePropertyHandle)
-{
-	int8 CurrentMinValue = 0.0f;
-	MinValuePropertyHandle->GetValue(CurrentMinValue);
-	const int32 ClampMinValue = MinValuePropertyHandle->GetIntMetaData("ClampMin");
-	int8 CurrentMaxValue = 0.0f;
-	MaxValuePropertyHandle->GetValue(CurrentMaxValue);
-	const int32 ClampMaxValue = MinValuePropertyHandle->GetIntMetaData("ClampMax");
-	TSharedPtr<SMinMaxSlider> CustomRangeSlider;
-	TSharedRef<SWidget> OutWidget = SNew(SHorizontalBox)
-		+SHorizontalBox::Slot()
-		.HAlign(HAlign_Left)
-		.Padding(0, 5, 0, 5)
-		.AutoWidth()
-		[
-			MinValuePropertyHandle->CreatePropertyValueWidget()
-		]
-		+SHorizontalBox::Slot()
-		.HAlign(HAlign_Fill)
-		.VAlign(VAlign_Center)
-		.FillWidth(1)
-		[
-			SAssignNew(CustomRangeSlider, SMinMaxSlider)
-			.MinValue(ClampMinValue)
-			.MaxValue(ClampMaxValue)
-			.LowerHandleValue(static_cast<float>(CurrentMinValue))
-			.UpperHandleValue(static_cast<float>(CurrentMaxValue))
-			.OnLowerHandleValueChanged_Lambda([this, MinValuePropertyHandle, MaxValuePropertyHandle](float NewValue)
+			int32 Column = 0;
+			int32 Row = 0;
+			for (const FName PropertyName : PropertyHeaderNames)
 			{
-				int8 MaxValue;
-				MaxValuePropertyHandle->GetValue(MaxValue);
-				int8 NewMinValue = FMath::RoundToInt(NewValue);
-				MinValuePropertyHandle->SetValue(NewMinValue <= MaxValue ? NewMinValue : MaxValue);
-			})
-			.OnUpperHandleValueChanged_Lambda([this, MinValuePropertyHandle, MaxValuePropertyHandle](float NewValue)
-			{
-				int8 MinValue;
-				MinValuePropertyHandle->GetValue(MinValue);
-				int8 NewMaxValue = FMath::RoundToInt(NewValue);
-				MaxValuePropertyHandle->SetValue(NewMaxValue >= MinValue ? NewMaxValue : MinValue);
-			})
-		]
-		+SHorizontalBox::Slot()
-		.HAlign(HAlign_Left)
-		.VAlign(VAlign_Center)
-		.AutoWidth()
-		.Padding(0, 5, 0, 5)
-		[
-			MaxValuePropertyHandle->CreatePropertyValueWidget()
-		];
-
-	//if the min or max note numeric boxes are edited, modify the slider values accordingly
-	MinValuePropertyHandle->SetOnPropertyValueChanged(FSimpleDelegate::CreateLambda([this, CustomRangeSlider, MinValuePropertyHandle, MaxValuePropertyHandle]()
-	{
-		int8 MinValue;
-		MinValuePropertyHandle->GetValue(MinValue);
-		int8 MaxValue;
-		MaxValuePropertyHandle->GetValue(MaxValue);
-		CustomRangeSlider->SetLowerValue(static_cast<float>(MinValue <= MaxValue ? MinValue : MaxValue));
-		CustomRangeSlider->SetUpperValue(static_cast<float>(MaxValue >= MinValue ? MaxValue : MinValue));
-	}));
+				if (TSharedPtr<IPropertyHandle> ChildHandle = Element->GetChildHandle(PropertyName))
+				{
+					GridPanel->AddSlot(Column, Row)
+						[
+							ChildHandle->CreatePropertyNameWidget()
+						];
+				}
+				else
+				{
+					GridPanel->AddSlot(Column, Row)
+						[
+							SNew(STextBlock)
+								.Font(IDetailLayoutBuilder::GetDetailFont())
+								.Text(FText::FromString(TEXT("Invalid")))
+						];
+				}
+				Column += 1;
+				Row = 0;
+			}
 
-	MaxValuePropertyHandle->SetOnPropertyValueChanged(FSimpleDelegate::CreateLambda([this, CustomRangeSlider, MinValuePropertyHandle, MaxValuePropertyHandle]()
-	{
-		int8 MinValue;
-		MinValuePropertyHandle->GetValue(MinValue);
-		int8 MaxValue;
-		MaxValuePropertyHandle->GetValue(MaxValue);
-		CustomRangeSlider->SetUpperValue(static_cast<float>(MaxValue >= MinValue ? MaxValue : MinValue));
-		CustomRangeSlider->SetLowerValue(static_cast<float>(MinValue <= MaxValue ? MinValue : MaxValue));
-	}));
+			TSharedPtr<IPropertyHandle> SoundWaveHandle = Element->GetChildHandle(GET_MEMBER_NAME_CHECKED(FKeyzoneSettings, SoundWave));
+			ChildrenBuilder.AddCustomRow(NSLOCTEXT("FusionPatch_Details", "KeyzonesHeader", "Keyzones"))
+				.NameContent()
+				[
+					SoundWaveHandle->CreatePropertyNameWidget()
+				]
+				.ValueContent()
+				[
+					GridPanel.ToSharedRef()
+				];
+		}
+	}
 	
-	return OutWidget;
-}
+	virtual void GenerateHeaderRowContent(FDetailWidgetRow& NodeRow) override
+	{
+		TSharedPtr<IPropertyHandle> PropertyHandle = GetPropertyHandle();
+		check(PropertyHandle);
 
-FDetailWidgetRow& FFusionPatchDetailCustomization::AddCustomMinMaxSliderRow(IDetailCategoryBuilder& FusionPatchDataCategory, const FText& DisplayName, TSharedPtr<IPropertyHandle> MinPropertyHandle, TSharedPtr<IPropertyHandle> MaxPropertyHandle)
-{
-	return FusionPatchDataCategory.AddCustomRow(FText::FromString("Fusion Patch"))
+		TSharedPtr<IPropertyHandleArray> ArrayHandle = PropertyHandle->AsArray();
+		check(ArrayHandle);
+		
+		FUIAction CopyAction;
+		FUIAction PasteAction;
+		PropertyHandle->CreateDefaultPropertyCopyPasteActions(CopyAction, PasteAction);
+		NodeRow
+		.FilterString(PropertyHandle->GetPropertyDisplayName())
 		.NameContent()
 		[
-			SNew(STextBlock)
-			.Text(DisplayName)
-			.Font(IDetailLayoutBuilder::GetDetailFont())
+			PropertyHandle->CreatePropertyNameWidget()
 		]
 		.ValueContent()
-		.MinDesiredWidth(350.f)
 		[
-			CreateMinMaxSliderWidget(MinPropertyHandle, MaxPropertyHandle)
-		];
-}
-
-void FFusionPatchDetailCustomization::AddKeyzonesNameToMenuArray(TSharedPtr<IPropertyHandleArray> KeyzonesHandle, int32 NumKeyzones)
-{
-	for (int32 KeyzoneIndex = 0; KeyzoneIndex < static_cast<int32>(NumKeyzones); ++KeyzoneIndex)
-	{
-		TSharedPtr<IPropertyHandle> SamplePathHandle = KeyzonesHandle->GetElement(KeyzoneIndex)->GetChildHandle(GET_MEMBER_NAME_CHECKED(FKeyzoneSettings, SoundWave));
-		FString SoundWaveName = FString::Printf(TEXT("KEYZONE_%d"), KeyzoneIndex);
+			SNew(SHorizontalBox)
+			+SHorizontalBox::Slot()
+			[
+				PropertyHandle->CreatePropertyValueWidget()
+			]
+		]
+		.CopyAction(CopyAction)
+		.PasteAction(PasteAction);
 		
-		UObject* SoundWave;
-		if (SamplePathHandle->GetValue(SoundWave) == FPropertyAccess::Success && SoundWave != nullptr)
-		{
-			SoundWaveName = SoundWave->GetName();
-		}
-		FString Path;
-		FString Filename;
-		FString Ext;
-		FPaths::Split(SoundWaveName, Path, Filename, Ext);
-		if (Filename.IsEmpty())
-		{
-			KeyzonesNameMenu.Add(MakeShareable(new FString(SoundWaveName)));
-		}
-		else
-		{
-			KeyzonesNameMenu.Add(MakeShareable(new FString(Filename)));
-		}
 	}
-}
+private:
+	
+	TArray<FName> PropertyHeaderNames;
+	TArray<TSharedPtr<IPropertyHandle>> OptionsPropertyHandles;
+	bool DisplayHeaderNames = true;
+};
 
-void FFusionPatchDetailCustomization::DrawSelectedKeyzoneProperties(IDetailCategoryBuilder& FusionPatchDataCategory, TSharedPtr<IPropertyHandleArray> KeyzonesHandle, int32 NumKeyzones)
+void FFusionPatchDetailCustomization::CustomizeDetails(IDetailLayoutBuilder& DetailLayout)
 {
-	for (int32 KeyzoneIndex = 0; KeyzoneIndex < static_cast<int32>(NumKeyzones); ++KeyzoneIndex)
+	IDetailCategoryBuilder& FusionPatchDataCategory = DetailLayout.EditCategory(TEXT("Fusion Patch Data"));
+	//IDetailCategoryBuilder& FusionPatchDataCategory = DetailLayout.EditCategory(TEXT("Fusion Patch Data"));
+	IDetailCategoryBuilder& FilePathCategory = DetailLayout.EditCategory(TEXT("File Path"));
+
+	//get a handle of Fusion Patch Data
+	TSharedPtr<IPropertyHandle> FusionPatchDataHandle = DetailLayout.GetProperty("FusionPatchData");
+	check(FusionPatchDataHandle);
+	TSharedPtr<IPropertyHandle> SettingsHandle = FusionPatchDataHandle->GetChildHandle("Settings");
+	check(SettingsHandle);
+
+	//get a handle of the keyzones as array
+	TSharedPtr<IPropertyHandle> KeyzonesPropertyHandle = FusionPatchDataHandle->GetChildHandle("Keyzones");
+	TSharedPtr<IPropertyHandleArray> KeyzonesArrayHandle = KeyzonesPropertyHandle->AsArray();
+	check(KeyzonesArrayHandle.IsValid());
+
+	TArray<FName> PropertyHeaderNames = {
+		GET_MEMBER_NAME_CHECKED(FKeyzoneSettings, MinNote),
+		GET_MEMBER_NAME_CHECKED(FKeyzoneSettings, RootNote),
+		GET_MEMBER_NAME_CHECKED(FKeyzoneSettings, MaxNote),
+		GET_MEMBER_NAME_CHECKED(FKeyzoneSettings, MinVelocity),
+		GET_MEMBER_NAME_CHECKED(FKeyzoneSettings, MaxVelocity)
+	};
+
+	bool DisplayHeaderNames = false;
+	TSharedPtr<FKeyzoneDetailArrayBuilder> ArrayBuilder = MakeShared<FKeyzoneDetailArrayBuilder>(KeyzonesPropertyHandle.ToSharedRef(), PropertyHeaderNames, DisplayHeaderNames);
+	
+	ArrayBuilder->OnGenerateArrayElementWidget(FOnGenerateArrayElementWidget::CreateLambda(
+	[PropertyHeaderNames, DisplayHeaderNames](TSharedRef<IPropertyHandle> Element, int32 Index, IDetailChildrenBuilder& ChildrenBuilder)
 	{
-		TFunction<EVisibility()> CustomVisibility = [this, KeyzoneIndex]() -> EVisibility
+		TSharedPtr<IPropertyHandle> SoundWaveHandle = Element->GetChildHandle(GET_MEMBER_NAME_CHECKED(FKeyzoneSettings, SoundWave));
+
+		SoundWaveHandle->MarkHiddenByCustomization();
+		TSharedPtr<SUniformGridPanel> GridPanel;
+		SAssignNew(GridPanel, SUniformGridPanel)
+		.MinDesiredSlotWidth(FKeyzoneDetailArrayBuilder::MinDesiredSlotWidth);
+
+		int32 Column = 0;
+		for (const FName PropertyName : PropertyHeaderNames)
 		{
-			if (KeyzoneIndex != CurrentKeyzoneIndex)
+			if (TSharedPtr<IPropertyHandle> ChildHandle = Element->GetChildHandle(PropertyName))
 			{
-				return EVisibility::Collapsed;
+				// display the header names along with the property values
+				// whene the array builder isn't displaying the header names itself
+				if (!DisplayHeaderNames)
+				{
+					ChildHandle->MarkHiddenByCustomization();
+					GridPanel->AddSlot(Column, 0)
+						[
+							ChildHandle->CreatePropertyNameWidget()
+						];
+					GridPanel->AddSlot(Column, 1)
+						[
+							ChildHandle->CreatePropertyValueWidget()
+						];
+				}
+				else // display only the property values
+				{
+					GridPanel->AddSlot(Column, 0)
+						[
+							ChildHandle->CreatePropertyValueWidget()
+						];
+				}
 			}
 			else
 			{
-				return EVisibility::Visible;
-			}
-		};
-		
-		TSharedRef<IPropertyHandle> KeyzoneHandle = KeyzonesHandle->GetElement(KeyzoneIndex);
-		uint32 NumChildren = 0;
-		KeyzoneHandle->GetNumChildren(NumChildren);
-		for (uint32 ChildIdx = 0; ChildIdx < NumChildren; ++ChildIdx)
-		{
-			TSharedPtr<IPropertyHandle> PropertyHandle = KeyzoneHandle->GetChildHandle(ChildIdx);
-			check(PropertyHandle);
-
-			if (PropertyHandle->GetProperty()->GetFName() == GET_MEMBER_NAME_CHECKED(FKeyzoneSettings, MinNote))
-			{
-				TSharedPtr<IPropertyHandle> MaxPropertyHandle = KeyzoneHandle->GetChildHandle(ChildIdx + 1);
-				check(MaxPropertyHandle);
-				FDetailWidgetRow& MinMaxRow = AddCustomMinMaxSliderRow(FusionPatchDataCategory, NSLOCTEXT("FusionPatch_Details", "NoteRange", "Note Range (Min/Max)"), PropertyHandle, MaxPropertyHandle);
-				MinMaxRow.Visibility(TAttribute<EVisibility>::Create(TAttribute<EVisibility>::FGetter::CreateLambda(CustomVisibility)));
-				// increment ChildIdx since we customized two properties
-				++ChildIdx;
-			}
-			else if (PropertyHandle->GetProperty()->GetFName() == GET_MEMBER_NAME_CHECKED(FKeyzoneSettings, MinVelocity))
-			{
-				TSharedPtr<IPropertyHandle> MaxPropertyHandle = KeyzoneHandle->GetChildHandle(ChildIdx + 1);
-				check(MaxPropertyHandle);
-				FDetailWidgetRow& MinMaxRow = AddCustomMinMaxSliderRow(FusionPatchDataCategory, NSLOCTEXT("FusionPatch_Details", "VelocityRange", "Velocity Range (Min/Max)"), PropertyHandle, MaxPropertyHandle);
-				MinMaxRow.Visibility(TAttribute<EVisibility>::Create(TAttribute<EVisibility>::FGetter::CreateLambda(CustomVisibility)));
-				// increment ChildIdx since we customized two properties
-				++ChildIdx;
-			}
-			else // default property
-			{
-				IDetailPropertyRow& PropertyRow = FusionPatchDataCategory.AddProperty(PropertyHandle);
-				PropertyRow.Visibility(TAttribute<EVisibility>::Create(TAttribute<EVisibility>::FGetter::CreateLambda(CustomVisibility)));
+				GridPanel->AddSlot(Column, 0)
+				[
+					SNew(STextBlock)
+					.Font(IDetailLayoutBuilder::GetDetailFont())
+					.Text(FText::FromString(TEXT("Invalid")))
+				];
 			}
 
+			Column += 1;
 		}
-	}
-}
\ No newline at end of file
+
+		ChildrenBuilder.AddProperty(Element).CustomWidget(true)
+			.NameContent()
+			[
+				SoundWaveHandle->CreatePropertyValueWidget()
+			]
+			.ValueContent()
+			[
+				GridPanel.ToSharedRef()
+			];
+	}));
+
+	FusionPatchDataCategory.AddCustomBuilder(ArrayBuilder.ToSharedRef());
+	//add the Settings properties back so they're visible
+	FusionPatchDataCategory.AddProperty(SettingsHandle);
+
+	//hide the non-customized versions of these properties
+	DetailLayout.HideProperty(FusionPatchDataHandle);
+}
diff --git a/Engine/Plugins/Runtime/Harmonix/Source/HarmonixDspEditor/Private/Customization/FusionPatchDetailCustomization.h b/Engine/Plugins/Runtime/Harmonix/Source/HarmonixDspEditor/Private/Customization/FusionPatchDetailCustomization.h
--- a/Engine/Plugins/Runtime/Harmonix/Source/HarmonixDspEditor/Private/Customization/FusionPatchDetailCustomization.h
+++ b/Engine/Plugins/Runtime/Harmonix/Source/HarmonixDspEditor/Private/Customization/FusionPatchDetailCustomization.h
@@ -21,19 +21,6 @@ public:
 	}
 
 	virtual void CustomizeDetails(IDetailLayoutBuilder& DetailLayout) override;
-private:
-	int32 CurrentKeyzoneIndex = 0;
-	
-	void AddKeyzonesNameToMenuArray(TSharedPtr<IPropertyHandleArray> KeyzonesHandle, int32 NumKeyzones);
-
-	//helper functions to paint customized properties
-	FDetailWidgetRow& AddCustomMinMaxSliderRow(IDetailCategoryBuilder& FusionPatchDataCategory, const FText& DisplayName, TSharedPtr<IPropertyHandle> MinPropertyHandle, TSharedPtr<IPropertyHandle> MaxPropertyHandle);
-	void DrawSelectedKeyzoneProperties(IDetailCategoryBuilder& FusionPatchDataCategory, TSharedPtr<IPropertyHandleArray> KeyzonesHandle, int32 NumKeyzones);
-
-	TSharedPtr<FString> CurrentKeyzoneName;
-	TArray<TSharedPtr<FString>> KeyzonesNameMenu;
-
-	TSharedRef<SWidget> CreateMinMaxSliderWidget(TSharedPtr<IPropertyHandle> MinValuePropertyHandle, TSharedPtr<IPropertyHandle> MaxValuePropertyHandle);
 };
 
 
