Fix MacArm64 compile errors.

make member constexpr

diff --git a/Engine/Plugins/Runtime/Harmonix/Source/HarmonixDspEditor/Private/Customization/FusionPatchDetailCustomization.cpp b/Engine/Plugins/Runtime/Harmonix/Source/HarmonixDspEditor/Private/Customization/FusionPatchDetailCustomization.cpp
--- a/Engine/Plugins/Runtime/Harmonix/Source/HarmonixDspEditor/Private/Customization/FusionPatchDetailCustomization.cpp
+++ b/Engine/Plugins/Runtime/Harmonix/Source/HarmonixDspEditor/Private/Customization/FusionPatchDetailCustomization.cpp
@@ -15,7 +15,7 @@
 class FKeyzoneDetailArrayBuilder : public FDetailArrayBuilder
 {
 public:
-	static const int32 MinDesiredSlotWidth = 80;
+	static constexpr int32 MinDesiredSlotWidth = 80;
 
 	FKeyzoneDetailArrayBuilder(TSharedRef<IPropertyHandle> InBaseProperty, const TArray<FName>& PropertyHeaderNames, bool DisplayHeaderNames = true)
 		: FDetailArrayBuilder(InBaseProperty)
