Make sure we're comparing CurrentMidiNotes with the Original notes and not the transposed notes

#tests PIE, Creative, Editor Metasounds

diff --git a/Engine/Plugins/Runtime/Harmonix/Source/HarmonixMetasound/Private/HarmonixMetasound/Nodes/StepSequencePlayerNode.cpp b/Engine/Plugins/Runtime/Harmonix/Source/HarmonixMetasound/Private/HarmonixMetasound/Nodes/StepSequencePlayerNode.cpp
--- a/Engine/Plugins/Runtime/Harmonix/Source/HarmonixMetasound/Private/HarmonixMetasound/Nodes/StepSequencePlayerNode.cpp
+++ b/Engine/Plugins/Runtime/Harmonix/Source/HarmonixMetasound/Private/HarmonixMetasound/Nodes/StepSequencePlayerNode.cpp
@@ -764,19 +764,20 @@ namespace HarmonixMetasound
 						uint8 MidiNote;
 						CurrentCellNotes[i].GetChannelAndNote(MidiCh, MidiNote);
 
+						int32 OriginalNote = SequenceTable->Notes[i].NoteNumber;
+						int32 TransposedNote = SequenceTable->Notes[i].NoteNumber + AdditionalOctaveNotes;
+
 						// If: 
 						// 1. The note that would play is not enabled (it is not a new note);
 						// 2. A note with this pitch would start up this tick;
 						// 3. The note is marked as a continuation note
 						// Then keep this note playing through the next cell
-						if (MidiNote == (SequenceTable->Notes[i].NoteNumber + AdditionalOctaveNotes) && !CurrentPage->Rows[i].Cells[NextCellInRow].bEnabled && CurrentPage->Rows[i].Cells[NextCellInRow].bContinuation)
+						if (MidiNote == OriginalNote && !CurrentPage->Rows[i].Cells[NextCellInRow].bEnabled && CurrentPage->Rows[i].Cells[NextCellInRow].bContinuation)
 						{
 							continue;
 						}
 
 						// note off!
-						int32 OriginalNote = SequenceTable->Notes[i].NoteNumber;
-						int32 TransposedNote = SequenceTable->Notes[i].NoteNumber + AdditionalOctaveNotes;
 						// create the midi event with the original note to maintain voice ids.
 						FMidiStreamEvent MidiEvent(CurrentCellNotes[i].GetGeneratorId(), FMidiMsg::CreateNoteOff(MidiCh, OriginalNote));
 						// and then assign the note directly to the midi message after
