Make MidiClockEvent contain a MidiClockMsg which itself is a union of different MidiClockMsg types

#tests Editor, Harmonix Tests

diff --git a/Engine/Plugins/Runtime/Harmonix/Source/HarmonixMetasound/Private/HarmonixMetasound/DataTypes/MidiClock.cpp b/Engine/Plugins/Runtime/Harmonix/Source/HarmonixMetasound/Private/HarmonixMetasound/DataTypes/MidiClock.cpp
--- a/Engine/Plugins/Runtime/Harmonix/Source/HarmonixMetasound/Private/HarmonixMetasound/DataTypes/MidiClock.cpp
+++ b/Engine/Plugins/Runtime/Harmonix/Source/HarmonixMetasound/Private/HarmonixMetasound/DataTypes/MidiClock.cpp
@@ -496,34 +496,34 @@ namespace HarmonixMetasound
 	{
 		const int32 FromTick = CurrentTick;
 		FMidiPlayCursor::Reset(ForceNoBroadcast);
-		AddEvent(FMidiClockEvent::MakeResetEvent(MyMidiClock->GetCurrentBlockFrameIndex(), FromTick, CurrentTick, ForceNoBroadcast));
+		AddEvent(FMidiClockEvent(MyMidiClock->GetCurrentBlockFrameIndex(), FMidiClockMsg::FReset(FromTick, CurrentTick, ForceNoBroadcast)));
 	}
 
 	void FMidiClock::FMidiClockEventCursor::OnLoop(int32 LoopStartTick, int32 LoopEndTick)
 	{
 		FMidiPlayCursor::OnLoop(LoopStartTick, LoopEndTick);
-		AddEvent(FMidiClockEvent::MakeLoopEvent(MyMidiClock->GetCurrentBlockFrameIndex(), LoopStartTick, LoopEndTick));
+		AddEvent(FMidiClockEvent(MyMidiClock->GetCurrentBlockFrameIndex(), FMidiClockMsg::FLoop(LoopStartTick, LoopEndTick)));
 	}
 
 	void FMidiClock::FMidiClockEventCursor::SeekToTick(int32 Tick) 
 	{
 		const int32 FromTick = CurrentTick;
 		FMidiPlayCursor::SeekToTick(Tick);
-		AddEvent(FMidiClockEvent::MakeSeekToEvent(MyMidiClock->GetCurrentBlockFrameIndex(), FromTick, Tick));
+		AddEvent(FMidiClockEvent(MyMidiClock->GetCurrentBlockFrameIndex(), FMidiClockMsg::FSeekTo(FromTick, Tick)));
 	}
 
 	void FMidiClock::FMidiClockEventCursor::SeekThruTick(int32 Tick)
 	{
 		const int32 FromTick = CurrentTick;
 		FMidiPlayCursor::SeekThruTick(Tick);
-		AddEvent(FMidiClockEvent::MakeSeekThruEvent(MyMidiClock->GetCurrentBlockFrameIndex(), FromTick, Tick));
+		AddEvent(FMidiClockEvent(MyMidiClock->GetCurrentBlockFrameIndex(), FMidiClockMsg::FSeekThru(FromTick, Tick)));
 	}
 
 	void FMidiClock::FMidiClockEventCursor::AdvanceThruTick(int32 Tick, bool IsPreRoll)
 	{
 		const int32 FromTick = CurrentTick;
 		FMidiPlayCursor::AdvanceThruTick(Tick, IsPreRoll);
-		AddEvent(FMidiClockEvent::MakeAdvanceThruEvent(MyMidiClock->GetCurrentBlockFrameIndex(), FromTick, Tick, IsPreRoll));
+		AddEvent(FMidiClockEvent(MyMidiClock->GetCurrentBlockFrameIndex(), FMidiClockMsg::FAdvanceThru(FromTick, Tick, IsPreRoll)));
 	}
 
 	void FMidiClock::FMidiClockEventCursor::OnTempo(int32 TrackIndex, int32 Tick, int32 Tempo, bool IsPreroll)
diff --git a/Engine/Plugins/Runtime/Harmonix/Source/HarmonixMetasound/Private/HarmonixMetasound/DataTypes/MidiClockEvent.cpp b/Engine/Plugins/Runtime/Harmonix/Source/HarmonixMetasound/Private/HarmonixMetasound/DataTypes/MidiClockEvent.cpp
--- a/Engine/Plugins/Runtime/Harmonix/Source/HarmonixMetasound/Private/HarmonixMetasound/DataTypes/MidiClockEvent.cpp
+++ b/Engine/Plugins/Runtime/Harmonix/Source/HarmonixMetasound/Private/HarmonixMetasound/DataTypes/MidiClockEvent.cpp
@@ -6,38 +6,67 @@
 
 namespace HarmonixMetasound
 {
-	FMidiClockEvent FMidiClockEvent::MakeResetEvent(int32 InBlockFrameIndex, int32 FromTick, int32 ToTick, bool ForceNoBroadcast)
+	const FMidiClockMsg::FReset& FMidiClockMsg::AsReset() const
 	{
-		return FMidiClockEvent(EType::Reset, InBlockFrameIndex, FromTick, ToTick, false, ForceNoBroadcast);
+		check(Type == EType::Reset);
+		return Reset;
 	}
-
-	FMidiClockEvent FMidiClockEvent::MakeLoopEvent(int32 InBlockFrameIndex, int32 InStartTick, int32 InEndTick)
+	const FMidiClockMsg::FLoop& FMidiClockMsg::AsLoop() const
 	{
-		return FMidiClockEvent(EType::Loop, InBlockFrameIndex, InStartTick, InEndTick);
+		check(Type == EType::Loop);
+		return Loop;
 	}
-
-	FMidiClockEvent FMidiClockEvent::MakeSeekToEvent(int32 InBlockFrameIndex, int32 FromTick, int32 ToTick)
+	const FMidiClockMsg::FSeekTo& FMidiClockMsg::AsSeekTo() const
 	{
-		return FMidiClockEvent(EType::SeekTo, InBlockFrameIndex, FromTick, ToTick);
+		check(Type == EType::SeekTo);
+		return SeekTo;
+	}
+	const FMidiClockMsg::FSeekThru& FMidiClockMsg::AsSeekThru() const
+	{
+		check(Type == EType::SeekThru);
+		return SeekThru;
+	}
+	const FMidiClockMsg::FAdvanceThru& FMidiClockMsg::AsAdvanceThru() const
+	{
+		check(Type == EType::AdvanceThru);
+		return AdvanceThru;
 	}
 
-	FMidiClockEvent FMidiClockEvent::MakeSeekThruEvent(int32 InBlockFrameIndex, int32 FromTick, int32 ThruTick)
+	int32 FMidiClockMsg::FromTick() const
+	{
+		switch (Type)
+		{
+		case EType::Reset: return Reset.FromTick;
+		case EType::SeekTo: return SeekTo.FromTick;
+		case EType::SeekThru: return SeekThru.FromTick;
+		case EType::AdvanceThru: return AdvanceThru.FromTick;
+		default: checkNoEntry();
+		}
+		return -1;
+	}
+	
+	int32 FMidiClockMsg::ToTick() const
 	{
-		return FMidiClockEvent(EType::SeekThru, InBlockFrameIndex, FromTick, ThruTick);
+		switch (Type)
+		{
+		case EType::Reset: return Reset.ToTick;
+		case EType::SeekTo: return SeekTo.ToTick;
+		default: checkNoEntry();
+		}
+		return -1;
 	}
 
-	FMidiClockEvent FMidiClockEvent::MakeAdvanceThruEvent(int32 InBlockFrameIndex, int32 FromTick, int32 ThruTick, bool IsPreRoll)
+	int32 FMidiClockMsg::ThruTick() const
 	{
-		return FMidiClockEvent(EType::AdvanceThru, InBlockFrameIndex, FromTick, ThruTick, IsPreRoll);
+		switch (Type)
+		{
+		case EType::SeekThru: return SeekThru.ThruTick;
+		case EType::AdvanceThru: return AdvanceThru.ThruTick;
+		default: checkNoEntry();
+		}
+		return -1;
 	}
 
-	FMidiClockEvent::FMidiClockEvent(EType InType, int32 InBlockFrameIndex, int32 InTick1, int32 InTick2, bool InIsPreRoll, bool InForceNoBroadcast)
-		: Type(InType)
-		, IsPreRoll(InIsPreRoll)
-		, ForceNoBroadcast(InForceNoBroadcast)
-		, BlockFrameIndex(InBlockFrameIndex)
-		, Tick1(InTick1)
-		, Tick2(InTick2)
-	{}
+
 
 }
diff --git a/Engine/Plugins/Runtime/Harmonix/Source/HarmonixMetasound/Private/HarmonixMetasound/Nodes/MidiClockOffsetNode.cpp b/Engine/Plugins/Runtime/Harmonix/Source/HarmonixMetasound/Private/HarmonixMetasound/Nodes/MidiClockOffsetNode.cpp
--- a/Engine/Plugins/Runtime/Harmonix/Source/HarmonixMetasound/Private/HarmonixMetasound/Nodes/MidiClockOffsetNode.cpp
+++ b/Engine/Plugins/Runtime/Harmonix/Source/HarmonixMetasound/Private/HarmonixMetasound/Nodes/MidiClockOffsetNode.cpp
@@ -1,4 +1,4 @@
-﻿// Copyright Epic Games, Inc. All Rights Reserved.
+// Copyright Epic Games, Inc. All Rights Reserved.
 
 #include "MetasoundExecutableOperator.h"
 #include "MetasoundFacade.h"
@@ -171,12 +171,12 @@ namespace HarmonixMetasound::Nodes::MidiClockOffset
 			for (int32 EventIndex = 0; EventIndex < MidiClockEvents.Num(); ++EventIndex)
 			{
 				const FMidiClockEvent& Event = MidiClockEvents[EventIndex];
-				switch (Event.Type)
+				switch (Event.Msg.Type)
 				{
-				case FMidiClockEvent::EType::SeekTo:
-				case FMidiClockEvent::EType::Reset:
+				case FMidiClockMsg::EType::SeekTo:
+				case FMidiClockMsg::EType::Reset:
 					{
-						const int32 Tick = GetTickWithOffset(Event.Tick2, OffsetBars, OffsetBeats, OffsetMs);
+						const int32 Tick = GetTickWithOffset(Event.Msg.ToTick(), OffsetBars, OffsetBeats, OffsetMs);
 						FMusicSeekTarget SeekTarget;
 						SeekTarget.Type = ESeekPointType::Millisecond;
 						SeekTarget.Ms = MidiClockOut->GetSongMaps().TickToMs(Tick);
@@ -184,9 +184,9 @@ namespace HarmonixMetasound::Nodes::MidiClockOffset
 						MidiClockOut->SeekTo(Event.BlockFrameIndex, SeekTarget, PrerollBars);
 						break;
 					}
-				case FMidiClockEvent::EType::SeekThru:
+				case FMidiClockMsg::EType::SeekThru:
 					{
-						const int32 Tick = GetTickWithOffset(Event.Tick2, OffsetBars, OffsetBeats, OffsetMs);
+						const int32 Tick = GetTickWithOffset(Event.Msg.ThruTick(), OffsetBars, OffsetBeats, OffsetMs);
 						FMusicSeekTarget SeekTarget;
 						SeekTarget.Type = ESeekPointType::Millisecond;
 						SeekTarget.Ms = MidiClockOut->GetSongMaps().TickToMs(Tick + 1);
@@ -194,9 +194,9 @@ namespace HarmonixMetasound::Nodes::MidiClockOffset
 						MidiClockOut->SeekTo(Event.BlockFrameIndex, SeekTarget, PrerollBars);
 						break;
 					}
-				case FMidiClockEvent::EType::AdvanceThru:
+				case FMidiClockMsg::EType::AdvanceThru:
 					{
-						const int32 Tick = GetTickWithOffset(Event.Tick2, OffsetBars, OffsetBeats, OffsetMs);
+						const int32 Tick = GetTickWithOffset(Event.Msg.ThruTick(), OffsetBars, OffsetBeats, OffsetMs);
 						
 						// if our offset changed while we were advancing, seek to the new offset first
 						if (!FMath::IsNearlyEqual(PrevOffsetMs, OffsetMs) || !FMath::IsNearlyEqual(PrevOffsetBeats, OffsetBeats) || PrevOffsetBars != OffsetBars)
@@ -214,12 +214,12 @@ namespace HarmonixMetasound::Nodes::MidiClockOffset
 
 						const float AdvanceToMs = MidiClockOut->GetSongMaps().TickToMs(Tick);
 						const float ClockInSpeed = MidiClockIn->GetSpeedAtBlockSampleFrame(EventIndex);
-						const float AdvanceRatio = MidiClockIn->GetSongMaps().GetTempoAtTick(Event.Tick2) / MidiClockOut->GetSongMaps().GetTempoAtTick(Tick);
+						const float AdvanceRatio = MidiClockIn->GetSongMaps().GetTempoAtTick(Event.Msg.ThruTick()) / MidiClockOut->GetSongMaps().GetTempoAtTick(Tick);
 						MidiClockOut->InformOfCurrentAdvanceRate(ClockInSpeed * AdvanceRatio);
 						MidiClockOut->AdvanceHiResToMs(Event.BlockFrameIndex, AdvanceToMs, true);
 						break;
 					}
-				case FMidiClockEvent::EType::Loop:
+				case FMidiClockMsg::EType::Loop:
 					{
 						break;
 					}
diff --git a/Engine/Plugins/Runtime/Harmonix/Source/HarmonixMetasound/Private/HarmonixMetasound/Nodes/MidiClockSubdivisionTriggerNode.cpp b/Engine/Plugins/Runtime/Harmonix/Source/HarmonixMetasound/Private/HarmonixMetasound/Nodes/MidiClockSubdivisionTriggerNode.cpp
--- a/Engine/Plugins/Runtime/Harmonix/Source/HarmonixMetasound/Private/HarmonixMetasound/Nodes/MidiClockSubdivisionTriggerNode.cpp
+++ b/Engine/Plugins/Runtime/Harmonix/Source/HarmonixMetasound/Private/HarmonixMetasound/Nodes/MidiClockSubdivisionTriggerNode.cpp
@@ -240,19 +240,19 @@ namespace HarmonixMetasound
 		{
 			for (const FMidiClockEvent& ClockEvent : MidiClockInPin->GetMidiClockEventsInBlock())
 			{
-				if (ClockEvent.Type == FMidiClockEvent::EType::AdvanceThru)
+				if (ClockEvent.Msg.Type == FMidiClockMsg::EType::AdvanceThru)
 				{
-					if (ClockEvent.IsPreRoll)
+					if (ClockEvent.Msg.AsAdvanceThru().IsPreRoll)
 					{
 						continue;
 					}
 					
-					int32 TickPreceedingThisAdvance = ClockEvent.Tick1 - GridOffsetTicks;
-					int32 LastTickProcessed = ClockEvent.Tick2 - GridOffsetTicks;
+					int32 TickPreceedingThisAdvance = ClockEvent.Msg.FromTick() - GridOffsetTicks;
+					int32 LastTickProcessed = ClockEvent.Msg.ThruTick() - GridOffsetTicks;
 					const FSongMaps& SongMaps = MidiClockInPin->GetSongMaps();
 					
 					// calculate new GridSizeTicks and GridOffsetTicks at the tick we're evaluating based on any time signature changes
-					GridSizeTicks = SongMaps.SubdivisionToMidiTicks(GridSizeUnits, ClockEvent.Tick2) * GridSizeMultiplier;
+					GridSizeTicks = SongMaps.SubdivisionToMidiTicks(GridSizeUnits, ClockEvent.Msg.ThruTick()) * GridSizeMultiplier;
 					
 					// if this is one of the quantization types that is variable in size we have to do it a little different...
 					if (GridSizeUnits == EMidiClockSubdivisionQuantization::Bar ||
diff --git a/Engine/Plugins/Runtime/Harmonix/Source/HarmonixMetasound/Private/HarmonixMetasound/Nodes/MidiClockTimestampTriggerNode.cpp b/Engine/Plugins/Runtime/Harmonix/Source/HarmonixMetasound/Private/HarmonixMetasound/Nodes/MidiClockTimestampTriggerNode.cpp
--- a/Engine/Plugins/Runtime/Harmonix/Source/HarmonixMetasound/Private/HarmonixMetasound/Nodes/MidiClockTimestampTriggerNode.cpp
+++ b/Engine/Plugins/Runtime/Harmonix/Source/HarmonixMetasound/Private/HarmonixMetasound/Nodes/MidiClockTimestampTriggerNode.cpp
@@ -221,19 +221,18 @@ namespace HarmonixMetasound
 		{
 			for (const FMidiClockEvent& ClockEvent : MidiClockInPin->GetMidiClockEventsInBlock())
 			{
-				switch (ClockEvent.Type)
+				switch (ClockEvent.Msg.Type)
 				{
-				case FMidiClockEvent::EType::SeekThru:
-				case FMidiClockEvent::EType::AdvanceThru:
+				case FMidiClockMsg::EType::AdvanceThru:
+					if (ClockEvent.Msg.AsAdvanceThru().IsPreRoll)
 					{
-						if (ClockEvent.IsPreRoll)
-						{
-							continue;
-						}
-						
-						const bool EventIsSeek = ClockEvent.Type == FMidiClockEvent::EType::SeekThru;
+						continue;
+					}
+				case FMidiClockMsg::EType::SeekThru:
+					{
+						const bool EventIsSeek = ClockEvent.Msg.Type == FMidiClockMsg::EType::SeekThru;
 
-						if (ClockEvent.Tick1 < TriggerTick && ClockEvent.Tick2 >= TriggerTick)
+						if (ClockEvent.Msg.FromTick() < TriggerTick && ClockEvent.Msg.ThruTick() >= TriggerTick)
 						{
 							if (!EventIsSeek || *TriggerDuringSeekInPin)
 							{
diff --git a/Engine/Plugins/Runtime/Harmonix/Source/HarmonixMetasound/Private/HarmonixMetasound/Nodes/MidiPlayerNode.cpp b/Engine/Plugins/Runtime/Harmonix/Source/HarmonixMetasound/Private/HarmonixMetasound/Nodes/MidiPlayerNode.cpp
--- a/Engine/Plugins/Runtime/Harmonix/Source/HarmonixMetasound/Private/HarmonixMetasound/Nodes/MidiPlayerNode.cpp
+++ b/Engine/Plugins/Runtime/Harmonix/Source/HarmonixMetasound/Private/HarmonixMetasound/Nodes/MidiPlayerNode.cpp
@@ -400,47 +400,47 @@ namespace HarmonixMetasound
 					// We should be advancing midi clock events in order (never repeating)!
 					check(CurrentMidiClockEventIndex < EventIndex);
 					CurrentMidiClockEventIndex = EventIndex;
-					switch (Event.Type)
+					switch (Event.Msg.Type)
 					{
-					case FMidiClockEvent::EType::Reset:
+					case FMidiClockMsg::EType::Reset:
 					{
-						int32 Tick = MidiClockOut->CalculateMappedTick(Event.Tick2);
+						int32 Tick = MidiClockOut->CalculateMappedTick(Event.Msg.ToTick());
 						MidiClockOut->SeekTo(Event.BlockFrameIndex, Tick + 1, PrerollBars);
 						break;
 					}
-					case FMidiClockEvent::EType::Loop:
+					case FMidiClockMsg::EType::Loop:
 					{
 						// loops should actually be handled by seeking and advancing...
 						break;
 					}
-					case FMidiClockEvent::EType::SeekTo:
+					case FMidiClockMsg::EType::SeekTo:
 					{
-						int32 Tick = MidiClockOut->CalculateMappedTick(Event.Tick2);
+						int32 Tick = MidiClockOut->CalculateMappedTick(Event.Msg.ToTick());
 						MidiClockOut->SeekTo(Event.BlockFrameIndex, Tick, PrerollBars);
 						break;
 					}
-					case FMidiClockEvent::EType::SeekThru:
+					case FMidiClockMsg::EType::SeekThru:
 					{
-						int32 Tick = MidiClockOut->CalculateMappedTick(Event.Tick2);
+						int32 Tick = MidiClockOut->CalculateMappedTick(Event.Msg.ThruTick());
 						MidiClockOut->SeekTo(Event.BlockFrameIndex, Tick + 1, PrerollBars);
 						break;
 					}
-					case FMidiClockEvent::EType::AdvanceThru:
+					case FMidiClockMsg::EType::AdvanceThru:
 					{
 						// if this advance is a preroll, perform a seek instead so we don't trigger events doing an advance
 						// NOTE: This code is/should be identical to the SeekThru event above
-						if (Event.IsPreRoll)
+						if (Event.Msg.AsAdvanceThru().IsPreRoll)
 						{
-							int32 Tick = MidiClockOut->CalculateMappedTick(Event.Tick2);
+							int32 Tick = MidiClockOut->CalculateMappedTick(Event.Msg.ThruTick());
 							MidiClockOut->SeekTo(Event.BlockFrameIndex, Tick + 1, PrerollBars);
 							break;
 						}
 						// The MidiClock handles looping on its own, so we can conveniently advance it
-						int32 Tick = MidiClockOut->GetCurrentMidiTick() + (Event.Tick2 - Event.Tick1);
+						int32 Tick = MidiClockOut->GetCurrentMidiTick() + (Event.Msg.ThruTick() - Event.Msg.FromTick());
 						float Ms = MidiClockOut->GetSongMaps().TickToMs(Tick);
 
 						float ClockInSpeed = MidiClockIn->GetSpeedAtBlockSampleFrame(StartFrameIndex);
-						float AdvanceRatio = MidiClockIn->GetSongMaps().GetTempoAtTick(Event.Tick1)
+						float AdvanceRatio = MidiClockIn->GetSongMaps().GetTempoAtTick(Event.Msg.FromTick())
 							               / MidiClockOut->GetSongMaps().GetTempoAtTick(MidiClockOut->GetCurrentMidiTick());
 						// midi clock needs to know how fast its advancing based on their authority
 						MidiClockOut->InformOfCurrentAdvanceRate(ClockInSpeed * *SpeedMultInPin * AdvanceRatio);
diff --git a/Engine/Plugins/Runtime/Harmonix/Source/HarmonixMetasound/Private/HarmonixMetasound/Nodes/StepSequencePlayerNode.cpp b/Engine/Plugins/Runtime/Harmonix/Source/HarmonixMetasound/Private/HarmonixMetasound/Nodes/StepSequencePlayerNode.cpp
--- a/Engine/Plugins/Runtime/Harmonix/Source/HarmonixMetasound/Private/HarmonixMetasound/Nodes/StepSequencePlayerNode.cpp
+++ b/Engine/Plugins/Runtime/Harmonix/Source/HarmonixMetasound/Private/HarmonixMetasound/Nodes/StepSequencePlayerNode.cpp
@@ -353,16 +353,16 @@ namespace HarmonixMetasound
 
 						if (Event.BlockFrameIndex >= StartFrameIndex)
 						{
-							switch (Event.Type)
+							switch (Event.Msg.Type)
 							{
-							case FMidiClockEvent::EType::AdvanceThru:
-								AdvanceThruTick(Event.BlockFrameIndex, Event.Tick2, Event.IsPreRoll);
+							case FMidiClockMsg::EType::AdvanceThru:
+								AdvanceThruTick(Event.BlockFrameIndex, Event.Msg.ThruTick(), Event.Msg.AsAdvanceThru().IsPreRoll);
 								break;
-							case FMidiClockEvent::EType::SeekThru:
-								SeekThruTick(Event.BlockFrameIndex, Event.Tick2);
+							case FMidiClockMsg::EType::SeekThru:
+								SeekThruTick(Event.BlockFrameIndex, Event.Msg.ThruTick());
 								break;
-							case FMidiClockEvent::EType::SeekTo:
-								SeekToTick(Event.BlockFrameIndex, Event.Tick2);
+							case FMidiClockMsg::EType::SeekTo:
+								SeekToTick(Event.BlockFrameIndex, Event.Msg.ToTick());
 								break;
 							}
 						}
diff --git a/Engine/Plugins/Runtime/Harmonix/Source/HarmonixMetasound/Public/HarmonixMetasound/DataTypes/MidiClockEvent.h b/Engine/Plugins/Runtime/Harmonix/Source/HarmonixMetasound/Public/HarmonixMetasound/DataTypes/MidiClockEvent.h
--- a/Engine/Plugins/Runtime/Harmonix/Source/HarmonixMetasound/Public/HarmonixMetasound/DataTypes/MidiClockEvent.h
+++ b/Engine/Plugins/Runtime/Harmonix/Source/HarmonixMetasound/Public/HarmonixMetasound/DataTypes/MidiClockEvent.h
@@ -6,7 +6,7 @@
 
 namespace HarmonixMetasound
 {
-	struct HARMONIXMETASOUND_API FMidiClockEvent
+	struct HARMONIXMETASOUND_API FMidiClockMsg
 	{
 		enum class EType : uint8
 		{
@@ -16,34 +16,94 @@ namespace HarmonixMetasound
 			SeekThru,
 			AdvanceThru
 		};
+		
+		struct FReset
+		{
+			int32 FromTick;
+			int32 ToTick;
+			bool ForceNoBroadcast;
+			FReset(int32 FromTick, int32 ToTick, bool ForceNoBroadcast)
+				: FromTick(FromTick), ToTick(ToTick), ForceNoBroadcast(ForceNoBroadcast) {}
+		};
 
-		const EType Type;
+		struct FLoop
+		{
+			int32 LoopStartTick;
+			int32 LoopEndTick;
+			FLoop(int32 LoopStartTick, int32 LoopEndTick)
+				: LoopStartTick(LoopStartTick), LoopEndTick(LoopEndTick) {}
+		};
 
-		// for "advance thru" events
-		const bool  IsPreRoll;
+		struct FSeekTo
+		{
+			int32 FromTick;
+			int32 ToTick;
+			FSeekTo(int32 FromTick, int32 ToTick)
+				: FromTick(FromTick), ToTick(ToTick) {}
+		};
 
-		// for "reset" events
-		const bool  ForceNoBroadcast;
-		const int32 BlockFrameIndex;
+		struct FSeekThru
+		{
+			int32 FromTick;
+			int32 ThruTick;
+			FSeekThru(int32 FromTick, int32 ThruTick)
+				: FromTick(FromTick), ThruTick(ThruTick) {}
+		};
+
+		struct FAdvanceThru
+		{
+			int32 FromTick;
+			int32 ThruTick;
+			bool IsPreRoll;
+			FAdvanceThru(int32 FromTick, int32 ThruTick, bool IsPreRoll)
+				: FromTick(FromTick), ThruTick(ThruTick), IsPreRoll(IsPreRoll) {}
+		};
 
-		// FromTick when event is SeekTo, SeekThru, AdvanceThru or Reset;
-		// LoopStartTick when OnLoop
-		const int32 Tick1;
+		FMidiClockMsg(const FReset& Reset) : Type(EType::Reset), Reset(Reset) {};
+		FMidiClockMsg(const FLoop& Loop) : Type(EType::Loop), Loop(Loop) {};
+		FMidiClockMsg(const FSeekTo& SeekTo) : Type(EType::SeekTo), SeekTo(SeekTo) {};
+		FMidiClockMsg(const FSeekThru& SeekThru) : Type(EType::SeekThru), SeekThru(SeekThru) {};
+		FMidiClockMsg(const FAdvanceThru AdvanceThru) : Type(EType::AdvanceThru), AdvanceThru(AdvanceThru) {};
 
-		// ToTick when SeekTo, Reset;
-		// ThruTick when SeekThru, AdvanceThru;
-		// LoopEndTick when OnLoop
-		const int32 Tick2;
+		// type safe way to access msg data
+		// asserts that is of the corresponding type
+		const FReset& AsReset() const;
+		const FLoop& AsLoop() const;
+		const FSeekTo& AsSeekTo() const;
+		const FSeekThru& AsSeekThru() const;
+		const FAdvanceThru& AsAdvanceThru() const;
 
-		static FMidiClockEvent MakeResetEvent(int32 InBlockFrameIndex, int32 FromTick, int32 ToTick, bool ForceNoBroadcast = false);
-		static FMidiClockEvent MakeLoopEvent(int32 InBlockFrameIndex, int32 LoopStartTick, int32 LoopEndTick);
-		static FMidiClockEvent MakeSeekToEvent(int32 InBlockFrameIndex, int32 FromTick, int32 ToTick);
-		static FMidiClockEvent MakeSeekThruEvent(int32 InBlockFrameIndex, int32 FromTick, int32 ThruTick);
-		static FMidiClockEvent MakeAdvanceThruEvent(int32 InBlockFrameIndex, int32 FromTick, int32 ThruTick, bool IsPreRoll);
+		// Asserts that this is a Reset, SeekTo, SeekThru, or AdvanceThru type
+		int32 FromTick() const;
 
-	private:
+		// Asserts that this is a Reset, or SeekTo type
+		int32 ToTick() const;
 
-		FMidiClockEvent(EType InType, int32 InBlockFrameIndex, int32 InTick1, int32 InTick2, bool InIsPreRoll = false, bool InForceNoBroadcast = false);
+		// Asserts that this is a SeekThru or AdvanceThru type
+		int32 ThruTick() const;
+
+		const EType Type;
+	private:
+		union 
+		{
+			const FReset Reset;
+			const FLoop Loop;
+			const FSeekTo SeekTo;
+			const FSeekThru SeekThru;
+			const FAdvanceThru AdvanceThru;
+		};
+		
+	};
+	
+	struct HARMONIXMETASOUND_API FMidiClockEvent
+	{
+		const int32 BlockFrameIndex;
+		const FMidiClockMsg Msg;
+		
+		FMidiClockEvent(int32 InBlockFrameIndex, const FMidiClockMsg& Msg)
+			: BlockFrameIndex(InBlockFrameIndex), Msg(Msg)
+		{}
+		
 	};
 
 };
