Move shared MidiClock logic to the MidiClock class.

Processing Midi Clock Events can be handled by the MidiClock itself.

MidiPlayer node and MetronomeNode replaces all MidiClock event handling with this new API.

Metronome node no longer needs to be a MidiPlayCursor.
MidiPlayer node still does because it handles midi events this way still (but could be refactored to not do it this way anymore).

Also, add way to get "Next" transport state for typical advancing of the transport. Simplifies a lot of logic.

#tests PIE, Harmonix Tests

diff --git a/Engine/Plugins/Runtime/Harmonix/Source/HarmonixMetasound/Private/HarmonixMetasound/DataTypes/MidiClock.cpp b/Engine/Plugins/Runtime/Harmonix/Source/HarmonixMetasound/Private/HarmonixMetasound/DataTypes/MidiClock.cpp
--- a/Engine/Plugins/Runtime/Harmonix/Source/HarmonixMetasound/Private/HarmonixMetasound/DataTypes/MidiClock.cpp
+++ b/Engine/Plugins/Runtime/Harmonix/Source/HarmonixMetasound/Private/HarmonixMetasound/DataTypes/MidiClock.cpp
@@ -389,6 +389,98 @@ namespace HarmonixMetasound
 		}
 	}
 
+	void FMidiClock::HandleTransportChange(int32 BlockFrameIndex, EMusicPlayerTransportState TransportState)
+	{
+		switch (TransportState)
+		{
+		case EMusicPlayerTransportState::Prepared:
+		case EMusicPlayerTransportState::Paused:
+			AddTransportStateChangeToBlock({ BlockFrameIndex, 0.0f, TransportState });
+			break;
+		case EMusicPlayerTransportState::Playing:
+			AddTransportStateChangeToBlock({ BlockFrameIndex, 0.0f, TransportState });
+			break;
+		}
+	}
+
+	void FMidiClock::Process(int32 StartFrame, int32 NumFrames, int32 PrerollBars, float Speed)
+	{
+		int32 EndFrame = StartFrame + NumFrames;
+		switch (CurrentTransportState.TransportState)
+		{
+		case EMusicPlayerTransportState::Playing:
+		case EMusicPlayerTransportState::Continuing:
+			WriteAdvance(StartFrame, EndFrame, Speed);
+			break;
+		}
+	}
+
+	void FMidiClock::Process(const FMidiClock& DrivingClock, int32 StartFrame, int32 NumFrames, int32 PrerollBars, float Speed)
+	{
+		int32 EndFrame = StartFrame + NumFrames;
+		const TArray<FMidiClockEvent>& ClockEvents = DrivingClock.GetMidiClockEventsInBlock();
+		int32 Index = Algo::LowerBoundBy(ClockEvents, StartFrame, &FMidiClockEvent::BlockFrameIndex);
+		while (ClockEvents.IsValidIndex(Index))
+		{
+			const FMidiClockEvent& Event = ClockEvents[Index];
+			if (Event.BlockFrameIndex >= EndFrame)
+			{
+				return;
+			}
+			
+			HandleClockEvent(DrivingClock, Event, PrerollBars, Speed);
+			++Index;
+		}
+	}
+	
+	void FMidiClock::HandleClockEvent(const FMidiClock& DrivingClock, const FMidiClockEvent& Event, int32 PrerollBars, float Speed)
+	{
+		switch (Event.Msg.Type)
+		{
+		case FMidiClockMsg::EType::Reset:
+			{
+				int32 Tick = CalculateMappedTick(Event.Msg.ToTick());
+				SeekTo(Event.BlockFrameIndex, Tick + 1, PrerollBars);
+				break;
+			}
+		case FMidiClockMsg::EType::Loop:
+			// ignore loops since the driving clock will loop us by seeking
+			break;
+		case FMidiClockMsg::EType::SeekTo:
+			{
+				int32 Tick = CalculateMappedTick(Event.Msg.ToTick());
+				SeekTo(Event.BlockFrameIndex, Tick, PrerollBars);
+				break;
+			}
+		case FMidiClockMsg::EType::SeekThru:
+			{
+				int32 Tick = CalculateMappedTick(Event.Msg.ThruTick());
+				SeekTo(Event.BlockFrameIndex, Tick + 1, PrerollBars);
+				break;
+			}
+		case FMidiClockMsg::EType::AdvanceThru:
+			{
+				if (Event.Msg.AsAdvanceThru().IsPreRoll)
+				{
+					int32 Tick = CalculateMappedTick(Event.Msg.ThruTick());
+					SeekTo(Event.BlockFrameIndex, Tick + 1, PrerollBars);
+					break;
+				}
+
+				// Advance based on the delta ticks, and not based on the absolute tick
+				int32 Tick = GetCurrentMidiTick() + (Event.Msg.ThruTick() - Event.Msg.FromTick());
+				float Ms = GetSongMaps().TickToMs(Tick);
+
+				float ClockInSpeed = DrivingClock.GetSpeedAtBlockSampleFrame(Event.BlockFrameIndex);
+				float AdvanceRatio = DrivingClock.GetSongMaps().GetTempoAtTick(Event.Msg.FromTick())
+								   / GetSongMaps().GetTempoAtTick(GetCurrentMidiTick());
+				InformOfCurrentAdvanceRate(ClockInSpeed * Speed * AdvanceRatio);
+				AdvanceHiResToMs(Event.BlockFrameIndex, Ms, true);
+				break;
+			}
+		}
+	}
+
 	void FMidiClock::WriteAdvance(int32 StartFrameIndex, int32 EndFrameIndex, float InSpeed /*= 1.0f*/)
 	{
 		int32 FramesToProcess = EndFrameIndex - StartFrameIndex;
@@ -451,8 +543,7 @@ namespace HarmonixMetasound
 		}
 		return Tick;
 	}
-
-
+	
 	TSharedPtr<FMidiFileData> FMidiClock::MakeClockConductorMidiData(float InTempoBPM, int32 InTimeSigNum, int32 InTimeSigDen)
 	{
 		TSharedPtr<FMidiFileData> OutMidiData = MakeShared<FMidiFileData>();
diff --git a/Engine/Plugins/Runtime/Harmonix/Source/HarmonixMetasound/Private/HarmonixMetasound/DataTypes/MusicTransport.cpp b/Engine/Plugins/Runtime/Harmonix/Source/HarmonixMetasound/Private/HarmonixMetasound/DataTypes/MusicTransport.cpp
--- a/Engine/Plugins/Runtime/Harmonix/Source/HarmonixMetasound/Private/HarmonixMetasound/DataTypes/MusicTransport.cpp
+++ b/Engine/Plugins/Runtime/Harmonix/Source/HarmonixMetasound/Private/HarmonixMetasound/DataTypes/MusicTransport.cpp
@@ -75,6 +75,36 @@ namespace HarmonixMetasound
 		TransportState = InitFn(DesiredState);
 	}
 
+	EMusicPlayerTransportState FMusicTransportControllable::GetNextTransportState(EMusicPlayerTransportState DesiredState) const
+	{
+		switch (DesiredState)
+     	{
+		case EMusicPlayerTransportState::Invalid:
+		case EMusicPlayerTransportState::Preparing:
+		case EMusicPlayerTransportState::Prepared:
+		case EMusicPlayerTransportState::Stopping:
+		case EMusicPlayerTransportState::Killing:
+			return EMusicPlayerTransportState::Prepared;
+			
+		case EMusicPlayerTransportState::Starting:
+		case EMusicPlayerTransportState::Playing:
+		case EMusicPlayerTransportState::Continuing:
+			return EMusicPlayerTransportState::Playing;
+			
+     	case EMusicPlayerTransportState::Seeking:
+     		return GetTransportState();
+     
+     	case EMusicPlayerTransportState::Pausing:
+		case EMusicPlayerTransportState::Paused:
+     		return EMusicPlayerTransportState::Paused;
+     
+     	default:
+     		checkNoEntry();
+     		return EMusicPlayerTransportState::Invalid;
+     	}
+	}
+
+
 	void FMusicTransportControllable::ExecuteTransportSpans(FMusicTransportEventStreamReadRef& InTransportPin, int32 InBlockSize, TransportSpanProcessor& Callback)
 	{
 		TransportSpanPostProcessor PostProcessor = [](int32, int32, EMusicPlayerTransportState) {};
diff --git a/Engine/Plugins/Runtime/Harmonix/Source/HarmonixMetasound/Private/HarmonixMetasound/Nodes/MetronomeNode.cpp b/Engine/Plugins/Runtime/Harmonix/Source/HarmonixMetasound/Private/HarmonixMetasound/Nodes/MetronomeNode.cpp
--- a/Engine/Plugins/Runtime/Harmonix/Source/HarmonixMetasound/Private/HarmonixMetasound/Nodes/MetronomeNode.cpp
+++ b/Engine/Plugins/Runtime/Harmonix/Source/HarmonixMetasound/Private/HarmonixMetasound/Nodes/MetronomeNode.cpp
@@ -23,7 +23,7 @@ namespace HarmonixMetasound
 {
 	using namespace Metasound;
 
-	class FMetronomeOperator : public TExecutableOperator<FMetronomeOperator>, public FMusicTransportControllable, public FMidiPlayCursor
+	class FMetronomeOperator : public TExecutableOperator<FMetronomeOperator>, public FMusicTransportControllable
 	{
 	public:
 		static const FNodeClassMetadata& GetNodeInfo();
@@ -48,14 +48,6 @@ namespace HarmonixMetasound
 		void Reset(const FResetParams& Params);
 		void Execute();
 
-		//~ BEGIN FMidiPlayCursor Overrides
-		virtual void Reset(bool ForceNoBroadcast = false) override;
-		virtual void OnLoop(int32 LoopStartTick, int32 LoopEndTick) override;
-		virtual void SeekToTick(int32 Tick) override;
-		virtual void SeekThruTick(int32 Tick) override;
-		virtual void AdvanceThruTick(int32 Tick, bool IsPreRoll) override;
-		//~ END FMidiPlayCursor Overrides
-
 	private:
 		void Init();
 		
@@ -86,7 +78,7 @@ namespace HarmonixMetasound
 		void UpdateMidi();
 		void AddTempoChangeForMidi(float TempoBPM);
 		void AddTimeSigChangeForMidi(int32 TimeSigNum, int32 TimeSigDenom);
-		void AddTransportStateChangeToMidiClock(int32 StartFrameIndex, EMusicPlayerTransportState NewTransportState);
+		void HandleTransportChange(int32 StartFrameIndex, EMusicPlayerTransportState NewTransportState);
 
 		FMidiClock& GetDrivingMidiClock() { return LoopInPin ? MetronomeClock : (*MidiClockOutPin); }
 	};
@@ -197,7 +189,6 @@ namespace HarmonixMetasound
 	{
 		if (LoopInPin)
 		{
-			MetronomeClock.RegisterHiResPlayCursor(this);
 			MidiClockOutPin->AttachToTimeAuthority(MetronomeClock);
 		}
 
@@ -251,55 +242,6 @@ namespace HarmonixMetasound
 		LastClockTickUpdate = -1;
 	}
 
-	void FMetronomeOperator::Reset(bool ForceNoBroadcast /*= false*/)
-	{
-		FMidiPlayCursor::Reset(ForceNoBroadcast);
-		int32 Tick = MidiClockOutPin->CalculateMappedTick(MetronomeClock.GetCurrentHiResTick());
-		MidiClockOutPin->SeekTo(MetronomeClock.GetCurrentBlockFrameIndex(), Tick + 1, SeekPreRollBarsInPin);
-	}
-
-	void FMetronomeOperator::OnLoop(int32 LoopStartTick, int32 LoopEndTick)
-	{
-		FMidiPlayCursor::OnLoop(LoopStartTick, LoopEndTick);
-	}
-
-	void FMetronomeOperator::SeekToTick(int32 Tick)
-	{
-		FMidiPlayCursor::SeekToTick(Tick);
-		int32 MappedTick = MidiClockOutPin->CalculateMappedTick(Tick);
-		MidiClockOutPin->SeekTo(MetronomeClock.GetCurrentBlockFrameIndex(), MappedTick, SeekPreRollBarsInPin);
-	}
-
-	void FMetronomeOperator::SeekThruTick(int32 Tick)
-	{
-		FMidiPlayCursor::SeekThruTick(Tick);
-		int32 MappedTick = MidiClockOutPin->CalculateMappedTick(Tick);
-		MidiClockOutPin->SeekTo(MetronomeClock.GetCurrentBlockFrameIndex(), MappedTick + 1, SeekPreRollBarsInPin);
-	}
-
-	void FMetronomeOperator::AdvanceThruTick(int32 InTick, bool IsPreRoll)
-	{
-		int32 PrevTick = CurrentTick;
-		FMidiPlayCursor::AdvanceThruTick(InTick, IsPreRoll);
-
-		// if this is a pre-roll, perform a seek instead
-		// NOTE: This code is/should be identical to SeekThruTick above
-		if (IsPreRoll)
-		{
-			int32 MappedTick = MidiClockOutPin->CalculateMappedTick(InTick);
-			MidiClockOutPin->SeekTo(MetronomeClock.GetCurrentBlockFrameIndex(), MappedTick + 1, SeekPreRollBarsInPin);
-			return;
-		}
-		
-		int32 NewTick = MidiClockOutPin->GetCurrentMidiTick() + (InTick - PrevTick);
-		float Ms = MidiClockOutPin->GetSongMaps().TickToMs(NewTick);
-		float CurrentSpeed = MetronomeClock.GetSpeedAtBlockSampleFrame(MetronomeClock.GetCurrentBlockFrameIndex());
-		float AdvanceRatio = GetOwner()->GetSongMaps().GetTempoAtTick(PrevTick)
-		                   / MidiClockOutPin->GetSongMaps().GetTempoAtTick(MidiClockOutPin->GetCurrentMidiTick());
-		MidiClockOutPin->InformOfCurrentAdvanceRate(CurrentSpeed * AdvanceRatio);
-		MidiClockOutPin->AdvanceHiResToMs(MetronomeClock.GetCurrentBlockFrameIndex(), Ms, true);
-	}
-
 	void FMetronomeOperator::Init()
 	{
 		BuildMidiData();
@@ -307,52 +249,20 @@ namespace HarmonixMetasound
 		FTransportInitFn InitFn = [this](EMusicPlayerTransportState CurrentState)
 		{
 			FMidiClock& DrivingMidiClock = GetDrivingMidiClock();
-			
 			switch (CurrentState)
 			{
-			case EMusicPlayerTransportState::Invalid:
-			case EMusicPlayerTransportState::Preparing:
-				MidiClockOutPin->AddTransportStateChangeToBlock({ 0, 0.0f, EMusicPlayerTransportState::Prepared });
-				return EMusicPlayerTransportState::Prepared;
-
-			case EMusicPlayerTransportState::Prepared:
-				return CurrentState;
-
 			case EMusicPlayerTransportState::Starting:
 				DrivingMidiClock.ResetAndStart(0, true);
-				return EMusicPlayerTransportState::Playing;
+				break;
 
 			case EMusicPlayerTransportState::Playing:
 				DrivingMidiClock.WriteAdvance(0, 0, *SpeedMultInPin);
-				return EMusicPlayerTransportState::Playing;
-
-			case EMusicPlayerTransportState::Seeking:
-				checkNoEntry();
-				return EMusicPlayerTransportState::Invalid;
-
-			case EMusicPlayerTransportState::Continuing:
-				MidiClockOutPin->AddTransportStateChangeToBlock({ 0, 0.0f, EMusicPlayerTransportState::Playing });
-				return EMusicPlayerTransportState::Playing;
-
-			case EMusicPlayerTransportState::Pausing:
-				MidiClockOutPin->AddTransportStateChangeToBlock({ 0, 0.0f, EMusicPlayerTransportState::Paused });
-				return EMusicPlayerTransportState::Paused;
-
-			case EMusicPlayerTransportState::Paused:
-				return EMusicPlayerTransportState::Paused;
-
-			case EMusicPlayerTransportState::Stopping:
-				MidiClockOutPin->AddTransportStateChangeToBlock({ 0, 0.0f, EMusicPlayerTransportState::Prepared });
-				return EMusicPlayerTransportState::Prepared;
-
-			case EMusicPlayerTransportState::Killing:
-				MidiClockOutPin->AddTransportStateChangeToBlock({ 0, 0.0f, EMusicPlayerTransportState::Prepared });
-				return EMusicPlayerTransportState::Prepared;
-
-			default:
-				checkNoEntry();
-				return EMusicPlayerTransportState::Invalid;
+				break;
 			}
+
+			const EMusicPlayerTransportState NextState = GetNextTransportState(CurrentState);
+			HandleTransportChange(0, NextState);
+			return NextState;
 		};
 		
 		FMusicTransportControllable::Init(*TransportInPin, MoveTemp(InitFn));
@@ -379,18 +289,25 @@ namespace HarmonixMetasound
 			LastClockTickUpdate = ClockTick;
 		}
 
+		TransportSpanPostProcessor HandleMidiClockEvents = [this](int32 StartFrameIndex, int32 EndFrameIndex, EMusicPlayerTransportState CurrentState)
+		{
+			int32 NumFrames = EndFrameIndex - StartFrameIndex;
+			HandleTransportChange(StartFrameIndex, CurrentState);
+			if (MidiClockOutPin->DoesLoop())
+			{
+				MetronomeClock.Process(StartFrameIndex, NumFrames, SeekPreRollBarsInPin, *SpeedMultInPin);
+				MidiClockOutPin->Process(MetronomeClock, StartFrameIndex, NumFrames, SeekPreRollBarsInPin, *SpeedMultInPin);
+			}
+			else
+			{
+				MidiClockOutPin->Process(StartFrameIndex, NumFrames, SeekPreRollBarsInPin, *SpeedMultInPin);
+			}
+		};
+
 		TransportSpanProcessor TransportHandler = [this, &DrivingMidiClock](int32 StartFrameIndex, int32 EndFrameIndex, EMusicPlayerTransportState CurrentState)
 		{
 			switch (CurrentState)
 			{
-			case EMusicPlayerTransportState::Invalid:
-			case EMusicPlayerTransportState::Preparing:
-				AddTransportStateChangeToMidiClock(StartFrameIndex, EMusicPlayerTransportState::Prepared);
-				return EMusicPlayerTransportState::Prepared;
-
-			case EMusicPlayerTransportState::Prepared:
-				return CurrentState;
-
 			case EMusicPlayerTransportState::Starting:
 				// Play from the beginning if we haven't received a seek call while we were stopped...
 				if (!ReceivedSeekWhileStopped())
@@ -399,13 +316,8 @@ namespace HarmonixMetasound
 					LastClockTickUpdate = -1;
 				}
 				DrivingMidiClock.ResetAndStart(StartFrameIndex, !ReceivedSeekWhileStopped());
-				DrivingMidiClock.WriteAdvance(StartFrameIndex, EndFrameIndex, *SpeedMultInPin);
 				return EMusicPlayerTransportState::Playing;
-
-			case EMusicPlayerTransportState::Playing:
-				DrivingMidiClock.WriteAdvance(StartFrameIndex, EndFrameIndex, *SpeedMultInPin);
-				return EMusicPlayerTransportState::Playing;
-
+				
 			case EMusicPlayerTransportState::Seeking:
 				BuildMidiData(false);
 				DrivingMidiClock.SeekTo(StartFrameIndex, TransportInPin->GetNextSeekDestination(), SeekPreRollBarsInPin);
@@ -413,33 +325,11 @@ namespace HarmonixMetasound
 				// Here we will return that we want to be in the same state we were in before this request to 
 				// seek since we can seek "instantaneously"...
 				return GetTransportState();
-
-			case EMusicPlayerTransportState::Continuing:
-				AddTransportStateChangeToMidiClock(StartFrameIndex, EMusicPlayerTransportState::Playing);
-				DrivingMidiClock.WriteAdvance(StartFrameIndex, EndFrameIndex, *SpeedMultInPin);
-				return EMusicPlayerTransportState::Playing;
-
-			case EMusicPlayerTransportState::Pausing:
-				AddTransportStateChangeToMidiClock(StartFrameIndex, EMusicPlayerTransportState::Paused);
-				return EMusicPlayerTransportState::Paused;
-
-			case EMusicPlayerTransportState::Paused:
-				return EMusicPlayerTransportState::Paused;
-
-			case EMusicPlayerTransportState::Stopping:
-				AddTransportStateChangeToMidiClock(StartFrameIndex, EMusicPlayerTransportState::Prepared);
-				return EMusicPlayerTransportState::Prepared;
-
-			case EMusicPlayerTransportState::Killing:
-				AddTransportStateChangeToMidiClock(StartFrameIndex, EMusicPlayerTransportState::Prepared);
-				return EMusicPlayerTransportState::Prepared;
-
-			default:
-				checkNoEntry();
-				return EMusicPlayerTransportState::Invalid;
 			}
+
+			return GetNextTransportState(CurrentState);
 		};
-		ExecuteTransportSpans(TransportInPin, BlockSize, TransportHandler);
+		ExecuteTransportSpans(TransportInPin, BlockSize, TransportHandler, HandleMidiClockEvents);
 
 		if (LoopInPin)
 		{
@@ -550,13 +440,12 @@ namespace HarmonixMetasound
 		}
 	}
 
-	void FMetronomeOperator::AddTransportStateChangeToMidiClock(int32 StartFrameIndex, EMusicPlayerTransportState NewTransportState)
+	void FMetronomeOperator::HandleTransportChange(int32 StartFrameIndex, EMusicPlayerTransportState NewTransportState)
 	{
-		FMidiTimestampTransportState NewState{ StartFrameIndex, 0.0f, NewTransportState };
-		MidiClockOutPin->AddTransportStateChangeToBlock(NewState);
+		MidiClockOutPin->HandleTransportChange(StartFrameIndex, NewTransportState);
 		if (LoopInPin)
 		{
-			MetronomeClock.AddTransportStateChangeToBlock(NewState);
+			MetronomeClock.HandleTransportChange(StartFrameIndex, NewTransportState);
 		}
 	}
 }
diff --git a/Engine/Plugins/Runtime/Harmonix/Source/HarmonixMetasound/Private/HarmonixMetasound/Nodes/MidiPlayerNode.cpp b/Engine/Plugins/Runtime/Harmonix/Source/HarmonixMetasound/Private/HarmonixMetasound/Nodes/MidiPlayerNode.cpp
--- a/Engine/Plugins/Runtime/Harmonix/Source/HarmonixMetasound/Private/HarmonixMetasound/Nodes/MidiPlayerNode.cpp
+++ b/Engine/Plugins/Runtime/Harmonix/Source/HarmonixMetasound/Private/HarmonixMetasound/Nodes/MidiPlayerNode.cpp
@@ -381,109 +381,18 @@ namespace HarmonixMetasound
 
 		MidiClockOut->CopySpeedAndTempoChanges(MidiClockIn.Get(), *SpeedMultInPin);
 
-		int32 CurrentMidiClockEventIndex = -1;
-		TransportSpanPostProcessor HandleMidiClockEventsInBlock = [&](int32 StartFrameIndex, int32 EndFrameIndex, EMusicPlayerTransportState CurrentState)
+		TransportSpanPostProcessor HandleMidiClockEventsInBlock = [this](int32 StartFrameIndex, int32 EndFrameIndex, EMusicPlayerTransportState CurrentState)
 		{
-			const TArray<FMidiClockEvent>& MidiClockEvents = MidiClockIn->GetMidiClockEventsInBlock();
-			for (int32 EventIndex = FMath::Max(CurrentMidiClockEventIndex, 0); EventIndex < MidiClockEvents.Num(); ++EventIndex)
-			{
-				const FMidiClockEvent& Event = MidiClockEvents[EventIndex];
-
-				// we've run past the end of the events. Might as well stop iterating.
-				if (Event.BlockFrameIndex >= EndFrameIndex)
-				{
-					break;
-				}
-
-				if (Event.BlockFrameIndex >= StartFrameIndex && Event.BlockFrameIndex < EndFrameIndex)
-				{
-					// We should be advancing midi clock events in order (never repeating)!
-					check(CurrentMidiClockEventIndex < EventIndex);
-					CurrentMidiClockEventIndex = EventIndex;
-					switch (Event.Msg.Type)
-					{
-					case FMidiClockMsg::EType::Reset:
-					{
-						int32 Tick = MidiClockOut->CalculateMappedTick(Event.Msg.ToTick());
-						MidiClockOut->SeekTo(Event.BlockFrameIndex, Tick + 1, PrerollBars);
-						break;
-					}
-					case FMidiClockMsg::EType::Loop:
-					{
-						// loops should actually be handled by seeking and advancing...
-						break;
-					}
-					case FMidiClockMsg::EType::SeekTo:
-					{
-						int32 Tick = MidiClockOut->CalculateMappedTick(Event.Msg.ToTick());
-						MidiClockOut->SeekTo(Event.BlockFrameIndex, Tick, PrerollBars);
-						break;
-					}
-					case FMidiClockMsg::EType::SeekThru:
-					{
-						int32 Tick = MidiClockOut->CalculateMappedTick(Event.Msg.ThruTick());
-						MidiClockOut->SeekTo(Event.BlockFrameIndex, Tick + 1, PrerollBars);
-						break;
-					}
-					case FMidiClockMsg::EType::AdvanceThru:
-					{
-						// if this advance is a preroll, perform a seek instead so we don't trigger events doing an advance
-						// NOTE: This code is/should be identical to the SeekThru event above
-						if (Event.Msg.AsAdvanceThru().IsPreRoll)
-						{
-							int32 Tick = MidiClockOut->CalculateMappedTick(Event.Msg.ThruTick());
-							MidiClockOut->SeekTo(Event.BlockFrameIndex, Tick + 1, PrerollBars);
-							break;
-						}
-						// The MidiClock handles looping on its own, so we can conveniently advance it
-						int32 Tick = MidiClockOut->GetCurrentMidiTick() + (Event.Msg.ThruTick() - Event.Msg.FromTick());
-						float Ms = MidiClockOut->GetSongMaps().TickToMs(Tick);
-
-						float ClockInSpeed = MidiClockIn->GetSpeedAtBlockSampleFrame(StartFrameIndex);
-						float AdvanceRatio = MidiClockIn->GetSongMaps().GetTempoAtTick(Event.Msg.FromTick())
-							               / MidiClockOut->GetSongMaps().GetTempoAtTick(MidiClockOut->GetCurrentMidiTick());
-						// midi clock needs to know how fast its advancing based on their authority
-						MidiClockOut->InformOfCurrentAdvanceRate(ClockInSpeed * *SpeedMultInPin * AdvanceRatio);
-						MidiClockOut->AdvanceHiResToMs(Event.BlockFrameIndex, Ms, true);
-							
-						break;
-					}
-					}
-				}
-			}
+			// clock should always process in post processor
+			int32 NumFrames = EndFrameIndex - StartFrameIndex;
+			MidiClockOut->HandleTransportChange(StartFrameIndex, CurrentState);
+			MidiClockOut->Process(*MidiClockIn, StartFrameIndex, NumFrames, PrerollBars, *SpeedMultInPin);
 		};
 
 		TransportSpanProcessor TransportHandler = [&](int32 StartFrameIndex, int32 EndFrameIndex, EMusicPlayerTransportState CurrentState)
 		{
-			CurrentBlockSpanStart = StartFrameIndex;
 			switch (CurrentState)
 			{
-			case EMusicPlayerTransportState::Invalid:
-			case EMusicPlayerTransportState::Preparing:
-				return EMusicPlayerTransportState::Prepared;
-
-			case EMusicPlayerTransportState::Prepared:
-				return CurrentState;
-
-			case EMusicPlayerTransportState::Starting:
-				return EMusicPlayerTransportState::Playing;
-
-			case EMusicPlayerTransportState::Playing:
-				return EMusicPlayerTransportState::Playing;
-
-			case EMusicPlayerTransportState::Seeking:
-				// We have nothing to do because our incoming midi clock would have seeked us. 
-				return GetTransportState();
-
-			case EMusicPlayerTransportState::Continuing:
-				return EMusicPlayerTransportState::Playing;
-
-			case EMusicPlayerTransportState::Pausing:
-				return EMusicPlayerTransportState::Paused;
-
-			case EMusicPlayerTransportState::Paused:
-				return EMusicPlayerTransportState::Paused;
-
 			case EMusicPlayerTransportState::Stopping:
 				{
 					FMidiStreamEvent MidiEvent(this, FMidiMsg::CreateAllNotesOff());
@@ -493,15 +402,9 @@ namespace HarmonixMetasound
 					MidiEvent.TrackIndex = 0;
 					MidiOutPin->AddMidiEvent(MidiEvent);
 				}
-				return EMusicPlayerTransportState::Prepared;
-
-			case EMusicPlayerTransportState::Killing:
-				return EMusicPlayerTransportState::Prepared;
-
-			default:
-				checkNoEntry();
-				return EMusicPlayerTransportState::Invalid;
+				break;
 			}
+			return GetNextTransportState(CurrentState);
 		};
 		ExecuteTransportSpans(TransportInPin, BlockSize, TransportHandler, HandleMidiClockEventsInBlock);
 	}
@@ -514,30 +417,12 @@ namespace HarmonixMetasound
 		{
 			switch (CurrentState)
 			{
-			case EMusicPlayerTransportState::Invalid:
-			case EMusicPlayerTransportState::Preparing:
-			case EMusicPlayerTransportState::Prepared:
-			case EMusicPlayerTransportState::Stopping:
-			case EMusicPlayerTransportState::Killing:
-				return EMusicPlayerTransportState::Prepared;
-
-			case EMusicPlayerTransportState::Starting:
-			case EMusicPlayerTransportState::Playing:
-			case EMusicPlayerTransportState::Continuing:
-				return EMusicPlayerTransportState::Playing;
-
 			case EMusicPlayerTransportState::Seeking: // seeking is omitted from init, shouldn't happen
 				checkNoEntry();
 				return EMusicPlayerTransportState::Invalid;
-
-			case EMusicPlayerTransportState::Pausing:
-			case EMusicPlayerTransportState::Paused:
-				return EMusicPlayerTransportState::Paused;
-
-			default:
-				checkNoEntry();
-				return EMusicPlayerTransportState::Invalid;
 			}
+
+			return GetNextTransportState(CurrentState);
 		};
 		Init(*TransportInPin, MoveTemp(InitFn));
 	}
@@ -550,87 +435,23 @@ namespace HarmonixMetasound
 
 		TransportSpanPostProcessor MidiClockTransportHandler = [&](int32 StartFrameIndex, int32 EndFrameIndex, EMusicPlayerTransportState CurrentState)
 		{
-			switch (CurrentState)
-			{
-			case EMusicPlayerTransportState::Playing:
-				MidiClockOut->WriteAdvance(StartFrameIndex, EndFrameIndex, *SpeedMultInPin);
-				return;
-
-			case EMusicPlayerTransportState::Continuing:
-				MidiClockOut->WriteAdvance(StartFrameIndex, EndFrameIndex, *SpeedMultInPin);
-				return;
-			}
+			int32 NumFrames = EndFrameIndex - StartFrameIndex;
+			MidiClockOut->HandleTransportChange(StartFrameIndex, CurrentState);
+			MidiClockOut->Process(StartFrameIndex, NumFrames, PrerollBars, *SpeedMultInPin);
 		};
 
 		TransportSpanProcessor TransportHandler = [this](int32 StartFrameIndex, int32 EndFrameIndex, EMusicPlayerTransportState CurrentState)
 		{
 			switch (CurrentState)
 			{
-			case EMusicPlayerTransportState::Invalid:
-			case EMusicPlayerTransportState::Preparing:
-				// midi clock out
-				MidiClockOut->AddTransportStateChangeToBlock({ StartFrameIndex, 0.0f, EMusicPlayerTransportState::Prepared });
-
-				// midi out
-				return EMusicPlayerTransportState::Prepared;
-
-			case EMusicPlayerTransportState::Prepared:
-				// midi out
-				return CurrentState;
-
 			case EMusicPlayerTransportState::Starting:
-				// midi clock out 
 				MidiClockOut->ResetAndStart(StartFrameIndex, !ReceivedSeekWhileStopped());
-
-				// midi out
-				return EMusicPlayerTransportState::Playing;
-
-			case EMusicPlayerTransportState::Playing:
-				return EMusicPlayerTransportState::Playing;
-
+				break;
 			case EMusicPlayerTransportState::Seeking:
-				// midi clock out
 				MidiClockOut->SeekTo(StartFrameIndex, TransportInPin->GetNextSeekDestination(), PrerollBars);
-
-				// We have nothing to do because our incoming midi clock would have seeked us. 
-				return GetTransportState();
-
-			case EMusicPlayerTransportState::Continuing:
-				// midi clock
-				MidiClockOut->AddTransportStateChangeToBlock({ StartFrameIndex, 0.0f, EMusicPlayerTransportState::Playing });
-
-				// midi out
-				return EMusicPlayerTransportState::Playing;
-
-			case EMusicPlayerTransportState::Pausing:
-				// midi clock out
-				MidiClockOut->AddTransportStateChangeToBlock({ StartFrameIndex, 0.0f, EMusicPlayerTransportState::Paused });
-
-				// midi out
-				return EMusicPlayerTransportState::Paused;
-
-			case EMusicPlayerTransportState::Paused:
-				// midi out
-				return EMusicPlayerTransportState::Paused;
-
-			case EMusicPlayerTransportState::Stopping:
-				// midi clock out
-				MidiClockOut->AddTransportStateChangeToBlock({ StartFrameIndex, 0.0f, EMusicPlayerTransportState::Prepared });
-
-				// midi out
-				return EMusicPlayerTransportState::Prepared;
-
-			case EMusicPlayerTransportState::Killing:
-				// midi clock out
-				MidiClockOut->AddTransportStateChangeToBlock({ StartFrameIndex, 0.0f, EMusicPlayerTransportState::Prepared });
-
-				// midi out
-				return EMusicPlayerTransportState::Prepared;
-
-			default:
-				checkNoEntry();
-				return EMusicPlayerTransportState::Invalid;
+				break;
 			}
+			return GetNextTransportState(CurrentState);
 		};
 		ExecuteTransportSpans(TransportInPin, BlockSize, TransportHandler, MidiClockTransportHandler);
 	}
diff --git a/Engine/Plugins/Runtime/Harmonix/Source/HarmonixMetasound/Public/HarmonixMetasound/DataTypes/MidiClock.h b/Engine/Plugins/Runtime/Harmonix/Source/HarmonixMetasound/Public/HarmonixMetasound/DataTypes/MidiClock.h
--- a/Engine/Plugins/Runtime/Harmonix/Source/HarmonixMetasound/Public/HarmonixMetasound/DataTypes/MidiClock.h
+++ b/Engine/Plugins/Runtime/Harmonix/Source/HarmonixMetasound/Public/HarmonixMetasound/DataTypes/MidiClock.h
@@ -179,8 +179,25 @@ namespace HarmonixMetasound
 
 		TSharedPtr<FMidiPlayCursorMgr> GetDrivingMidiPlayCursorMgr() const { return DrivingMidiPlayCursorMgr.ToSharedPtr(); }
 		//*****************************************************************************************
+		
+		// handle and add transport change event
+		// to be called within the Transport Post Processor with the new transport state after the normal Processor
+		void HandleTransportChange(int32 BlockFrameIndex, EMusicPlayerTransportState TransportState);
+
+		// handle single clock event
+		void HandleClockEvent(const FMidiClock& DrivingClock, const FMidiClockEvent& Event, int32 PrerollBars, float Speed = 1.0f);
+		
+		// process and advance the clock based on the driving clock given sample frames
+		// will handle the driving clock events based on the frame range
+		void Process(const FMidiClock& DrivingClock, int32 StartFrame, int32 NumFrames, int32 PrerollBars, float Speed = 1.0f);
 
+		// process and advance the clock normally based on the given sample frames
+		void Process(int32 StartFrame, int32 NumFrames, int32 PrerollBars, float Speed = 1.0f);
+		
+		// directly perform and write an advance to this clock
 		void WriteAdvance(int32 StartFrameIndex, int32 EndFrameIndex, float InSpeed = 1.0f);
+
+		// directly seek this clock with a musical seek target or a specific tick
 		void SeekTo(int32 BlockFrameIndex, const FMusicSeekTarget& InTarget, int32 InPrerollBars);
 		void SeekTo(int32 BlockFrameIndex, int32 Tick, int32 InPrerollBars);
 		
diff --git a/Engine/Plugins/Runtime/Harmonix/Source/HarmonixMetasound/Public/HarmonixMetasound/DataTypes/MusicTransport.h b/Engine/Plugins/Runtime/Harmonix/Source/HarmonixMetasound/Public/HarmonixMetasound/DataTypes/MusicTransport.h
--- a/Engine/Plugins/Runtime/Harmonix/Source/HarmonixMetasound/Public/HarmonixMetasound/DataTypes/MusicTransport.h
+++ b/Engine/Plugins/Runtime/Harmonix/Source/HarmonixMetasound/Public/HarmonixMetasound/DataTypes/MusicTransport.h
@@ -90,6 +90,7 @@ namespace HarmonixMetasound
 
 		EMusicPlayerTransportState GetTransportState() const { return TransportState; }
 		void SetTransportState(EMusicPlayerTransportState NewState) { TransportState = NewState; }
+		EMusicPlayerTransportState GetNextTransportState(EMusicPlayerTransportState DesiredState) const;
 		
 		using TransportSpanProcessor = TUniqueFunction<EMusicPlayerTransportState(int32,int32,EMusicPlayerTransportState)>;
 		// post processor gets called immediately after the SpanProcessor so that the TransportState has been updated
diff --git a/Engine/Plugins/Runtime/Harmonix/Source/HarmonixMetasoundTests/Private/DataTypes/MidiClock.spec.cpp b/Engine/Plugins/Runtime/Harmonix/Source/HarmonixMetasoundTests/Private/DataTypes/MidiClock.spec.cpp
--- a/Engine/Plugins/Runtime/Harmonix/Source/HarmonixMetasoundTests/Private/DataTypes/MidiClock.spec.cpp
+++ b/Engine/Plugins/Runtime/Harmonix/Source/HarmonixMetasoundTests/Private/DataTypes/MidiClock.spec.cpp
@@ -8,16 +8,19 @@
 
 namespace HarmonixMetasoundTests::MidiClock
 {
+	using namespace HarmonixMetasound;
 	BEGIN_DEFINE_SPEC(
 		FHarmonixMetasoundMidiClockSpec,
 		"Harmonix.Metasound.DataTypes.MidiClock",
 		EAutomationTestFlags::EditorContext | EAutomationTestFlags::EngineFilter)
 
-	TUniquePtr<HarmonixMetasound::FMidiClock> TestClock;
+	TUniquePtr<FMidiClock> TestClock;
+	TUniquePtr<FMidiClock> DrivingClock;
+	Metasound::FOperatorSettings OperatorSettings {48000, 100};
 
-	void AddStateAtFrame(HarmonixMetasound::EMusicPlayerTransportState State, int32 Frame) const
+	void AddStateAtFrame(EMusicPlayerTransportState State, int32 Frame) const
 	{
-		const HarmonixMetasound::FMidiTimestampTransportState NewState
+		const FMidiTimestampTransportState NewState
 		{
 			Frame,
 			0,
@@ -28,78 +31,367 @@ namespace HarmonixMetasoundTests::MidiClock
 		TestClock->AddTransportStateChangeToBlock(NewState);
 	}
 
+	void ExecuteWriteAdvance(int32 StartFrameIndex, int32 EndFrameIndex, float Speed)
+	{
+		TestClock->SeekTo(0, 1, 0);
+		TestClock->PrepareBlock();
+
+		TestEqual("Clock.GetCurrentHiResMs()", TestClock->GetCurrentHiResMs(), 0.0f);
+		TestEqual("Clock.GetCurrentMidiTick()", TestClock->GetCurrentMidiTick(), 0);
+		
+		float OldMs = TestClock->GetCurrentHiResMs();
+		int32 OldTick = TestClock->GetCurrentMidiTick();
+		TestClock->WriteAdvance(StartFrameIndex, EndFrameIndex, Speed);
+		float NewMs = TestClock->GetCurrentHiResMs();
+		int32 NewTick = TestClock->GetCurrentMidiTick();
+
+		TestTrue("Non looping clock advanced forward in time", NewMs > OldMs);
+		TestTrue("Non looping Clock Advanced forward in ticks", NewTick > OldTick);
+
+		int32 DeltaFrames = FMidiClock::kMidiGranularity * Speed;
+		float DeltaMs = DeltaFrames * 1000.0f / OperatorSettings.GetSampleRate();
+		float Ms = OldMs;
+
+		int32 BlockFrameIndex = 0; 
+		for (const FMidiClockEvent& Event : TestClock->GetMidiClockEventsInBlock())
+		{
+			int32 Tick1 = (int32)(TestClock->GetTempoMap().MsToTick(Ms) + 0.5f);
+			Ms += DeltaMs;
+			int32 Tick2 = (int32)(TestClock->GetTempoMap().MsToTick(Ms) + 0.5f);
+
+			TestEqual(FString::Printf(TEXT("Frame-%d: Event.BlockFrameIndex"), Event.BlockFrameIndex), Event.BlockFrameIndex, BlockFrameIndex);
+			TestEqual(FString::Printf(TEXT("Frame-%d: Event.Type"), Event.BlockFrameIndex), Event.Msg.Type, FMidiClockMsg::EType::AdvanceThru);
+			TestEqual(FString::Printf(TEXT("Frame-%d: Event.FromTick"), Event.BlockFrameIndex), Event.Msg.FromTick(), Tick1);
+			TestEqual(FString::Printf(TEXT("Frame-%d: Event.ThruTick"), Event.BlockFrameIndex), Event.Msg.ThruTick(), Tick2);
+			TestEqual(FString::Printf(TEXT("Frame-%d: Event.IsPreRoll"), Event.BlockFrameIndex), Event.Msg.AsAdvanceThru().IsPreRoll, false);
+			
+			BlockFrameIndex += FMidiClock::kMidiGranularity;
+		}
+
+		int32 Tick = (int32)(TestClock->GetTempoMap().MsToTick(Ms) + 0.5f);
+		TestEqual("Clock.GetCurrentMidiTick()", TestClock->GetCurrentMidiTick(), Tick);
+		TestEqual("Clock.GetCurrentHiResMs()", TestClock->GetCurrentHiResMs(), Ms);
+	};
+	
 	END_DEFINE_SPEC(FHarmonixMetasoundMidiClockSpec)
 
 	void FHarmonixMetasoundMidiClockSpec::Define()
 	{
-		Describe("When a new transport state is added to the block", [this]()
+		BeforeEach([this]()
 		{
-			BeforeEach([this]()
-			{
-				TestClock = MakeUnique<HarmonixMetasound::FMidiClock>(Metasound::FOperatorSettings{ 48000, 100 });
-			});
-			
-			It("and the block has no other state changes, it should succeed", [this]()
+			TestClock = MakeUnique<FMidiClock>(OperatorSettings);
+		});
+
+		AfterEach([this]()
+		{
+			TestClock.Reset();
+		});
+		
+		Describe("AddTransportStateChangeToBlock(NewState)", [this]()
+		{
+			It("should always add NewState if there are no changes in the block", [this]()
 			{
 				TestClock->PrepareBlock();
 
 				TestFalse("No transport changes in block", TestClock->HasTransportStateChangeInBlock());
 
-				constexpr HarmonixMetasound::EMusicPlayerTransportState NewState = HarmonixMetasound::EMusicPlayerTransportState::Playing;
+				constexpr EMusicPlayerTransportState NewState = EMusicPlayerTransportState::Playing;
 				AddStateAtFrame(NewState, 0);
 
 				TestTrue("There is a transport state change in block", TestClock->HasTransportStateChangeInBlock());
 				TestEqual("State at end of block matches the one we added", TestClock->GetTransportStateAtEndOfBlock(), NewState);
 			});
 
-			It("and the new state happens after the latest one, it should succeed", [this]()
+			It("should add NewState if its frame is greater than the last one in the block", [this]()
 			{
-				AddStateAtFrame(HarmonixMetasound::EMusicPlayerTransportState::Playing, 0);
+				AddStateAtFrame(EMusicPlayerTransportState::Playing, 0);
 				TestTrue("There is already a transport state change in the block", TestClock->HasTransportStateChangeInBlock());
 
 				const int32 NumInitialStates = TestClock->GetTransportTimestampsInBlock().Num();
 				const int32 LastStateFrame = TestClock->GetTransportTimestampsInBlock().Last().BlockSampleFrameIndex;
 
-				constexpr HarmonixMetasound::EMusicPlayerTransportState NewState = HarmonixMetasound::EMusicPlayerTransportState::Pausing;
+				constexpr EMusicPlayerTransportState NewState = EMusicPlayerTransportState::Pausing;
 				AddStateAtFrame(NewState, LastStateFrame + 1);
 
 				TestEqual("There is another transport state change in block", TestClock->GetTransportTimestampsInBlock().Num(), NumInitialStates + 1);
 				TestEqual("State at end of block matches the one we added", TestClock->GetTransportStateAtEndOfBlock(), NewState);
 			});
 
-			It("and the new state happens at the same time as the latest one, it should succeed", [this]()
+			It("should add NewState if it has the same frame as the last one in the block", [this]()
 			{
-				AddStateAtFrame(HarmonixMetasound::EMusicPlayerTransportState::Playing, 0);
+				AddStateAtFrame(EMusicPlayerTransportState::Playing, 0);
 				TestTrue("There is already a transport state change in the block", TestClock->HasTransportStateChangeInBlock());
 
 				const int32 NumInitialStates = TestClock->GetTransportTimestampsInBlock().Num();
 				const int32 LastStateFrame = TestClock->GetTransportTimestampsInBlock().Last().BlockSampleFrameIndex;
 
-				constexpr HarmonixMetasound::EMusicPlayerTransportState NewState = HarmonixMetasound::EMusicPlayerTransportState::Paused;
+				constexpr EMusicPlayerTransportState NewState = EMusicPlayerTransportState::Paused;
 				AddStateAtFrame(NewState, LastStateFrame );
 
 				TestEqual("There is another transport state change in block", TestClock->GetTransportTimestampsInBlock().Num(), NumInitialStates + 1);
 				TestEqual("State at end of block matches the one we added", TestClock->GetTransportStateAtEndOfBlock(), NewState);
 			});
 
-			It("and the new state happens before the latest one, it should fail", [this]()
+			It("should not add NewState if the frame is less than the last one in the block", [this]()
 			{
-				AddStateAtFrame(HarmonixMetasound::EMusicPlayerTransportState::Playing, 0);
+				AddStateAtFrame(EMusicPlayerTransportState::Playing, 0);
 				TestTrue("There is already a transport state change in the block", TestClock->HasTransportStateChangeInBlock());
 
 				const int32 NumInitialStates = TestClock->GetTransportTimestampsInBlock().Num();
-				const HarmonixMetasound::EMusicPlayerTransportState LastState = TestClock->GetTransportStateAtEndOfBlock();
+				const EMusicPlayerTransportState LastState = TestClock->GetTransportStateAtEndOfBlock();
 				const int32 LastStateFrame = TestClock->GetTransportTimestampsInBlock().Last().BlockSampleFrameIndex;
 
-				constexpr HarmonixMetasound::EMusicPlayerTransportState NewState = HarmonixMetasound::EMusicPlayerTransportState::Continuing;
+				constexpr EMusicPlayerTransportState NewState = EMusicPlayerTransportState::Continuing;
 				AddStateAtFrame(NewState, LastStateFrame - 1);
 
 				TestEqual("There is not another transport state change in block", TestClock->GetTransportTimestampsInBlock().Num(), NumInitialStates);
 				TestEqual("State at end of block matches the one that was already there", TestClock->GetTransportStateAtEndOfBlock(), LastState);
 			});
 		});
+
+		Describe("PrepareBlock()", [this]()
+		{
+			It("should reset transport, speed, tempo, and clock states", [this]()
+			{
+				const float LastTempo = TestClock->GetTempoAtEndOfBlock();
+				const float LastSpeed = TestClock->GetSpeedAtEndOfBlock();
+				TestClock->PrepareBlock();
+				TestEqual("Clock.TransportChangesInBlock.Num()", TestClock->GetTransportTimestampsInBlock().Num(), 0);
+				TestEqual("Clock.SpeedChangesInBlock.Num()", TestClock->GetTempoSpeedTimestampsInBlock().Num(), 1);
+				TestEqual("Clock.SpeedAtBlockSampleFrame(0)", TestClock->GetSpeedAtBlockSampleFrame(0), LastSpeed);
+				TestEqual("Clock.SpeedAtEndOfBlock()", TestClock->GetSpeedAtEndOfBlock(), LastSpeed);
+				TestFalse("Clock.HasSpeedChangesInBlock()", TestClock->HasSpeedChangesInBlock());
+				TestEqual("Clock.TempoChangesInBlock.Num()", TestClock->GetNumTempoChangesInBlock(), 1);
+				TestEqual("Clock.TempoAtBlockSampleFrame(0)", TestClock->GetTempoAtBlockSampleFrame(0), LastTempo);
+				TestEqual("Clock.GetTempoAtEndOfBlock()", TestClock->GetTempoAtEndOfBlock(), LastTempo);
+				TestFalse("Clock.HasTempoChangesInBlock()", TestClock->HasTempoChangesInBlock());
+				TestEqual("Clock.GetMidiClockEventsInBlock().Num()", TestClock->GetMidiClockEventsInBlock().Num(), 0);
+			});
+		});
+
+		Describe("ResetAndStart()", [this]()
+		{
+			It("should reset speed and tempo changes, and be \"playing\"", [this]()
+			{
+				const float LastTempo = TestClock->GetTempoAtEndOfBlock();
+				const float LastSpeed = TestClock->GetSpeedAtEndOfBlock();
+				TestClock->ResetAndStart(0, false);
+				TestEqual("Clock.GetTransportStateAtEndOfBlock()", TestClock->GetTransportStateAtEndOfBlock(), EMusicPlayerTransportState::Playing);
+				TestEqual("Clock.SpeedAtBlockSampleFrame(0)", TestClock->GetSpeedAtBlockSampleFrame(0), LastSpeed);
+				TestEqual("Clock.SpeedAtEndOfBlock()", TestClock->GetSpeedAtEndOfBlock(), LastSpeed);
+				TestFalse("Clock.HasSpeedChangesInBlock()", TestClock->HasSpeedChangesInBlock());
+				TestEqual("Clock.TempoChangesInBlock.Num()", TestClock->GetNumTempoChangesInBlock(), 1);
+				TestEqual("Clock.TempoAtBlockSampleFrame(0)", TestClock->GetTempoAtBlockSampleFrame(0), LastTempo);
+				TestEqual("Clock.GetTempoAtEndOfBlock()", TestClock->GetTempoAtEndOfBlock(), LastTempo);
+				TestFalse("Clock.HasTempoChangesInBlock()", TestClock->HasTempoChangesInBlock());
+			});
+			
+			It("should be at the correct block frame index", [this]()
+			{
+				const float LastTempo = TestClock->GetTempoAtEndOfBlock();
+				const float LastSpeed = TestClock->GetSpeedAtEndOfBlock();
+				int32 BlockFrame = 100;
+				TestClock->ResetAndStart(BlockFrame, false);
+				TestEqual("Clock.Get", TestClock->GetCurrentBlockFrameIndex(), BlockFrame);
+			});
+			
+			It("should seek to start when requested", [this]()
+			{
+				int32 BlockFrame = 100;
+				FMusicSeekTarget SeekTarget;
+				SeekTarget.Type = ESeekPointType::BarBeat;
+				SeekTarget.BarBeat = FMusicTimestamp(2, 1.0f);
+				TestClock->SeekTo(SeekTarget, 0);
+				TestFalse("Clock seeked, Clock.GetCurrentMidiTick() == 0", TestClock->GetCurrentMidiTick() == 0);
+				
+				TestClock->ResetAndStart(BlockFrame, true);
+				TestEqual("Clock.GetCurrentBlockFrameIndex()", TestClock->GetCurrentBlockFrameIndex(), BlockFrame);
+				TestEqual("Clock.GetCurrentMidiTick()", TestClock->GetCurrentMidiTick(), -1);
+			});
+
+			It("should otherwise not seek to start", [this]()
+			{
+				int32 BlockFrame = 100;
+				FMusicSeekTarget SeekTarget;
+				SeekTarget.Type = ESeekPointType::BarBeat;
+				SeekTarget.BarBeat = FMusicTimestamp(2, 1.0f);
+				TestClock->SeekTo(SeekTarget, 0);
+				int32 NewMidiTick = TestClock->GetCurrentMidiTick();
+				TestFalse("Clock seeked, Clock.GetCurrentMidiTick() == 0", TestClock->GetCurrentMidiTick() == 0);
+							
+				TestClock->ResetAndStart(BlockFrame, false);
+				TestEqual("Clock.GetCurrentBlockFrameIndex()", TestClock->GetCurrentBlockFrameIndex(), BlockFrame);
+				TestEqual("Clock.GetCurrentMidiTick()", TestClock->GetCurrentMidiTick(), NewMidiTick);
+			});
+			
+		});
+
+		Describe("SetLoop()", [this]()
+		{
+			It("should correctly set the tempo", [this]()
+			{
+				const float TicksPerQuarterNote = Harmonix::Midi::Constants::GTicksPerQuarterNote;
+				const float TempoBPM = 120.0f;
+				int32 LoopStartTick = 0;
+				int32 LoopEndTick = 1000;
+
+				// initial clock tempo
+				TestEqual("Clock.Tempo", TestClock->GetTempoAtEndOfBlock(), TempoBPM);
+				TestEqual("Clock.Tempo", TestClock->GetTempoMap().GetTempoAtTick(0), TempoBPM);
+
+				// microseconds per quarternote
+				int32 MidiTempo = Harmonix::Midi::Constants::BPMToMidiTempo(TempoBPM);
+							
+				float MsPerTick = MidiTempo / TicksPerQuarterNote / 1000.0f;
+				float LoopStartMs = MsPerTick * LoopStartTick;
+				float LoopEndMs = MsPerTick * LoopEndTick;
+
+				TestFalse("Initial -> Clock.DoesLoop()", TestClock->DoesLoop());
+				TestClock->SetLoop(LoopStartTick, LoopEndTick);
+				TestTrue("SetLoop -> Clock.DoesLoop()", TestClock->DoesLoop());
+				TestEqual("Clock.LoopStartTick()", TestClock->GetLoopStartTick(), LoopStartTick);
+				TestEqual("Clock.LoopEndTick()", TestClock->GetLoopEndTick(), LoopEndTick);
+				TestEqual("Clock.LoopStartMs()", TestClock->GetLoopStartMs(), LoopStartMs);
+				TestEqual("Clock.LoopEndMs()", TestClock->GetLoopEndMs(), LoopEndMs);
+
+				TestClock->ClearLoop();
+				TestFalse("Cleared -> Clock.DoesLoop()", TestClock->DoesLoop());
+			});
+
+		});
+
+		Describe("WriteAdvance", [this]()
+		{
+			It("should advance one block correctly with speed 1", [this]
+			{
+				int32 StartFrame = 0;
+				int32 EndFrame = StartFrame + OperatorSettings.GetNumFramesPerBlock();
+				float Speed = 1.0f;
+				ExecuteWriteAdvance(0, EndFrame, Speed);
+			});
+
+			It("should advance one block correctly with speed 2", [this]
+			{
+				int32 StartFrame = 0;
+				int32 EndFrame = StartFrame + OperatorSettings.GetNumFramesPerBlock();
+				float Speed = 2.0f;
+				ExecuteWriteAdvance(0, EndFrame, Speed);
+			});
+			
+			It("should advance one block correctly with speed 1/2", [this]
+			{
+				int32 StartFrame = 0;
+				int32 EndFrame = StartFrame + OperatorSettings.GetNumFramesPerBlock();
+				float Speed = 0.5f;
+				ExecuteWriteAdvance(0, EndFrame, Speed);
+			});
+			
+		});
+
+		Describe("ProcessClockEvent(EventType)", [this]()
+		{
+			BeforeEach([&, this]
+			{
+				DrivingClock = MakeUnique<FMidiClock>(OperatorSettings);
+				TestClock->AttachToTimeAuthority(*DrivingClock);
+			});
+
+			AfterEach([&, this]
+			{
+				TestClock->ClearLoop();
+				TestClock->DetachFromTimeAuthority();
+				DrivingClock.Reset();
+			});
+
+			It("EventType::AdvanceThru.NonLooping", [&, this]
+			{
+				float DrivingClockSpeed = 1.0f;
+				int32 StartFrame = 0;
+				int32 DeltaFrames = FMidiClock::kMidiGranularity * DrivingClockSpeed;
+				float DeltaMs = DeltaFrames * 1000.0f / OperatorSettings.GetSampleRate();
+				int32 Tick = (int32)(TestClock->GetTempoMap().MsToTick(DeltaMs) + 0.5f);
+
+				float Speed = 1.0f;
+				int32 PrerollBars = 8;
+				FMidiClockEvent Event = FMidiClockEvent(StartFrame, FMidiClockMsg::FAdvanceThru(0, Tick, false));
+
+				TestClock->PrepareBlock();
+				int32 OldEventsNum = TestClock->GetMidiClockEventsInBlock().Num();
+				TestClock->HandleClockEvent(*DrivingClock, Event, PrerollBars, Speed);
+
+				if (!TestTrue("Clock has new clock events", TestClock->GetMidiClockEventsInBlock().Num() > OldEventsNum))
+				{
+					return;
+				}
+				TestEqual("Last Clock Event in block", TestClock->GetMidiClockEventsInBlock().Last().Msg.Type, FMidiClockMsg::EType::AdvanceThru);
+				TestEqual("Clock Current Tick", TestClock->GetCurrentMidiTick(), Event.Msg.ThruTick() - 1);
+			});
+
+			It("EventType::AdvanceThru.Looping", [&, this]
+			{
+				FMusicTimestamp LoopEndTimestamp;
+				LoopEndTimestamp.Bar = 2;
+				LoopEndTimestamp.Beat = 1.0f;
+				
+				int32 LoopStartTick = 0;
+				int32 LoopEndTick = TestClock->GetBarMap().MusicTimestampToTick(LoopEndTimestamp);
+				TestClock->SetLoop(LoopStartTick, LoopEndTick);
+				TestClock->SeekTo(0, LoopEndTick - 1, 0);
+				
+				float DrivingClockSpeed = 1.0f;
+				int32 StartFrame = 0;
+				int32 DeltaFrames = FMidiClock::kMidiGranularity * DrivingClockSpeed;
+				float DeltaMs = DeltaFrames * 1000.0f / OperatorSettings.GetSampleRate();
+				int32 DeltaTicks = (int32)(TestClock->GetTempoMap().MsToTick(DeltaMs) + 0.5f);
+				int32 Tick = DeltaTicks;
+
+				int32 ExpectedTick = TestClock->GetCurrentMidiTick() + DeltaTicks + 1;
+				if (ExpectedTick > LoopEndTick)
+				{
+					ExpectedTick = ExpectedTick - LoopEndTick + LoopStartTick;
+				}
+
+				float Speed = 1.0f;
+				int32 PrerollBars = 8;
+				FMidiClockEvent Event = FMidiClockEvent(StartFrame, FMidiClockMsg::FAdvanceThru(0, Tick, false));
+
+				TestClock->PrepareBlock();
+				int32 OldEventsNum = TestClock->GetMidiClockEventsInBlock().Num();
+				TestClock->HandleClockEvent(*DrivingClock, Event, PrerollBars, Speed);
+
+				if (!TestTrue("Clock has new clock events", TestClock->GetMidiClockEventsInBlock().Num() > OldEventsNum))
+				{
+					return;
+				}
+				TestEqual("Last Clock Event in block", TestClock->GetMidiClockEventsInBlock().Last().Msg.Type, FMidiClockMsg::EType::AdvanceThru);
+				TestEqual("Clock Current Tick", TestClock->GetCurrentMidiTick(), ExpectedTick - 1);
+			});
+
+			It("ProccessClockEvent(SeekTo)", [&, this]
+			{
+				int32 StartFrame = 0;
+				int32 Tick = 1000;
+				float Speed = 1.0f;
+				int32 PrerollBars = 8;
+				FMidiClockEvent Event = FMidiClockEvent(StartFrame, FMidiClockMsg::FSeekTo(0, Tick));
+
+				TestClock->PrepareBlock();
+				int32 OldEventsNum = TestClock->GetMidiClockEventsInBlock().Num();
+				TestClock->HandleClockEvent(*DrivingClock, Event, PrerollBars, Speed);
+
+				if (!TestTrue("Clock has new clock events", TestClock->GetMidiClockEventsInBlock().Num() > OldEventsNum))
+				{
+					return;
+				}
+				TestEqual("Last Clock Event in block", TestClock->GetMidiClockEventsInBlock()[0].Msg.Type, FMidiClockMsg::EType::SeekThru);
+				TestEqual("Clock Current Tick", TestClock->GetCurrentMidiTick(), Tick - 1);
+			});
+			
+		});
+
 	}
 
+
 }
 
 #endif
\ No newline at end of file
