Use the NoteIndex instead of the NoteNumber from the StepSequence to associate midi events with.

#tests PIE, Creative, Metasound

diff --git a/Engine/Plugins/Runtime/Harmonix/Source/HarmonixMetasound/Private/HarmonixMetasound/Nodes/StepSequencePlayerNode.cpp b/Engine/Plugins/Runtime/Harmonix/Source/HarmonixMetasound/Private/HarmonixMetasound/Nodes/StepSequencePlayerNode.cpp
--- a/Engine/Plugins/Runtime/Harmonix/Source/HarmonixMetasound/Private/HarmonixMetasound/Nodes/StepSequencePlayerNode.cpp
+++ b/Engine/Plugins/Runtime/Harmonix/Source/HarmonixMetasound/Private/HarmonixMetasound/Nodes/StepSequencePlayerNode.cpp
@@ -760,28 +760,30 @@ namespace HarmonixMetasound
 				{
 					if (CurrentCellNotes[i])
 					{
+						int32 NoteIndex = i;
 						uint8 MidiCh;
 						uint8 MidiNote;
 						CurrentCellNotes[i].GetChannelAndNote(MidiCh, MidiNote);
-
-						int32 OriginalNote = SequenceTable->Notes[i].NoteNumber;
-						int32 TransposedNote = SequenceTable->Notes[i].NoteNumber + AdditionalOctaveNotes;
-
+						
 						// If: 
 						// 1. The note that would play is not enabled (it is not a new note);
 						// 2. A note with this pitch would start up this tick;
 						// 3. The note is marked as a continuation note
 						// Then keep this note playing through the next cell
-						if (MidiNote == OriginalNote && !CurrentPage->Rows[i].Cells[NextCellInRow].bEnabled && CurrentPage->Rows[i].Cells[NextCellInRow].bContinuation)
+						if (MidiNote == NoteIndex && !CurrentPage->Rows[i].Cells[NextCellInRow].bEnabled && CurrentPage->Rows[i].Cells[NextCellInRow].bContinuation)
 						{
 							continue;
 						}
 
 						// note off!
-						// create the midi event with the original note to maintain voice ids.
-						FMidiStreamEvent MidiEvent(CurrentCellNotes[i].GetGeneratorId(), FMidiMsg::CreateNoteOff(MidiCh, OriginalNote));
+						// use the "NoteIndex" as the note number for the midi event to track the midi event for our note off message
+						// the note number and transposition of the step sequencer can change out from under us, so we can't use those to id events
+						// we transpose the note later
+						int32 OriginalNote = SequenceTable->Notes[i].NoteNumber;
+						int32 TransposedNote = FMath::Clamp(OriginalNote + AdditionalOctaveNotes, 0, 127);
+						FMidiStreamEvent MidiEvent(CurrentCellNotes[i].GetGeneratorId(), FMidiMsg::CreateNoteOff(MidiCh, NoteIndex));
 						// and then assign the note directly to the midi message after
-						MidiEvent.MidiMessage.Data1 = FMath::Clamp(TransposedNote, 0, 127);
+						MidiEvent.MidiMessage.Data1 = TransposedNote;
 						MidiEvent.BlockSampleFrameIndex = BlockFrameIndex;
 						MidiEvent.AuthoredMidiTick = ProcessedThruTick;
 						MidiEvent.CurrentMidiTick = ProcessedThruTick;
@@ -802,12 +804,13 @@ namespace HarmonixMetasound
 						if (!CurrentCellNotes[i])
 						{
 							// note on!
+							int32 NoteIndex = i;
 							int32 OriginalNote = SequenceTable->Notes[i].NoteNumber;
-							int32 TransposedNote = SequenceTable->Notes[i].NoteNumber + AdditionalOctaveNotes;
+							int32 TransposedNote = FMath::Clamp(OriginalNote + AdditionalOctaveNotes, 0, 127);
 							// create the midi event with the original note to maintain voice ids.
-							FMidiStreamEvent MidiEvent(this, FMidiMsg::CreateNoteOn(0, OriginalNote, SequenceTable->Notes[i].Velocity));
+							FMidiStreamEvent MidiEvent(this, FMidiMsg::CreateNoteOn(0, NoteIndex, SequenceTable->Notes[i].Velocity));
 							// and then assign the note directly to the midi message after
-							MidiEvent.MidiMessage.Data1 = FMath::Clamp(TransposedNote, 0, 127);
+							MidiEvent.MidiMessage.Data1 = TransposedNote;
 							float NoteOnVelocity = FMath::Clamp(static_cast<float>(SequenceTable->Notes[i].Velocity) * CurrentVelocityMultiplierValue, 0.0f, 127.0f);
 							MidiEvent.MidiMessage.SetNoteOnVelocity(static_cast<uint8>(NoteOnVelocity));
 							MidiEvent.BlockSampleFrameIndex = BlockFrameIndex;
