Delete non-serialized, non-editable, cached value "FineTuneCentsAdjustment" which is the pre-calculated frequency given the original "FineTuneCents" of the keyzone.

Simply use the PitchShiftCents during the calculation of the frequency adjustment instead of caching the value, which introduces unneeded complexity to the loadind and saving of keyzones.

#tests PIE, Metasound, Creative

diff --git a/Engine/Plugins/Runtime/Harmonix/Source/HarmonixDsp/Private/HarmonixDsp/FusionSampler/FusionVoice.cpp b/Engine/Plugins/Runtime/Harmonix/Source/HarmonixDsp/Private/HarmonixDsp/FusionSampler/FusionVoice.cpp
--- a/Engine/Plugins/Runtime/Harmonix/Source/HarmonixDsp/Private/HarmonixDsp/FusionSampler/FusionVoice.cpp
+++ b/Engine/Plugins/Runtime/Harmonix/Source/HarmonixDsp/Private/HarmonixDsp/FusionSampler/FusionVoice.cpp
@@ -694,9 +694,13 @@ uint32 FFusionVoice::Process(uint32 InSliceIndex, uint32 InSubsliceIndex, float*
 		}
 	}
 
+	// apply the per keyzone FineTuneCents to the semitone bend.
+	// Convert cents to semitones (divide by 100)
+	SemiTonesBend += KeyZone->FineTuneCents / 100.0f;
+
+	// convert semitones to frequency
 	PitchBend *= FMath::Pow(2.0, SemiTonesBend / 12.0);
 	ResampleInc *= PitchBend;
-	ResampleInc *= KeyZone->FineTuneAdjustment;
 	check(ResampleInc > 0.0);  // for now we only support playing forward
 
 	OutputBuffer.SetAliasedChannelData(OutData, InNumChannels);
diff --git a/Engine/Plugins/Runtime/Harmonix/Source/HarmonixDsp/Private/HarmonixDsp/FusionSampler/Settings/KeyzoneSettings.cpp b/Engine/Plugins/Runtime/Harmonix/Source/HarmonixDsp/Private/HarmonixDsp/FusionSampler/Settings/KeyzoneSettings.cpp
--- a/Engine/Plugins/Runtime/Harmonix/Source/HarmonixDsp/Private/HarmonixDsp/FusionSampler/Settings/KeyzoneSettings.cpp
+++ b/Engine/Plugins/Runtime/Harmonix/Source/HarmonixDsp/Private/HarmonixDsp/FusionSampler/Settings/KeyzoneSettings.cpp
@@ -74,10 +74,4 @@ void FKeyzoneSettings::SetVolumeDb(float Db)
 float FKeyzoneSettings::GetVolumeDb() const
 {
 	return HarmonixDsp::LinearToDB(FMath::Clamp(Gain, 0.0f, 1.0f));
-}
-
-void FKeyzoneSettings::SetFineTuneCents(float InCents)
-{
-	FineTuneCents = InCents;
-	FineTuneAdjustment = FMath::Pow(2.0f, FineTuneCents / 1200.0f);
 }
\ No newline at end of file
diff --git a/Engine/Plugins/Runtime/Harmonix/Source/HarmonixDsp/Public/HarmonixDsp/FusionSampler/Settings/KeyzoneSettings.h b/Engine/Plugins/Runtime/Harmonix/Source/HarmonixDsp/Public/HarmonixDsp/FusionSampler/Settings/KeyzoneSettings.h
--- a/Engine/Plugins/Runtime/Harmonix/Source/HarmonixDsp/Public/HarmonixDsp/FusionSampler/Settings/KeyzoneSettings.h
+++ b/Engine/Plugins/Runtime/Harmonix/Source/HarmonixDsp/Public/HarmonixDsp/FusionSampler/Settings/KeyzoneSettings.h
@@ -40,10 +40,7 @@ struct HARMONIXDSP_API FKeyzoneSettings
 
 	GENERATED_BODY()
 
-	FKeyzoneSettings() 
-	{
-		SetFineTuneCents(FineTuneCents);
-	};
+	FKeyzoneSettings() {}
 
 	FKeyzoneSettings(const FKeyzoneArgs& InArgs)
 	{
@@ -66,8 +63,6 @@ struct HARMONIXDSP_API FKeyzoneSettings
 		TimeStretchConfig.PitchShifter = InArgs.PitchShifter;
 		TimeStretchConfig.bSyncTempo = InArgs.bSyncTempo;
 		TimeStretchConfig.OriginalTempo = InArgs.OriginalTempo;
-
-		SetFineTuneCents(FineTuneCents);
 	}
 
 public:
@@ -106,9 +101,6 @@ public:
 	
 	UPROPERTY(EditDefaultsOnly, Category="Settings", Meta = (UIMin = -100, UIMax = 100, ClampMin = -100, ClampMax = 100))
 	float FineTuneCents = 0.0f;
-
-	// This value shouldn't be serialized. It's updated on construction
-	float FineTuneAdjustment = 0.0f;
 	
 	UPROPERTY(EditDefaultsOnly, Category="Settings", Meta=(PostEditType="Trivial", UIMin="0.0", UIMax="1.0", ClampMin="0.0", ClampMax="1.0"))
 	float Gain = 1.0f;
@@ -145,7 +137,6 @@ public:
 	bool ContainsNoteAndVelocity(uint8 InNote, uint8 InVelocity) const;
 	void SetVolumeDb(float Db);
 	float GetVolumeDb() const;
-	void SetFineTuneCents(float InCents);
 	bool IsSingleton() const {  return SingletonFusionVoicePool != nullptr; }
 	bool IsNoteOnZone() const { return !bIsNoteOffZone; }
 
diff --git a/Engine/Plugins/Runtime/Harmonix/Source/HarmonixDspEditor/Private/SettingsJsonImporter.cpp b/Engine/Plugins/Runtime/Harmonix/Source/HarmonixDspEditor/Private/SettingsJsonImporter.cpp
--- a/Engine/Plugins/Runtime/Harmonix/Source/HarmonixDspEditor/Private/SettingsJsonImporter.cpp
+++ b/Engine/Plugins/Runtime/Harmonix/Source/HarmonixDspEditor/Private/SettingsJsonImporter.cpp
@@ -42,7 +42,6 @@ bool FSettingsJsonImporter::TryParseJson(TSharedPtr<FJsonObject> SettingsJson, F
 	}
 
 	TryGetNumberField(SettingsJson, "fine_tune", Settings.FineTuneCents, Defaults.FineTuneCents);
-	Settings.SetFineTuneCents(Settings.FineTuneCents);
 
 	TryGetBoolField(SettingsJson, "is_note_off_zone", Settings.bIsNoteOffZone, Defaults.bIsNoteOffZone);
 	TryGetNumberField(SettingsJson, "priority", Settings.Priority, Defaults.Priority);
