Move Editor icons back into the Harmonix/Content folder.

Make sure code references it correctly

diff --git a/Engine/Plugins/Runtime/Harmonix/Source/HarmonixMetasoundEditor/Private/HarmonixMetasoundSlateStyle.cpp b/Engine/Plugins/Runtime/Harmonix/Source/HarmonixMetasoundEditor/Private/HarmonixMetasoundSlateStyle.cpp
--- a/Engine/Plugins/Runtime/Harmonix/Source/HarmonixMetasoundEditor/Private/HarmonixMetasoundSlateStyle.cpp
+++ b/Engine/Plugins/Runtime/Harmonix/Source/HarmonixMetasoundEditor/Private/HarmonixMetasoundSlateStyle.cpp
@@ -22,7 +22,7 @@ namespace HarmonixMetasoundEditor
 		const FVector2D Icon18x18(18.0f, 18.0f);
 		TSharedPtr<IPlugin> Plugin = IPluginManager::Get().FindPlugin(TEXT("Harmonix"));
 		check(Plugin);
-		SetContentRoot(Plugin->GetBaseDir() / TEXT("Editor/Slate"));
+		SetContentRoot(Plugin->GetContentDir() / TEXT("Editor/Slate"));
 
 #define IMAGE_BRUSH( RelativePath, ... ) FSlateImageBrush( RootToContentDir( RelativePath, TEXT(".png") ), __VA_ARGS__ )
 		SetCustomPinStyle("MIDIStream", MidiStreamColor, new IMAGE_BRUSH(TEXT("Icons/MidiConnectedPin"), Icon22x22), new IMAGE_BRUSH(TEXT("Icons/MidiDisconnectedPin"), Icon22x22));
