[HarmonixMetasound] Minor typo found in MusicTransport

diff --git a/Engine/Plugins/Runtime/Harmonix/Source/HarmonixMetasound/Private/HarmonixMetasound/DataTypes/MusicTransport.cpp b/Engine/Plugins/Runtime/Harmonix/Source/HarmonixMetasound/Private/HarmonixMetasound/DataTypes/MusicTransport.cpp
--- a/Engine/Plugins/Runtime/Harmonix/Source/HarmonixMetasound/Private/HarmonixMetasound/DataTypes/MusicTransport.cpp
+++ b/Engine/Plugins/Runtime/Harmonix/Source/HarmonixMetasound/Private/HarmonixMetasound/DataTypes/MusicTransport.cpp
@@ -169,7 +169,7 @@ namespace HarmonixMetasound
 		case EMusicPlayerTransportState::Pausing:    return TEXT("Pausing");
 		case EMusicPlayerTransportState::Paused:     return TEXT("Paused");
 		case EMusicPlayerTransportState::Continuing: return TEXT("Continuing");
-		case EMusicPlayerTransportState::Stopping:   return TEXT("Preparing");
+		case EMusicPlayerTransportState::Stopping:   return TEXT("Stopping");
 		case EMusicPlayerTransportState::Killing:    return TEXT("Killing");
 		}
 		return TEXT("<unknown>");
