- Fix issue with importing MIDI files generated with Ableton Live.
- Add 'VeryVerbose' logging that can be turned on when importing a MIDI file. This can be useful when debugging importing from different DAWs.

checked in on behalf of buzz.burrowes

#tests Exported MIDI from Ableton and imported into editor successfully

diff --git a/Engine/Plugins/Runtime/Harmonix/Source/HarmonixMidi/Private/HarmonixMidi/BarMap.cpp b/Engine/Plugins/Runtime/Harmonix/Source/HarmonixMidi/Private/HarmonixMidi/BarMap.cpp
--- a/Engine/Plugins/Runtime/Harmonix/Source/HarmonixMidi/Private/HarmonixMidi/BarMap.cpp
+++ b/Engine/Plugins/Runtime/Harmonix/Source/HarmonixMidi/Private/HarmonixMidi/BarMap.cpp
@@ -359,6 +359,13 @@ bool FBarMap::AddTimeSignatureAtBarIncludingCountIn(int32 BarIndex, int32 Numera
 	{
 		if (!Points.IsEmpty())
 		{
+			if (Points[0].BarIndex == 0 && Points[0].TimeSignature.Numerator == Numerator && Points[0].TimeSignature.Denominator == Denominator)
+			{
+				// Some DAWs are exporting midi with multiple duplicate time signatures in the file. (Looking at you Ableton!)
+				UE_LOG(LogMIDI, VeryVerbose, TEXT("Ignoring duplicate time signature at beginning of the MIDI data."));
+				return true;
+			}
+
 			if (FailOnError)
 			{
 				checkf(false, TEXT("Multiple time signatures at start of song"));
diff --git a/Engine/Plugins/Runtime/Harmonix/Source/HarmonixMidi/Private/HarmonixMidi/MidiConstants.cpp b/Engine/Plugins/Runtime/Harmonix/Source/HarmonixMidi/Private/HarmonixMidi/MidiConstants.cpp
--- a/Engine/Plugins/Runtime/Harmonix/Source/HarmonixMidi/Private/HarmonixMidi/MidiConstants.cpp
+++ b/Engine/Plugins/Runtime/Harmonix/Source/HarmonixMidi/Private/HarmonixMidi/MidiConstants.cpp
@@ -228,4 +228,67 @@ namespace Harmonix::Midi::Constants
 	{
 		return (int8)MidiNoteNumber / GNotesPerOctave - 1;
 	}
+
+	FString MakeStdMsgString(uint8 Status, uint8 Data1, uint8 Data2)
+	{
+		uint8 Channel = GetChannel(Status);
+		uint8 Type = GetType(Status);
+
+		switch (Type)
+		{
+		case GNoteOff:
+			return FString::Printf(TEXT("Ch #%d - Note Off: %d"), Channel, Data1);
+		case GNoteOn:
+			return FString::Printf(TEXT("Ch #%d - Note On: %d - %d"), Channel, Data1, Data2);
+		case GPolyPres:
+			return FString::Printf(TEXT("Ch #%d - Polyphonic Pressure : %d - %d"), Channel, Data1, Data2);
+		case GControl:
+			return FString::Printf(TEXT("Ch #%d - Control: %s: %d"), Channel, *GetControllerName((EControllerID)Data1), Data2);
+		case GProgram:
+			return FString::Printf(TEXT("Ch #%d - Patch Change: %d"), Channel, Data1);
+		case GChanPres:
+			return FString::Printf(TEXT("Ch #%d - Channel Pressure: %d"), Channel, Data1);
+		case GPitch:
+			return FString::Printf(TEXT("Ch #%d - Pitch Wheel: %d"), Channel, (uint32)Data1 | ((uint32)Data2 << 7));
+		}
+		return FString(TEXT("Unrecognized MIDI Event Type!"));
+	}
+
+	FString GetMetaEventTypeName(uint8 Type)
+	{
+		switch (Type)
+		{
+		case GMeta_Text:
+			return FString(TEXT("Text"));
+		case GMeta_Copyright:
+			return FString(TEXT("Copyright"));
+		case GMeta_TrackName:
+			return FString(TEXT("Track Name"));
+		case GMeta_InstrumentName:
+			return FString(TEXT("Instrument Name"));
+		case GMeta_Lyric:
+			return FString(TEXT("Lyric"));
+		case GMeta_Marker:
+			return FString(TEXT("Marker"));
+		case GMeta_CuePoint:
+			return FString(TEXT("Cur Point"));
+		case GMeta_ChannelPrefix:
+			return FString(TEXT("Channel Prefix"));
+		case GMeta_Port:
+			return FString(TEXT("Port"));
+		case GMeta_EndOfTrack:
+			return FString(TEXT("End Of Track"));
+		case GMeta_Tempo:
+			return FString(TEXT("Tempo"));
+		case GMeta_SMPTE:
+			return FString(TEXT("SMPTE"));
+		case GMeta_TimeSig:
+			return FString(TEXT("Time Signature"));
+		case GMeta_KeySig:
+			return FString(TEXT("Key Signature"));
+		case GMeta_Special:
+			return FString(TEXT("Special"));
+		}
+		return FString(TEXT("Unrecognized MIDI Meta Event Type!"));
+	}
 }
diff --git a/Engine/Plugins/Runtime/Harmonix/Source/HarmonixMidi/Private/HarmonixMidi/MidiReader.cpp b/Engine/Plugins/Runtime/Harmonix/Source/HarmonixMidi/Private/HarmonixMidi/MidiReader.cpp
--- a/Engine/Plugins/Runtime/Harmonix/Source/HarmonixMidi/Private/HarmonixMidi/MidiReader.cpp
+++ b/Engine/Plugins/Runtime/Harmonix/Source/HarmonixMidi/Private/HarmonixMidi/MidiReader.cpp
@@ -155,6 +155,8 @@ void FStdMidiFileReader::ReadAllTracks()
 	// read from the top of the stream
 	Archive->Seek(0);
 
+	UE_LOG(LogMIDI, VeryVerbose, TEXT("READING ALL MIDI TRACKS:\n========================"));
+
 	bool TrackReadSuccessfully = false;
 	do 
 	{
@@ -210,6 +212,8 @@ bool FStdMidiFileReader::ReadEvents(int32 Count)
 // Has no effect if called from within EndOfTrack().
 void FStdMidiFileReader::SkipCurrentTrack()
 {
+	UE_LOG(LogMIDI, VeryVerbose, TEXT("Skipping Current Track (%d)"), CurrentTrackIndex);
+
 	if (State == EState::InTrack)
 	{
 		if (CurrentTrackIndex == NumTracks - 1)
@@ -260,6 +264,7 @@ void FStdMidiFileReader::ReadNextEventImpl()
 	case EState::NewTrack:
 		if (Format == 0 && CurrentTrackIndex == 0)
 		{
+			UE_LOG(LogMIDI, VeryVerbose, TEXT("Type 0 file so rewinding to re-read track in 'non-conductor' mode."));
 			// we already ready the one and only track as the "conductor track".
 			// we have to read it again as "track 1", so...
 			Archive->Seek(LastTracksFilePosition);
@@ -285,6 +290,8 @@ void FStdMidiFileReader::ReadFileHeader()
 {
 	check(State == EState::Start);
 
+	UE_LOG(LogMIDI, VeryVerbose, TEXT("Reading file header..."));
+
 	FMidiChunkHeader Header(*Archive);
 
 	// Verify the MIDI header is properly formed.
@@ -332,6 +339,8 @@ void FStdMidiFileReader::ReadFileHeader()
 	TickConversionFactor = float(DestinationTicksPerQuarterNote) / (float)TicksPerQuarter;
 
 	State = EState::NewTrack;
+
+	UE_LOG(LogMIDI, VeryVerbose, TEXT("... Done reading file header. Format = %d, File PPQ = %d, NumTracks = %d"), Format, TicksPerQuarter, Format == 0 ? 1 : NumTracks);
 }
 
 // Read the standard Midi track chunk header (MTrk)
@@ -339,6 +348,8 @@ void FStdMidiFileReader::ReadTrackHeader()
 {
 	check(State == EState::NewTrack);
 
+	UE_LOG(LogMIDI, VeryVerbose, TEXT("Reading track header (starting %d bytes into file)..."), Archive->Tell());
+
 	FMidiChunkHeader TrackHeader(*Archive);
 
 	MIDIREADER_EXPECT_WITH_ERROR(TrackHeader.GetID() == FMidiChunkID::kMTrk, TEXT("%s: MIDI track header for track %d is corrupt"), *Filename, CurrentTrackIndex + 1);
@@ -375,6 +386,8 @@ void FStdMidiFileReader::ReadTrackHeader()
 
 		MIDIREADER_EXPECT_WITH_ERROR(Receiver->OnText(0, CurrentTrackName, Harmonix::Midi::Constants::GMeta_TrackName), TEXT("%s: Failed to add track name to Track (%d)"), *Filename, CurrentTrackIndex);
 	}
+
+	UE_LOG(LogMIDI, VeryVerbose, TEXT("... Done reading track header. It contains %d bytes of midi events."), TrackHeader.GetLength());
 }
 
 // Read an event from the track
@@ -512,6 +525,7 @@ void FStdMidiFileReader::ReadMidiEvent(int32 Tick, uint8 Status, uint8 Data1)
 	{
 		// At this time, we only care about Note On / Note Off events.
 		QueueChannelMsg(Tick, Status, Data1, Data2);
+		UE_LOG(LogMIDI, VeryVerbose, TEXT("TICK %d :: %s"), Tick, *Harmonix::Midi::Constants::MakeStdMsgString(Status, Data1, Data2));
 	}
 }
 
@@ -529,6 +543,7 @@ void FStdMidiFileReader::ReadSystemEvent(int32 Tick, uint8 Status)
 			// seek past that length and ignore these events
 			// entirely.
 			int32 PacketLength = Midi::VarLenNumber::Read(*Archive);
+			UE_LOG(LogMIDI, VeryVerbose, TEXT("Seeking %d bytes past SysEx or Escape system event."), PacketLength);
 			if (FailIfBytesNotAvailableInTrackData((int64)PacketLength))
 			{
 				return;
@@ -579,6 +594,7 @@ void FStdMidiFileReader::ReadMetaEvent(int32 Tick, uint8 Type)
 			// so we ignore the name we got from the file and normalize it here to "conductor"...
 			CurrentTrackName = TEXT("Conductor");
 		}
+		UE_LOG(LogMIDI, VeryVerbose, TEXT("Got track name event... '%s'"), *CurrentTrackName);
 		MIDIREADER_EXPECT_WITH_ERROR(Receiver->OnText(Tick, CurrentTrackName, Type), TEXT("%s (%s): Failed to add track name to Track"), *Filename, *CurrentTrackName);
 		break;
 	case GMeta_Copyright:
@@ -588,6 +604,7 @@ void FStdMidiFileReader::ReadMetaEvent(int32 Tick, uint8 Type)
 		if (TrackFilteringMode != ETrackFilteringMode::NonConductorEvents)
 		{
 			MIDIREADER_EXPECT_WITH_ERROR(Receiver->OnText(Tick, WorkingString, Type), TEXT("%s (%s): Failed to add copyright, marker, or cuepoint to Track"), *Filename, *CurrentTrackName);
+			UE_LOG(LogMIDI, VeryVerbose, TEXT("TICK %d :: Got Copyright, Marker, or CuePoint event... '%s'"), Tick, *WorkingString);
 		}
 		break;
 	case GMeta_Text:
@@ -596,6 +613,7 @@ void FStdMidiFileReader::ReadMetaEvent(int32 Tick, uint8 Type)
 		if (TrackFilteringMode != ETrackFilteringMode::ConductorEvents)
 		{
 			MIDIREADER_EXPECT_WITH_ERROR(Receiver->OnText(Tick, WorkingString, Type), TEXT("%s (%s): Failed to add text or lyric to Track"), *Filename, *CurrentTrackName);
+			UE_LOG(LogMIDI, VeryVerbose, TEXT("TICK %d :: Got Text or Lyric event... '%s'"), Tick, *WorkingString);
 		}
 		break;
 	case GMeta_Tempo:
@@ -610,6 +628,7 @@ void FStdMidiFileReader::ReadMetaEvent(int32 Tick, uint8 Type)
 			if (TrackFilteringMode != ETrackFilteringMode::NonConductorEvents)
 			{
 				MIDIREADER_EXPECT_WITH_ERROR(Receiver->OnTempo(Tick, Tempo), TEXT("%s (%s): Failed to add tempo to Track"), *Filename, *CurrentTrackName);
+				UE_LOG(LogMIDI, VeryVerbose, TEXT("TICK %d :: Got Tempo event... %f bpm"), Tick, Harmonix::Midi::Constants::MidiTempoToBPM(Tempo));
 			}
 		}
 		break;
@@ -617,6 +636,7 @@ void FStdMidiFileReader::ReadMetaEvent(int32 Tick, uint8 Type)
 	case GMeta_EndOfTrack:
 		ProcessMidiList();  // in case there is anything left to send out.
 		MIDIREADER_EXPECT_WITH_ERROR(Receiver->OnEndOfTrack(Tick), TEXT("%s (%s): Failed ending current track read"), *Filename, *CurrentTrackName);
+		UE_LOG(LogMIDI, VeryVerbose, TEXT("TICK %d :: Got End Of Trackevent."), Tick);
 		if (CurrentTrackIndex == NumTracks - 1)
 		{
 			State = EState::End;
@@ -652,6 +672,7 @@ void FStdMidiFileReader::ReadMetaEvent(int32 Tick, uint8 Type)
 				// so we can print better error and warning messages...
 				check(Tick == 0 || BarMap->GetNumTimeSignaturePoints() > 0);
 				int32 BarIndex = BarMap->TickToBarIncludingCountIn(Tick);
+				UE_LOG(LogMIDI, VeryVerbose, TEXT("TICK %d :: Got Time Signature Event. %d/%d"), Tick, Numerator, Denominator);
 				MIDIREADER_EXPECT_WITH_ERROR(BarMap->AddTimeSignatureAtBarIncludingCountIn(BarIndex, Numerator, Denominator, true, false),
 					TEXT("%s (%s): Time signature %d/%d at %s overlaps or conflicts with nearby time signatures"),
 					*Filename, *CurrentTrackName, Numerator, Denominator, *MidiTickFormat(Tick, BarMap.Get(), Midi::EMusicTimeStringFormat::Position));
@@ -675,8 +696,9 @@ void FStdMidiFileReader::ReadMetaEvent(int32 Tick, uint8 Type)
 	case GMeta_KeySig:
 	case GMeta_SMPTE:
 	case GMeta_InstrumentName:
-		// these are valid events, but not currently supported by
-		// MidiReceiver; do nothing
+		// these are valid events, but not currently supported by MidiReceiver
+		UE_LOG(LogMIDI, Warning, TEXT("Track '%s', TICK = %d :: Skipping unsupported MIDI event -> %d - %s"), *CurrentTrackName, Tick, Type, *Harmonix::Midi::Constants::GetMetaEventTypeName(Type));
+
 		break;
 
 	default:
diff --git a/Engine/Plugins/Runtime/Harmonix/Source/HarmonixMidi/Public/HarmonixMidi/MidiConstants.h b/Engine/Plugins/Runtime/Harmonix/Source/HarmonixMidi/Public/HarmonixMidi/MidiConstants.h
--- a/Engine/Plugins/Runtime/Harmonix/Source/HarmonixMidi/Public/HarmonixMidi/MidiConstants.h
+++ b/Engine/Plugins/Runtime/Harmonix/Source/HarmonixMidi/Public/HarmonixMidi/MidiConstants.h
@@ -63,6 +63,8 @@ namespace Harmonix::Midi::Constants
 	inline constexpr uint8 GMeta_KeySig          = 0x59;
 	inline constexpr uint8 GMeta_Special         = 0x7f;
 
+	HARMONIXMIDI_API FString GetMetaEventTypeName(uint8 Type);
+
 	// if FMidiMsg type is "Runtime" these are the possible status byte.
 	// note: runtime messages are not serialized!
 	inline constexpr uint8 GRuntimeAllNotesOffStatus = 0x01;  // Receiver should feel free to allow notes to "adsr release"
@@ -191,6 +193,8 @@ namespace Harmonix::Midi::Constants
 	HARMONIXMIDI_API FString GetControllerName(EControllerID ControllerId);
 	HARMONIXMIDI_API void GetControllerNames(TArray<FString>& Names);
 
+	HARMONIXMIDI_API FString MakeStdMsgString(uint8 Status, uint8 Data1, uint8 Data2);
+
 	// allows choice between various naming conventions for enharmonic notes
 	enum class ENoteNameEnharmonicStyle : uint8
 	{
