handle uneven BlockRate in the FusionSamplerNode.

The FusionSamplerNode process frames in 128 frame slices, and caches remaining frames for the next process.

#tests Set BlockRate to 600, resulting in FramesPerBlock of 80, which triggers the unsafe condition

diff --git a/Engine/Plugins/Runtime/Harmonix/Source/HarmonixMetasound/Private/HarmonixMetasound/Nodes/FusionSamplerNode.cpp b/Engine/Plugins/Runtime/Harmonix/Source/HarmonixMetasound/Private/HarmonixMetasound/Nodes/FusionSamplerNode.cpp
--- a/Engine/Plugins/Runtime/Harmonix/Source/HarmonixMetasound/Private/HarmonixMetasound/Nodes/FusionSamplerNode.cpp
+++ b/Engine/Plugins/Runtime/Harmonix/Source/HarmonixMetasound/Private/HarmonixMetasound/Nodes/FusionSamplerNode.cpp
@@ -179,6 +179,7 @@ namespace HarmonixMetasound
 			
 			Prepare(Params.OperatorSettings.GetSampleRate(), OutputChannelLayout, AudioRendering::kFramesPerRenderBuffer, true);
 			BusBuffer.SetNumValidFrames(0);
+			ResidualBuffer.SetNumValidFrames(0);
 			FramesPerBlock = Params.OperatorSettings.GetNumFramesPerBlock();
 			ResetInstrumentState();
 			SetSampleRate(Params.OperatorSettings.GetSampleRate());
@@ -193,13 +194,15 @@ namespace HarmonixMetasound
 	protected:
 		virtual void ResetInternal(const FResetParams& Params) = 0;
 		virtual void ZeroOutput() = 0;
-		virtual int32 UseInitialPortion(int32 OutputOffset, int32 MaxFrames) = 0;
-		virtual void UseResidual() = 0;
+		virtual int32 CopyFramesToOutput(int32 OutputOffset, const TAudioBuffer<float>& InputBuffer, int32 NumFrames) = 0;
 		virtual void SetupOutputAlias(TAudioBuffer<float>& AliasBuffer, int32 Offset) = 0;
 
 		//** OUTPUTS
 		FFusionSyncLinkWriteRef SyncLinkOutPin;
 
+		//** DATA
+		TAudioBuffer<float> ResidualBuffer;
+
 	private:
 
 		//** INPUTS
@@ -395,29 +398,19 @@ namespace HarmonixMetasound
 			}
 		}
 
-		virtual int32 UseInitialPortion(int32 OutputOffset, int32 MaxFrames) override
+		virtual int32 CopyFramesToOutput(int32 OutputOffset, const TAudioBuffer<float>& InputBuffer, int32 NumFrames) override
 		{
-			int32 CopyFrames = FMath::Min(MaxFrames, AudioRendering::kFramesPerRenderBuffer);
+			int32 CopyFrames = FMath::Min(InputBuffer.GetNumValidFrames(), NumFrames);
 			int32 CopySizeBytes = sizeof(float) * CopyFrames;
 			for (int32 i = 0; i < NUM_CHANNELS; ++i)
 			{
-				float* SourceSamples = BusBuffer.GetRawChannelData(i);
+				check(CopyFrames <= AudioOuts[i]->Num() - OutputOffset);
+				const float* SourceSamples = InputBuffer.GetRawChannelData(i);
 				FMemory::Memcpy(AudioOuts[i]->GetData() + OutputOffset, SourceSamples, CopySizeBytes);
 			}
 			return CopyFrames;
 		}
 
-		virtual void UseResidual() override
-		{
-			int32 SourceOffset = BusBuffer.GetMaxNumFrames() - BusBuffer.GetNumValidFrames();
-			int32 CopySizeBytes = sizeof(float) * BusBuffer.GetNumValidFrames();
-			for (int32 i = 0; i < NUM_CHANNELS; ++i)
-			{
-				float* SourceSamples = BusBuffer.GetRawChannelData(i) + SourceOffset;
-				FMemory::Memcpy(AudioOuts[i]->GetData(), SourceSamples, CopySizeBytes);
-			}
-		}
-
 		virtual void SetupOutputAlias(TAudioBuffer<float>& AliasBuffer, int32 Offset) override
 		{
 			float* ChannelDataPtrs[NUM_CHANNELS];
@@ -609,19 +602,17 @@ namespace HarmonixMetasound
 
 		MadeAudioLastFrame = false;
 
-		if (BusBuffer.GetNumValidFrames() > 0)
+		if (ResidualBuffer.GetNumValidFrames() > 0)
 		{
 			// we have some samples from the last render we need to use...
-			check(BusBuffer.GetNumValidFrames() < FramesPerBlock);
-
-			UseResidual();
+			int32 NumFramesCopied = CopyFramesToOutput(CurrentBlockFrameIndex, ResidualBuffer, FramesRequired);
+			ResidualBuffer.AdvanceAliasedDataPointers(NumFramesCopied);
 
-			FramesRequired -= BusBuffer.GetNumValidFrames();
-			CurrentBlockFrameIndex += BusBuffer.GetNumValidFrames();
-			BusBuffer.SetNumValidFrames(0);
+			FramesRequired -= NumFramesCopied;
+			CurrentBlockFrameIndex += NumFramesCopied;
 			MadeAudioLastFrame = true;
 		}
-
+		
 		TAudioBuffer<float> AliasedOutputBuffer;
 		SetupOutputAlias(AliasedOutputBuffer, CurrentBlockFrameIndex);
 
@@ -692,11 +683,13 @@ namespace HarmonixMetasound
 				Process(SliceIndex++, 0, BusBuffer);
 				MadeAudioLastFrame = MadeAudioLastFrame || !BusBuffer.GetIsSilent();
 				
-				int32 CopiedFrames = UseInitialPortion(CurrentBlockFrameIndex, FramesRequired);
+				int32 CopiedFrames = CopyFramesToOutput(CurrentBlockFrameIndex, BusBuffer, FramesRequired);
+				// Alias the BusBuffer and initialize it with the offset of the num frames copied
+				// this allows us to advance through the frames as we copy them from the residual buffer
+				ResidualBuffer.Alias(BusBuffer, CopiedFrames);
 
 				FramesRequired -= CopiedFrames;
 				CurrentBlockFrameIndex += CopiedFrames;
-				BusBuffer.SetNumValidFrames(AudioRendering::kFramesPerRenderBuffer - CopiedFrames);
 			}
 		}
 
