Remove unneeded `check` in PannerDetailsCustomization. We null check the Property in other places and handle it more gracefully.

diff --git a/Engine/Plugins/Runtime/Harmonix/Source/HarmonixDspEditor/Private/Customization/PannerDetailsCustomization.cpp b/Engine/Plugins/Runtime/Harmonix/Source/HarmonixDspEditor/Private/Customization/PannerDetailsCustomization.cpp
--- a/Engine/Plugins/Runtime/Harmonix/Source/HarmonixDspEditor/Private/Customization/PannerDetailsCustomization.cpp
+++ b/Engine/Plugins/Runtime/Harmonix/Source/HarmonixDspEditor/Private/Customization/PannerDetailsCustomization.cpp
@@ -24,12 +24,6 @@ void FPannerDetailsCustomization::CustomizeHeader(TSharedRef<class IPropertyHand
 	MyPropertyUtils = StructCustomizationUtils.GetPropertyUtilities();
 	TArray<UObject*> Objects;
 	PropertyHandle->GetOuterObjects(Objects);
-
-	check(Objects.Num() == 1);
-	FPannerDetails* PannerDetails = (FPannerDetails*)(PropertyHandle->GetValueBaseAddress((uint8*)(Objects[0])));
-	ensure(PannerDetails);
-
-	EPannerMode PannerMode = PannerDetails ? PannerDetails->Mode : EPannerMode::Invalid;
 	UEnum* PannerModeEnum = StaticEnum<EPannerMode>();
 
 	HeaderRow
