[HarmonixMetasound] Allow StepSequencePlayerNode to update sequence input after input binding.

diff --git a/Engine/Plugins/Runtime/Harmonix/Source/HarmonixMetasound/Private/HarmonixMetasound/Nodes/StepSequencePlayerNode.cpp b/Engine/Plugins/Runtime/Harmonix/Source/HarmonixMetasound/Private/HarmonixMetasound/Nodes/StepSequencePlayerNode.cpp
--- a/Engine/Plugins/Runtime/Harmonix/Source/HarmonixMetasound/Private/HarmonixMetasound/Nodes/StepSequencePlayerNode.cpp
+++ b/Engine/Plugins/Runtime/Harmonix/Source/HarmonixMetasound/Private/HarmonixMetasound/Nodes/StepSequencePlayerNode.cpp
@@ -62,6 +62,7 @@ namespace HarmonixMetasound
 
 	private:
 		void Init();
+		void InitSequenceTable();
 		
 		//** INPUTS
 		FMidiStepSequenceAssetReadRef SequenceAssetInPin;
@@ -421,19 +422,7 @@ namespace HarmonixMetasound
 		MidiOutPin->SetClock(*MidiClockInPin);
 		MidiOutPin->PrepareBlock();
 
-		SequenceTable = SequenceAssetInPin->GetRenderable();
-		if (SequenceTable)
-		{
-			ResizeCellStatesForTable();
-			UE_LOG(LogStepSequencePlayer, Verbose, TEXT("Got a Sequence: %d pages, %d rows, %d columns"), 
-				SequenceTable->Pages.Num(), 
-				SequenceTable->Pages.IsEmpty() ? 0 : SequenceTable->Pages[0].Rows.Num(), 
-				SequenceTable->Pages.IsEmpty() || SequenceTable->Pages[0].Rows.IsEmpty() ? 0 : SequenceTable->Pages[0].Rows[0].Cells.Num());
-		}
-		else
-		{
-			UE_LOG(LogStepSequencePlayer, Verbose, TEXT("No Sequence Provided!"));
-		}
+		InitSequenceTable();
 
 		FTransportInitFn InitFn = [this](EMusicPlayerTransportState CurrentState)
 		{
@@ -481,6 +470,23 @@ namespace HarmonixMetasound
 		FMusicTransportControllable::Init(*TransportInPin, MoveTemp(InitFn));
 	}
 
+	void FStepSequencePlayerOperator::InitSequenceTable()
+	{
+		SequenceTable = SequenceAssetInPin->GetRenderable();
+		if (SequenceTable)
+		{
+			ResizeCellStatesForTable();
+			UE_LOG(LogStepSequencePlayer, Verbose, TEXT("Got a Sequence: %d pages, %d rows, %d columns"),
+				SequenceTable->Pages.Num(),
+				SequenceTable->Pages.IsEmpty() ? 0 : SequenceTable->Pages[0].Rows.Num(),
+				SequenceTable->Pages.IsEmpty() || SequenceTable->Pages[0].Rows.IsEmpty() ? 0 : SequenceTable->Pages[0].Rows[0].Cells.Num());
+		}
+		else
+		{
+			UE_LOG(LogStepSequencePlayer, Verbose, TEXT("No Sequence Provided!"));
+		}
+	}
+
 	void FStepSequencePlayerOperator::SeekToTick(int32 BlockFrameIndex, int32 Tick)
 	{
 		SeekThruTick(BlockFrameIndex, Tick-1);
@@ -836,6 +842,13 @@ namespace HarmonixMetasound
 
 	void FStepSequencePlayerOperator::CheckForUpdatedSequenceTable()
 	{
+		FStepSequenceTableProxy::NodePtr Tester = SequenceAssetInPin->GetRenderable();
+		if (Tester != SequenceTable)
+		{
+			InitSequenceTable();
+			return;
+		}
+
 		if (SequenceTable)
 		{
 			TRefCountedAudioRenderableWithQueuedChanges<FStepSequenceTable>* Table = SequenceTable;
