Fix bug where the wrong BlockFrameIndex was being used.

diff --git a/Engine/Plugins/Runtime/Harmonix/Source/HarmonixMetasound/Private/HarmonixMetasound/Nodes/MidiClockSubdivisionTriggerNode.cpp b/Engine/Plugins/Runtime/Harmonix/Source/HarmonixMetasound/Private/HarmonixMetasound/Nodes/MidiClockSubdivisionTriggerNode.cpp
--- a/Engine/Plugins/Runtime/Harmonix/Source/HarmonixMetasound/Private/HarmonixMetasound/Nodes/MidiClockSubdivisionTriggerNode.cpp
+++ b/Engine/Plugins/Runtime/Harmonix/Source/HarmonixMetasound/Private/HarmonixMetasound/Nodes/MidiClockSubdivisionTriggerNode.cpp
@@ -265,14 +265,14 @@ namespace HarmonixMetasound
 						{
 							if (End.Bar % GridSizeMultiplier == 0)
 							{
-								TriggerOutPin->TriggerFrame(MidiClockInPin->GetCurrentBlockFrameIndex());
+								TriggerOutPin->TriggerFrame(ClockEvent.BlockFrameIndex);
 							}
 						}
 						else if (GridSizeUnits == EMidiClockSubdivisionQuantization::Beat && FMath::FloorToInt32(Start.Beat) != FMath::FloorToInt32(End.Beat))
 						{
 							if (FMath::FloorToInt32(End.Beat-1.0f) % GridSizeMultiplier == 0)
 							{
-								TriggerOutPin->TriggerFrame(MidiClockInPin->GetCurrentBlockFrameIndex());
+								TriggerOutPin->TriggerFrame(ClockEvent.BlockFrameIndex);
 							}
 						}
 					}
@@ -283,7 +283,7 @@ namespace HarmonixMetasound
 						int32 EndingGridSquare = LastTickProcessed / GridSizeTicks;
 						if (StartingGridSquare != EndingGridSquare)
 						{
-							TriggerOutPin->TriggerFrame(MidiClockInPin->GetCurrentBlockFrameIndex());
+							TriggerOutPin->TriggerFrame(ClockEvent.BlockFrameIndex);
 						}
 					}
 				}
