Step Sequence Player Node Header file

#tests Harmonix Functional Tests

diff --git a/Engine/Plugins/Runtime/Harmonix/Source/HarmonixMetasound/Private/HarmonixMetasound/Nodes/StepSequencePlayerNode.cpp b/Engine/Plugins/Runtime/Harmonix/Source/HarmonixMetasound/Private/HarmonixMetasound/Nodes/StepSequencePlayerNode.cpp
--- a/Engine/Plugins/Runtime/Harmonix/Source/HarmonixMetasound/Private/HarmonixMetasound/Nodes/StepSequencePlayerNode.cpp
+++ b/Engine/Plugins/Runtime/Harmonix/Source/HarmonixMetasound/Private/HarmonixMetasound/Nodes/StepSequencePlayerNode.cpp
@@ -24,10 +24,26 @@
 
 DEFINE_LOG_CATEGORY_STATIC(LogStepSequencePlayer, Log, All);
 
-namespace HarmonixMetasound
+namespace HarmonixMetasound::Nodes::StepSequencePlayer
 {
 	using namespace Metasound;
 	using namespace Harmonix;
+	
+	const Metasound::FNodeClassName& GetClassName()
+	{
+		static Metasound::FNodeClassName ClassName
+		{
+			HarmonixNodeNamespace,
+			"StepSequencePlayer",
+			""
+		};
+		return ClassName;
+	}
+
+	int32 GetCurrentMajorVersion()
+	{
+		return 0;
+	}
 
 	class FStepSequencePlayerOperator : public TExecutableOperator<FStepSequencePlayerOperator>, public FMusicTransportControllable, public FMidiVoiceGeneratorBase
 	{
@@ -133,8 +149,8 @@ namespace HarmonixMetasound
 		auto InitNodeInfo = []() -> FNodeClassMetadata
 		{
 			FNodeClassMetadata Info;
-			Info.ClassName        = { HarmonixNodeNamespace, TEXT("StepSequencePlayer"), TEXT("") };
-			Info.MajorVersion     = 0;
+			Info.ClassName        = GetClassName();
+			Info.MajorVersion     = GetCurrentMajorVersion();
 			Info.MinorVersion     = 1;
 			Info.DisplayName      = METASOUND_LOCTEXT("StepSequencePlayerNode_DisplayName", "Step Sequence Player");
 			Info.Description      = METASOUND_LOCTEXT("StepSequencePlayerNode_Description", "Plays a Step Sequence Asset.");
@@ -151,40 +167,46 @@ namespace HarmonixMetasound
 		return Info;
 	}
 
-	namespace StepSequencePlayerPinNames
+	namespace Inputs
 	{
-		METASOUND_PARAM(InputSequenceAsset, "Step Sequence Asset", "Step sequence to play.");
-		METASOUND_PARAM(InputVelocityMultiplier, "Velocity Multiplier", "Multiplies the current note velocity by this number");
-		METASOUND_PARAM(InputMaxColumns, "Max Columns", "The Maximinum Number of cells to play per step sequence row.");
-		METASOUND_PARAM(InputAdditionalOctaves, "Additional Octaves", "The number of octaves to add to the authored step sequence note.");
-		METASOUND_PARAM(InputStepSizeQuarterNotes, "Step Size Quarter Notes", "The size, in quarter notes, of each step");
-		METASOUND_PARAM(InputActivePage, "Active Page", "The page of the step sequence to play (1 indexed)");
-		METASOUND_PARAM(InputAutoPage, "Auto Page", "Whether to calculate the page of the step sequence based on current position");
-		METASOUND_PARAM(InputAutoPagePlaysBlankPages, "Auto Page Plays Blank Pages", "If autopaging, should blank pages be played?");
-		METASOUND_PARAM(InputLoop, "Loop", "If not looping, the sequence will be played as a one-shot starting on the next stepsizebeats interval (autopaging if it's enabled)");
-		METASOUND_PARAM(InputEnabled, "Enabled", "Whether the sequencer is enabled. If loop is off, enabling will trigger to start on next stepsizebeats interval");
+		DEFINE_INPUT_METASOUND_PARAM(SequenceAsset, "Step Sequence Asset", "Step sequence to play.");
+		DEFINE_INPUT_METASOUND_PARAM(VelocityMultiplier, "Velocity Multiplier", "Multiplies the current note velocity by this number");
+		DEFINE_INPUT_METASOUND_PARAM(MaxColumns, "Max Columns", "The Maximinum Number of cells to play per step sequence row.");
+		DEFINE_INPUT_METASOUND_PARAM(AdditionalOctaves, "Additional Octaves", "The number of octaves to add to the authored step sequence note.");
+		DEFINE_INPUT_METASOUND_PARAM(StepSizeQuarterNotes, "Step Size Quarter Notes", "The size, in quarter notes, of each step");
+		DEFINE_INPUT_METASOUND_PARAM(ActivePage, "Active Page", "The page of the step sequence to play (1 indexed)");
+		DEFINE_INPUT_METASOUND_PARAM(AutoPage, "Auto Page", "Whether to calculate the page of the step sequence based on current position");
+		DEFINE_INPUT_METASOUND_PARAM(AutoPagePlaysBlankPages, "Auto Page Plays Blank Pages", "If autopaging, should blank pages be played?");
+		DEFINE_METASOUND_PARAM_ALIAS(Transport, CommonPinNames::Inputs::Transport);
+		DEFINE_METASOUND_PARAM_ALIAS(MidiClock, CommonPinNames::Inputs::MidiClock);
+		DEFINE_METASOUND_PARAM_ALIAS(Speed, CommonPinNames::Inputs::Speed);
+		DEFINE_METASOUND_PARAM_ALIAS(Loop, CommonPinNames::Inputs::Loop);
+		DEFINE_METASOUND_PARAM_ALIAS(Enabled, CommonPinNames::Inputs::Enable);
+	}
+
+	namespace Outputs
+	{
+		DEFINE_METASOUND_PARAM_ALIAS(MidiStream, CommonPinNames::Outputs::MidiStream);
 	}
 
 	const FVertexInterface& FStepSequencePlayerOperator::GetVertexInterface()
 	{
-		using namespace StepSequencePlayerPinNames;
-		using namespace CommonPinNames;
 
 		static const FVertexInterface Interface(
 			FInputVertexInterface(
-				TInputDataVertex<FMidiStepSequenceAsset>(METASOUND_GET_PARAM_NAME_AND_METADATA(InputSequenceAsset)),
+				TInputDataVertex<FMidiStepSequenceAsset>(METASOUND_GET_PARAM_NAME_AND_METADATA(Inputs::SequenceAsset)),
 				TInputDataVertex<FMusicTransportEventStream>(METASOUND_GET_PARAM_NAME_AND_METADATA(Inputs::Transport)),
 				TInputDataVertex<FMidiClock>(METASOUND_GET_PARAM_NAME_AND_METADATA(Inputs::MidiClock)),
 				TInputDataVertex<float>(METASOUND_GET_PARAM_NAME_AND_METADATA(Inputs::Speed), 1.0f),
-				TInputDataVertex<float>(METASOUND_GET_PARAM_NAME_AND_METADATA(InputVelocityMultiplier), 1.0f),
-				TInputDataVertex<float>(METASOUND_GET_PARAM_NAME_AND_METADATA(InputMaxColumns), 64.0f),
-				TInputDataVertex<float>(METASOUND_GET_PARAM_NAME_AND_METADATA(InputAdditionalOctaves), 0),
-				TInputDataVertex<float>(METASOUND_GET_PARAM_NAME_AND_METADATA(InputStepSizeQuarterNotes), 0.25f),
-				TInputDataVertex<float>(METASOUND_GET_PARAM_NAME_AND_METADATA(InputActivePage), 0),
-				TInputDataVertex<bool>(METASOUND_GET_PARAM_NAME_AND_METADATA(InputAutoPage), false),
-				TInputDataVertex<bool>(METASOUND_GET_PARAM_NAME_AND_METADATA(InputAutoPagePlaysBlankPages), true),
-				TInputDataVertex<bool>(METASOUND_GET_PARAM_NAME_AND_METADATA(InputLoop), true),
-				TInputDataVertex<bool>(METASOUND_GET_PARAM_NAME_AND_METADATA(InputEnabled), true)
+				TInputDataVertex<float>(METASOUND_GET_PARAM_NAME_AND_METADATA(Inputs::VelocityMultiplier), 1.0f),
+				TInputDataVertex<float>(METASOUND_GET_PARAM_NAME_AND_METADATA(Inputs::MaxColumns), 64.0f),
+				TInputDataVertex<float>(METASOUND_GET_PARAM_NAME_AND_METADATA(Inputs::AdditionalOctaves), 0),
+				TInputDataVertex<float>(METASOUND_GET_PARAM_NAME_AND_METADATA(Inputs::StepSizeQuarterNotes), 0.25f),
+				TInputDataVertex<float>(METASOUND_GET_PARAM_NAME_AND_METADATA(Inputs::ActivePage), 0),
+				TInputDataVertex<bool>(METASOUND_GET_PARAM_NAME_AND_METADATA(Inputs::AutoPage), false),
+				TInputDataVertex<bool>(METASOUND_GET_PARAM_NAME_AND_METADATA(Inputs::AutoPagePlaysBlankPages), true),
+				TInputDataVertex<bool>(METASOUND_GET_PARAM_NAME_AND_METADATA(Inputs::Loop), true),
+				TInputDataVertex<bool>(METASOUND_GET_PARAM_NAME_AND_METADATA(Inputs::Enabled), true)
 			),
 			FOutputVertexInterface(
 				TOutputDataVertex<FMidiStream>(METASOUND_GET_PARAM_NAME_AND_METADATA(Outputs::MidiStream))
@@ -195,26 +217,23 @@ namespace HarmonixMetasound
 
 	TUniquePtr<IOperator> FStepSequencePlayerOperator::CreateOperator(const FBuildOperatorParams& InParams, FBuildResults& OutResults)
 	{
-		using namespace StepSequencePlayerPinNames;
-		using namespace CommonPinNames;
-
 		const FStepSequencePlayerNode& PlayerNode = static_cast<const FStepSequencePlayerNode&>(InParams.Node);
 
 		const FInputVertexInterfaceData& InputData = InParams.InputData;
 
-		FMidiStepSequenceAssetReadRef InSequenceAsset = InputData.GetOrConstructDataReadReference<FMidiStepSequenceAsset>(METASOUND_GET_PARAM_NAME(InputSequenceAsset));
+		FMidiStepSequenceAssetReadRef InSequenceAsset = InputData.GetOrConstructDataReadReference<FMidiStepSequenceAsset>(METASOUND_GET_PARAM_NAME(Inputs::SequenceAsset));
 		FMusicTransportEventStreamReadRef InTransport = InputData.GetOrConstructDataReadReference<FMusicTransportEventStream>(METASOUND_GET_PARAM_NAME(Inputs::Transport), InParams.OperatorSettings);
 		FFloatReadRef InSpeed = InputData.GetOrCreateDefaultDataReadReference<float>(METASOUND_GET_PARAM_NAME(Inputs::Speed), InParams.OperatorSettings);
-		FFloatReadRef InVelocityMult = InputData.GetOrCreateDefaultDataReadReference<float>(METASOUND_GET_PARAM_NAME(InputVelocityMultiplier), InParams.OperatorSettings);
+		FFloatReadRef InVelocityMult = InputData.GetOrCreateDefaultDataReadReference<float>(METASOUND_GET_PARAM_NAME(Inputs::VelocityMultiplier), InParams.OperatorSettings);
 		FMidiClockReadRef InMidiClock = InputData.GetOrConstructDataReadReference<FMidiClock>(METASOUND_GET_PARAM_NAME(Inputs::MidiClock), InParams.OperatorSettings);
-		FFloatReadRef InMaxColumns = InputData.GetOrCreateDefaultDataReadReference<float>(METASOUND_GET_PARAM_NAME(InputMaxColumns), InParams.OperatorSettings);
-		FFloatReadRef InAdditionalOctaves = InputData.GetOrCreateDefaultDataReadReference<float>(METASOUND_GET_PARAM_NAME(InputAdditionalOctaves), InParams.OperatorSettings);
-		FFloatReadRef InStepSize = InputData.GetOrCreateDefaultDataReadReference<float>(METASOUND_GET_PARAM_NAME(InputStepSizeQuarterNotes), InParams.OperatorSettings);
-		FFloatReadRef InActivePage = InputData.GetOrCreateDefaultDataReadReference<float>(METASOUND_GET_PARAM_NAME(InputActivePage), InParams.OperatorSettings);
-		FBoolReadRef InAutoPage = InputData.GetOrCreateDefaultDataReadReference<bool>(METASOUND_GET_PARAM_NAME(InputAutoPage), InParams.OperatorSettings);
-		FBoolReadRef InAutoPagePlaysBlankPages = InputData.GetOrCreateDefaultDataReadReference<bool>(METASOUND_GET_PARAM_NAME(InputAutoPagePlaysBlankPages), InParams.OperatorSettings);
-		FBoolReadRef InLoop = InputData.GetOrCreateDefaultDataReadReference<bool>(METASOUND_GET_PARAM_NAME(InputLoop), InParams.OperatorSettings);
-		FBoolReadRef InEnabled = InputData.GetOrCreateDefaultDataReadReference<bool>(METASOUND_GET_PARAM_NAME(InputEnabled), InParams.OperatorSettings);
+		FFloatReadRef InMaxColumns = InputData.GetOrCreateDefaultDataReadReference<float>(METASOUND_GET_PARAM_NAME(Inputs::MaxColumns), InParams.OperatorSettings);
+		FFloatReadRef InAdditionalOctaves = InputData.GetOrCreateDefaultDataReadReference<float>(METASOUND_GET_PARAM_NAME(Inputs::AdditionalOctaves), InParams.OperatorSettings);
+		FFloatReadRef InStepSize = InputData.GetOrCreateDefaultDataReadReference<float>(METASOUND_GET_PARAM_NAME(Inputs::StepSizeQuarterNotes), InParams.OperatorSettings);
+		FFloatReadRef InActivePage = InputData.GetOrCreateDefaultDataReadReference<float>(METASOUND_GET_PARAM_NAME(Inputs::ActivePage), InParams.OperatorSettings);
+		FBoolReadRef InAutoPage = InputData.GetOrCreateDefaultDataReadReference<bool>(METASOUND_GET_PARAM_NAME(Inputs::AutoPage), InParams.OperatorSettings);
+		FBoolReadRef InAutoPagePlaysBlankPages = InputData.GetOrCreateDefaultDataReadReference<bool>(METASOUND_GET_PARAM_NAME(Inputs::AutoPagePlaysBlankPages), InParams.OperatorSettings);
+		FBoolReadRef InLoop = InputData.GetOrCreateDefaultDataReadReference<bool>(METASOUND_GET_PARAM_NAME(Inputs::Loop), InParams.OperatorSettings);
+		FBoolReadRef InEnabled = InputData.GetOrCreateDefaultDataReadReference<bool>(METASOUND_GET_PARAM_NAME(Inputs::Enabled), InParams.OperatorSettings);
 		return MakeUnique<FStepSequencePlayerOperator>(InParams, InSequenceAsset, InTransport, InMidiClock, InSpeed, InVelocityMult, InMaxColumns, InAdditionalOctaves, InStepSize, InActivePage, InAutoPage, InAutoPagePlaysBlankPages, InLoop, InEnabled);
 	}
 
@@ -258,21 +277,19 @@ namespace HarmonixMetasound
 
 	void FStepSequencePlayerOperator::BindInputs(FInputVertexInterfaceData& InVertexData)
 	{
-		using namespace StepSequencePlayerPinNames;
-		using namespace CommonPinNames;
-		InVertexData.BindReadVertex(METASOUND_GET_PARAM_NAME(InputSequenceAsset), SequenceAssetInPin);
+		InVertexData.BindReadVertex(METASOUND_GET_PARAM_NAME(Inputs::SequenceAsset), SequenceAssetInPin);
 		InVertexData.BindReadVertex(METASOUND_GET_PARAM_NAME(Inputs::Transport), TransportInPin);
 		InVertexData.BindReadVertex(METASOUND_GET_PARAM_NAME(Inputs::MidiClock), MidiClockInPin);
 		InVertexData.BindReadVertex(METASOUND_GET_PARAM_NAME(Inputs::Speed), SpeedMultInPin);
-		InVertexData.BindReadVertex(METASOUND_GET_PARAM_NAME(InputVelocityMultiplier), VelocityMultInPin);
-		InVertexData.BindReadVertex(METASOUND_GET_PARAM_NAME(InputMaxColumns), MaxColumnsInPin);
-		InVertexData.BindReadVertex(METASOUND_GET_PARAM_NAME(InputAdditionalOctaves), AdditionalOctavesInPin);
-		InVertexData.BindReadVertex(METASOUND_GET_PARAM_NAME(InputStepSizeQuarterNotes), StepSizeQuarterNotesInPin);
-		InVertexData.BindReadVertex(METASOUND_GET_PARAM_NAME(InputActivePage), ActivePageInPin);
-		InVertexData.BindReadVertex(METASOUND_GET_PARAM_NAME(InputAutoPage), AutoPageInPin);
-		InVertexData.BindReadVertex(METASOUND_GET_PARAM_NAME(InputAutoPagePlaysBlankPages), AutoPagePlaysBlankPagesInPin);
-		InVertexData.BindReadVertex(METASOUND_GET_PARAM_NAME(InputLoop), LoopInPin);
-		InVertexData.BindReadVertex(METASOUND_GET_PARAM_NAME(InputEnabled), EnabledInPin);
+		InVertexData.BindReadVertex(METASOUND_GET_PARAM_NAME(Inputs::VelocityMultiplier), VelocityMultInPin);
+		InVertexData.BindReadVertex(METASOUND_GET_PARAM_NAME(Inputs::MaxColumns), MaxColumnsInPin);
+		InVertexData.BindReadVertex(METASOUND_GET_PARAM_NAME(Inputs::AdditionalOctaves), AdditionalOctavesInPin);
+		InVertexData.BindReadVertex(METASOUND_GET_PARAM_NAME(Inputs::StepSizeQuarterNotes), StepSizeQuarterNotesInPin);
+		InVertexData.BindReadVertex(METASOUND_GET_PARAM_NAME(Inputs::ActivePage), ActivePageInPin);
+		InVertexData.BindReadVertex(METASOUND_GET_PARAM_NAME(Inputs::AutoPage), AutoPageInPin);
+		InVertexData.BindReadVertex(METASOUND_GET_PARAM_NAME(Inputs::AutoPagePlaysBlankPages), AutoPagePlaysBlankPagesInPin);
+		InVertexData.BindReadVertex(METASOUND_GET_PARAM_NAME(Inputs::Loop), LoopInPin);
+		InVertexData.BindReadVertex(METASOUND_GET_PARAM_NAME(Inputs::Enabled), EnabledInPin);
 
 		Init();
 	}
diff --git a/Engine/Plugins/Runtime/Harmonix/Source/HarmonixMetasound/Public/HarmonixMetasound/Nodes/StepSequencePlayerNode.h b/Engine/Plugins/Runtime/Harmonix/Source/HarmonixMetasound/Public/HarmonixMetasound/Nodes/StepSequencePlayerNode.h
new file mode 100644
--- /dev/null
+++ b/Engine/Plugins/Runtime/Harmonix/Source/HarmonixMetasound/Public/HarmonixMetasound/Nodes/StepSequencePlayerNode.h
@@ -0,0 +1,35 @@
+﻿// Copyright Epic Games, Inc. All Rights Reserved.
+
+#pragma once
+
+#include "MetasoundNodeInterface.h"
+
+#include "HarmonixMetasound/Common.h"
+
+namespace HarmonixMetasound::Nodes::StepSequencePlayer
+{
+	HARMONIXMETASOUND_API Metasound::FNodeClassName GetClassName();
+	HARMONIXMETASOUND_API int32 GetCurrentMajorVersion();
+
+	namespace Inputs
+	{
+		DECLARE_METASOUND_PARAM_EXTERN(SequenceAsset, "Step Sequence Asset", "Step sequence to play.");
+		DECLARE_METASOUND_PARAM_EXTERN(VelocityMultiplier, "Velocity Multiplier", "Multiplies the current note velocity by this number");
+		DECLARE_METASOUND_PARAM_EXTERN(MaxColumns, "Max Columns", "The Maximinum Number of cells to play per step sequence row.");
+		DECLARE_METASOUND_PARAM_EXTERN(AdditionalOctaves, "Additional Octaves", "The number of octaves to add to the authored step sequence note.");
+		DECLARE_METASOUND_PARAM_EXTERN(StepSizeQuarterNotes, "Step Size Quarter Notes", "The size, in quarter notes, of each step");
+		DECLARE_METASOUND_PARAM_EXTERN(ActivePage, "Active Page", "The page of the step sequence to play (1 indexed)");
+		DECLARE_METASOUND_PARAM_EXTERN(AutoPage, "Auto Page", "Whether to calculate the page of the step sequence based on current position");
+		DECLARE_METASOUND_PARAM_EXTERN(AutoPagePlaysBlankPages, "Auto Page Plays Blank Pages", "If autopaging, should blank pages be played?");
+		DECLARE_METASOUND_PARAM_ALIAS(Transport);
+		DECLARE_METASOUND_PARAM_ALIAS(MidiClock);
+		DECLARE_METASOUND_PARAM_ALIAS(Speed);
+		DECLARE_METASOUND_PARAM_ALIAS(Loop);
+		DECLARE_METASOUND_PARAM_ALIAS(Enabled);
+	}
+
+	namespace Outputs
+	{
+		DECLARE_METASOUND_PARAM_ALIAS(MidiStream);
+	}
+}
