MIDI Writer Node
- MidiStreamWriter class which contains the actual functionality
- MidiStreamWriterNode which wraps the functionality
- MIDI files are written out to "AudioCaptures/../MIDICaptures"
- MIDI files are timestamped for uniqueness

The Midi Stream Writer uses the "Clock Advance" events to determine the relative tick of the midi events to write out.

#tests Harmonix Functional Tests

diff --git a/Engine/Plugins/Runtime/Harmonix/Source/HarmonixMetasound/Private/HarmonixMetasound/MidiOps/MidiStreamWriter.cpp b/Engine/Plugins/Runtime/Harmonix/Source/HarmonixMetasound/Private/HarmonixMetasound/MidiOps/MidiStreamWriter.cpp
new file mode 100644
--- /dev/null
+++ b/Engine/Plugins/Runtime/Harmonix/Source/HarmonixMetasound/Private/HarmonixMetasound/MidiOps/MidiStreamWriter.cpp
@@ -0,0 +1,56 @@
+﻿// Copyright Epic Games, Inc. All Rights Reserved.
+
+#include "HarmonixMetasound/MidiOps/MidiStreamWriter.h"
+#include "HarmonixMetasound/DataTypes/MidiStream.h"
+#include "HarmonixMidi/MidiWriter.h"
+
+namespace Harmonix::Midi::Ops
+{
+	using namespace HarmonixMetasound;
+	FMidiStreamWriter::FMidiStreamWriter(TUniquePtr<FArchive>&& InArchive)
+		: Archive(MoveTemp(InArchive))
+	{}
+
+	void FMidiStreamWriter::Process(const FMidiStream& InStream)
+	{
+		ensureMsgf(InStream.GetClock(), TEXT("Midi stream must have a midi clock for the MidiStreamWriter to process it"));
+		if (TSharedPtr<const FMidiClock, ESPMode::NotThreadSafe> MidiClock = InStream.GetClock())
+		{
+			for (const FMidiClockEvent& ClockEvent : MidiClock->GetMidiClockEventsInBlock())
+			{
+				if (ClockEvent.Msg.Type == FMidiClockMsg::EType::AdvanceThru && !ClockEvent.Msg.AsAdvanceThru().IsPreRoll)
+				{
+					Process(InStream, ClockEvent.Msg.FromTick(), ClockEvent.Msg.ThruTick());
+				}
+			}
+		}
+	}
+
+
+	void FMidiStreamWriter::Process(const FMidiStream& InStream, int32 FromTick, int32 ThruTick)
+	{
+		bool AddedEvents = false;
+		for (const FMidiStreamEvent& MidiEvent : InStream.GetEventsInBlock())
+		{
+			if (MidiEvent.CurrentMidiTick > FromTick && MidiEvent.CurrentMidiTick <= ThruTick)
+			{
+				const int32 OffsetTick = MidiEvent.CurrentMidiTick - FromTick;
+				const int32 MidiTick = CurrentWriteTick + OffsetTick;
+				FMidiTrack& MidiTrack = MidiTracks.FindOrAdd(MidiEvent.TrackIndex);
+				MidiTrack.AddEvent(FMidiEvent(MidiTick, MidiEvent.MidiMessage));
+				AddedEvents = true;
+			}
+		}
+		CurrentWriteTick += ThruTick - FromTick;
+
+		if (AddedEvents)
+		{
+			Archive->Seek(0);
+			FMidiWriter MidiWriter = FMidiWriter(*Archive);
+			for (TPair<int32, const FMidiTrack&> Pair : MidiTracks)
+			{
+				Pair.Value.WriteStdMidi(MidiWriter);
+			}
+		}
+	}
+};
\ No newline at end of file
diff --git a/Engine/Plugins/Runtime/Harmonix/Source/HarmonixMetasound/Private/HarmonixMetasound/Nodes/MidiStreamWriterNode.cpp b/Engine/Plugins/Runtime/Harmonix/Source/HarmonixMetasound/Private/HarmonixMetasound/Nodes/MidiStreamWriterNode.cpp
new file mode 100644
--- /dev/null
+++ b/Engine/Plugins/Runtime/Harmonix/Source/HarmonixMetasound/Private/HarmonixMetasound/Nodes/MidiStreamWriterNode.cpp
@@ -0,0 +1,188 @@
+﻿// Copyright Epic Games, Inc. All Rights Reserved.
+
+#include "HarmonixMetasound/Nodes/MidiStreamWriterNode.h"
+
+#include "MetasoundExecutableOperator.h"
+#include "MetasoundFacade.h"
+#include "MetasoundNodeInterface.h"
+#include "MetasoundParamHelper.h"
+#include "MetasoundNodeRegistrationMacro.h"
+#include "MetasoundStandardNodesCategories.h"
+#include "HarmonixMetasound/Common.h"
+#include "HarmonixMetasound/DataTypes/MidiStream.h"
+#include "HarmonixMetasound/MidiOps/MidiStreamWriter.h"
+
+#include "Misc/Paths.h"
+#include "Misc/DateTime.h"
+#include "HAL/FileManager.h"
+
+#define LOCTEXT_NAMESPACE "HarmonixMetaSound"
+
+namespace HarmonixMetasound::Nodes::MidiStreamWriter
+{
+	using namespace Metasound;
+
+	const FNodeClassName& GetClassName()
+	{
+		static FNodeClassName ClassName
+		{
+			HarmonixNodeNamespace,
+			"MIDIStreamWriter",
+			""
+		};
+		return ClassName;
+	}
+
+	int32 GetCurrentMajorVersion()
+	{
+		return 0;
+	}
+
+	namespace Inputs
+	{
+		DEFINE_METASOUND_PARAM_ALIAS(Enable, CommonPinNames::Inputs::Enable);
+		DEFINE_METASOUND_PARAM_ALIAS(MidiStream, CommonPinNames::Inputs::MidiStream);
+		DEFINE_INPUT_METASOUND_PARAM(FilenamePrefix,  "Filename Prefix", "Filename Prefix of file to write to");
+	}
+
+	class FMidiStreamWriterOperator final : public TExecutableOperator<FMidiStreamWriterOperator>
+	{
+	public:
+		static const FNodeClassMetadata& GetNodeInfo()
+		{
+			using namespace HarmonixMetasound;
+			auto InitNodeInfo = []() -> FNodeClassMetadata
+			{
+				FNodeClassMetadata Info;
+				Info.ClassName = GetClassName();
+				Info.MajorVersion = GetCurrentMajorVersion();
+				Info.MinorVersion = 1;
+				Info.DisplayName = METASOUND_LOCTEXT("MidiStreamWriterNode_DisplayName", "MIDI Writer");
+				Info.Description = METASOUND_LOCTEXT("MidiStreamWriterNode_Description", "Writes the input midi stream to a standard midi file");
+				Info.Author = PluginAuthor;
+				Info.PromptIfMissing = PluginNodeMissingPrompt;
+				Info.DefaultInterface = GetVertexInterface();
+				Info.CategoryHierarchy = { MetasoundNodeCategories::Harmonix, NodeCategories::Music}; 
+
+				return Info;
+			};
+
+			static const FNodeClassMetadata Info = InitNodeInfo();
+			return Info;
+		}
+		
+		static const FVertexInterface& GetVertexInterface()
+		{
+			static const FVertexInterface Interface(
+				FInputVertexInterface(
+					TInputDataVertex<bool>(METASOUND_GET_PARAM_NAME_AND_METADATA(Inputs::Enable), true),
+					TInputDataVertex<FString>(METASOUND_GET_PARAM_NAME_AND_METADATA(Inputs::FilenamePrefix)),
+					TInputDataVertex<FMidiStream>(METASOUND_GET_PARAM_NAME_AND_METADATA(Inputs::MidiStream))
+				),
+				FOutputVertexInterface()
+			);
+			return Interface;
+		}
+
+		static TUniquePtr<IOperator> CreateOperator(const FBuildOperatorParams& InParams, FBuildResults& OutResults)
+		{
+			const FInputVertexInterfaceData& InputData = InParams.InputData;
+
+			FBoolReadRef InEnabled = InputData.GetOrCreateDefaultDataReadReference<bool>(METASOUND_GET_PARAM_NAME(Inputs::Enable), InParams.OperatorSettings);
+			FStringReadRef InFilenamePrefix = InputData.GetOrCreateDefaultDataReadReference<FString>(METASOUND_GET_PARAM_NAME(Inputs::FilenamePrefix), InParams.OperatorSettings);
+			FMidiStreamReadRef InMidiStream = InputData.GetOrConstructDataReadReference<FMidiStream>(METASOUND_GET_PARAM_NAME(Inputs::MidiStream));
+			
+			return MakeUnique<FMidiStreamWriterOperator>(InParams, InEnabled, InFilenamePrefix, InMidiStream);
+		}
+
+		FMidiStreamWriterOperator(const FBuildOperatorParams& InParams, const FBoolReadRef& InEnable, const FStringReadRef& InFilenamePrefix, const FMidiStreamReadRef& InMidiStream)
+			: EnableInPin(InEnable)
+			, FilenamePrefixInPin(InFilenamePrefix)
+			, MidiStreamInPin(InMidiStream)
+		{
+			Reset(InParams);
+		}
+
+		virtual void BindInputs(FInputVertexInterfaceData& InVertexData) override
+		{
+			InVertexData.BindReadVertex(METASOUND_GET_PARAM_NAME(Inputs::Enable), EnableInPin);
+			InVertexData.BindReadVertex(METASOUND_GET_PARAM_NAME(Inputs::FilenamePrefix), FilenamePrefixInPin);
+			InVertexData.BindReadVertex(METASOUND_GET_PARAM_NAME(Inputs::MidiStream), MidiStreamInPin);
+		}
+		
+		virtual void BindOutputs(FOutputVertexInterfaceData& InVertexData) override
+		{
+			
+		}
+		
+		void Reset(const FResetParams& ResetParams)
+		{
+			MidiStreamWriter.Reset();
+		}
+		
+		void Execute()
+		{
+			SetEnabled(*EnableInPin);
+			if (bEnabled && MidiStreamWriter.IsValid())
+			{
+				MidiStreamWriter->Process(*MidiStreamInPin);
+			}
+		}
+
+	private:
+
+		void SetEnabled(bool NewEnabled)
+		{
+			// enable
+			if (!bEnabled && NewEnabled)
+			{
+				const FString Filepath = GetMidiCaptureDir() / GenerateTimestampedFilename(*FilenamePrefixInPin);
+				TUniquePtr<FArchive> Archive = TUniquePtr<FArchive>(IFileManager::Get().CreateFileWriter(*Filepath, IO_WRITE));
+				MidiStreamWriter = MakeUnique<Harmonix::Midi::Ops::FMidiStreamWriter>(MoveTemp(Archive));
+				
+			}
+			// disbale
+			else if (bEnabled && !NewEnabled)
+			{
+				MidiStreamWriter.Reset();
+			}
+			bEnabled = NewEnabled;
+		}
+
+		static FString GetMidiCaptureDir()
+		{
+			static const FString MidiCaptureDir = FPaths::AudioCaptureDir() / TEXT("../MIDICaptures");
+			return MidiCaptureDir;
+		}
+
+		static FString GenerateTimestampedFilename(const FString& BaseName)
+		{
+			static const FString Separator = TEXT("_");
+			static const FString Extension = TEXT(".midi");
+			const FString FileID = FDateTime::Now().ToString(TEXT("%Y%m%d_%H%M%S"));
+			return BaseName + Separator + FileID + Extension;
+		}
+		
+		//** INPUTS
+		FBoolReadRef EnableInPin;
+		FStringReadRef FilenamePrefixInPin;
+		FMidiStreamReadRef MidiStreamInPin;
+
+		//** DATA
+		bool bEnabled = false;
+		TUniquePtr<Harmonix::Midi::Ops::FMidiStreamWriter> MidiStreamWriter;
+	};
+
+	class FMidiStreamWriterNode : public FNodeFacade
+	{
+	public:
+		FMidiStreamWriterNode(const FNodeInitData& InInitData)
+			: FNodeFacade(InInitData.InstanceName, InInitData.InstanceID, TFacadeOperatorClass<FMidiStreamWriterOperator>())
+		{}
+	};
+	
+	METASOUND_REGISTER_NODE(FMidiStreamWriterNode);
+}
+
+
+#undef LOCTEXT_NAMESPACE // "HarmonixMetaSound"
diff --git a/Engine/Plugins/Runtime/Harmonix/Source/HarmonixMetasound/Public/HarmonixMetasound/MidiOps/MidiStreamWriter.h b/Engine/Plugins/Runtime/Harmonix/Source/HarmonixMetasound/Public/HarmonixMetasound/MidiOps/MidiStreamWriter.h
new file mode 100644
--- /dev/null
+++ b/Engine/Plugins/Runtime/Harmonix/Source/HarmonixMetasound/Public/HarmonixMetasound/MidiOps/MidiStreamWriter.h
@@ -0,0 +1,47 @@
+﻿// Copyright Epic Games, Inc. All Rights Reserved.
+
+#pragma once
+
+#include "HarmonixMidi/MidiTrack.h"
+ 
+ namespace HarmonixMetasound
+ {
+ 	class FMidiStream;
+ 	class FMidiClock;
+ }
+
+namespace Harmonix::Midi::Ops
+{
+	/**
+	 * Write the midi events from a midi stream into a midi file
+	 */
+	class HARMONIXMETASOUND_API FMidiStreamWriter
+	{
+	public:
+		
+		FMidiStreamWriter(TUniquePtr<FArchive>&& Archive);
+
+		/**
+		 * Writes the midi events in the midi stream
+		 * Midi events are appended to the existing midi events
+		 * 
+		 * @param InStream The MIDI stream to read from
+		 */
+		void Process(const HarmonixMetasound::FMidiStream& InStream);
+		
+		/**
+		 * Writes the midi events in the midi stream
+		 * Midi events are appended to the existing midi events
+		 * 
+		 * @param InStream The MIDI stream to read from
+		 */
+		void Process(const HarmonixMetasound::FMidiStream& InStream, int32 FromTick, int32 ThruTick);
+		
+
+	private:
+		TUniquePtr<FArchive> Archive;
+		TMap<int32, FMidiTrack> MidiTracks;
+		int32 CurrentWriteTick = 0;
+	
+	};
+};
\ No newline at end of file
diff --git a/Engine/Plugins/Runtime/Harmonix/Source/HarmonixMetasound/Public/HarmonixMetasound/Nodes/MidiStreamWriterNode.h b/Engine/Plugins/Runtime/Harmonix/Source/HarmonixMetasound/Public/HarmonixMetasound/Nodes/MidiStreamWriterNode.h
new file mode 100644
--- /dev/null
+++ b/Engine/Plugins/Runtime/Harmonix/Source/HarmonixMetasound/Public/HarmonixMetasound/Nodes/MidiStreamWriterNode.h
@@ -0,0 +1,24 @@
+﻿// Copyright Epic Games, Inc. All Rights Reserved.
+
+#pragma once
+
+#include "MetasoundNodeInterface.h"
+
+#include "HarmonixMetasound/Common.h"
+
+namespace HarmonixMetasound::Nodes::MidiStreamWriter
+{
+	HARMONIXMETASOUND_API const Metasound::FNodeClassName& GetClassName();
+	HARMONIXMETASOUND_API int32 GetCurrentMajorVersion();
+
+	namespace Inputs
+	{
+		DECLARE_METASOUND_PARAM_ALIAS(Enable);
+		DECLARE_METASOUND_PARAM_ALIAS(MidiStream);
+		DECLARE_METASOUND_PARAM_EXTERN(FilenamePrefix);
+	}
+
+	namespace Outputs
+	{
+	}
+}
