[Harmonix] MIDI Clock: Track tempo and time signature changes. Use TVariant for FMidiClockMsg

#tests automated tests

diff --git a/Engine/Plugins/Runtime/Harmonix/Source/HarmonixMetasound/Private/HarmonixMetasound/DataTypes/MidiClock.cpp b/Engine/Plugins/Runtime/Harmonix/Source/HarmonixMetasound/Private/HarmonixMetasound/DataTypes/MidiClock.cpp
--- a/Engine/Plugins/Runtime/Harmonix/Source/HarmonixMetasound/Private/HarmonixMetasound/DataTypes/MidiClock.cpp
+++ b/Engine/Plugins/Runtime/Harmonix/Source/HarmonixMetasound/Private/HarmonixMetasound/DataTypes/MidiClock.cpp
@@ -435,48 +435,45 @@ namespace HarmonixMetasound
 	
 	void FMidiClock::HandleClockEvent(const FMidiClock& DrivingClock, const FMidiClockEvent& Event, int32 PrerollBars, float Speed)
 	{
-		switch (Event.Msg.Type)
+		if (Event.Msg.IsType<MidiClockMessageTypes::FReset>())
+		{
+			const int32 Tick = CalculateMappedTick(Event.Msg.Get<MidiClockMessageTypes::FReset>().ToTick);
+			SeekTo(Event.BlockFrameIndex, Tick + 1, PrerollBars);
+		}
+		else if (Event.Msg.IsType<MidiClockMessageTypes::FLoop>())
 		{
-		case FMidiClockMsg::EType::Reset:
-			{
-				int32 Tick = CalculateMappedTick(Event.Msg.ToTick());
-				SeekTo(Event.BlockFrameIndex, Tick + 1, PrerollBars);
-				break;
-			}
-		case FMidiClockMsg::EType::Loop:
 			// ignore loops since the driving clock will loop us by seeking
-			break;
-		case FMidiClockMsg::EType::SeekTo:
-			{
-				int32 Tick = CalculateMappedTick(Event.Msg.ToTick());
-				SeekTo(Event.BlockFrameIndex, Tick, PrerollBars);
-				break;
-			}
-		case FMidiClockMsg::EType::SeekThru:
+		}
+		else if (Event.Msg.IsType<MidiClockMessageTypes::FSeekTo>())
+		{
+			const int32 Tick = CalculateMappedTick(Event.Msg.Get<MidiClockMessageTypes::FSeekTo>().ToTick);
+			SeekTo(Event.BlockFrameIndex, Tick, PrerollBars);
+		}
+		else if (Event.Msg.IsType<MidiClockMessageTypes::FSeekThru>())
+		{
+			const int32 Tick = CalculateMappedTick(Event.Msg.Get<MidiClockMessageTypes::FSeekThru>().ThruTick);
+			SeekTo(Event.BlockFrameIndex, Tick + 1, PrerollBars);
+		}
+		else if (Event.Msg.IsType<MidiClockMessageTypes::FAdvanceThru>())
+		{
+			const MidiClockMessageTypes::FAdvanceThru& AdvanceThru = Event.Msg.Get<MidiClockMessageTypes::FAdvanceThru>();
+			
+			if (AdvanceThru.IsPreRoll)
 			{
-				int32 Tick = CalculateMappedTick(Event.Msg.ThruTick());
+				const int32 Tick = CalculateMappedTick(AdvanceThru.ThruTick);
 				SeekTo(Event.BlockFrameIndex, Tick + 1, PrerollBars);
-				break;
 			}
-		case FMidiClockMsg::EType::AdvanceThru:
+			else
 			{
-				if (Event.Msg.AsAdvanceThru().IsPreRoll)
-				{
-					int32 Tick = CalculateMappedTick(Event.Msg.ThruTick());
-					SeekTo(Event.BlockFrameIndex, Tick + 1, PrerollBars);
-					break;
-				}
-
 				// Advance based on the delta ticks, and not based on the absolute tick
-				int32 Tick = GetCurrentMidiTick() + (Event.Msg.ThruTick() - Event.Msg.FromTick());
-				float Ms = GetSongMaps().TickToMs(Tick);
+				const int32 Tick = GetCurrentMidiTick() + (AdvanceThru.ThruTick - AdvanceThru.FromTick);
+				const float Ms = GetSongMaps().TickToMs(Tick);
 
-				float ClockInSpeed = DrivingClock.GetSpeedAtBlockSampleFrame(Event.BlockFrameIndex);
-				float AdvanceRatio = DrivingClock.GetSongMaps().GetTempoAtTick(Event.Msg.FromTick())
+				const float ClockInSpeed = DrivingClock.GetSpeedAtBlockSampleFrame(Event.BlockFrameIndex);
+				const float AdvanceRatio = DrivingClock.GetSongMaps().GetTempoAtTick(AdvanceThru.FromTick)
 								   / GetSongMaps().GetTempoAtTick(GetCurrentMidiTick());
 				InformOfCurrentAdvanceRate(ClockInSpeed * Speed * AdvanceRatio);
 				AdvanceHiResToMs(Event.BlockFrameIndex, Ms, true);
-				break;
 			}
 		}
 	}
@@ -581,49 +578,52 @@ namespace HarmonixMetasound
 		: MyMidiClock(MidiClock)
 	{
 		check(MidiClock);
+
+		// Ignore midi messages, test events, and preroll notes
+		SetMessageFilter(EFilterPassFlags::All & ~(EFilterPassFlags::MidiMessage | EFilterPassFlags::Text | EFilterPassFlags::PreRollNoteOn));
 	}
 
 	void FMidiClock::FMidiClockEventCursor::Reset(bool ForceNoBroadcast /*= false*/)
 	{
 		const int32 FromTick = CurrentTick;
 		FMidiPlayCursor::Reset(ForceNoBroadcast);
-		AddEvent(FMidiClockEvent(MyMidiClock->GetCurrentBlockFrameIndex(), FMidiClockMsg::FReset(FromTick, CurrentTick, ForceNoBroadcast)));
+		AddEvent(FMidiClockEvent(MyMidiClock->GetCurrentBlockFrameIndex(), MidiClockMessageTypes::FReset(FromTick, CurrentTick, ForceNoBroadcast)));
 	}
 
 	void FMidiClock::FMidiClockEventCursor::OnLoop(int32 LoopStartTick, int32 LoopEndTick)
 	{
 		FMidiPlayCursor::OnLoop(LoopStartTick, LoopEndTick);
-		AddEvent(FMidiClockEvent(MyMidiClock->GetCurrentBlockFrameIndex(), FMidiClockMsg::FLoop(LoopStartTick, LoopEndTick)));
+		AddEvent(FMidiClockEvent(MyMidiClock->GetCurrentBlockFrameIndex(), MidiClockMessageTypes::FLoop(LoopStartTick, LoopEndTick)));
 	}
 
 	void FMidiClock::FMidiClockEventCursor::SeekToTick(int32 Tick) 
 	{
 		const int32 FromTick = CurrentTick;
 		FMidiPlayCursor::SeekToTick(Tick);
-		AddEvent(FMidiClockEvent(MyMidiClock->GetCurrentBlockFrameIndex(), FMidiClockMsg::FSeekTo(FromTick, Tick)));
+		AddEvent(FMidiClockEvent(MyMidiClock->GetCurrentBlockFrameIndex(), MidiClockMessageTypes::FSeekTo(FromTick, Tick)));
 	}
 
 	void FMidiClock::FMidiClockEventCursor::SeekThruTick(int32 Tick)
 	{
 		const int32 FromTick = CurrentTick;
 		FMidiPlayCursor::SeekThruTick(Tick);
-		AddEvent(FMidiClockEvent(MyMidiClock->GetCurrentBlockFrameIndex(), FMidiClockMsg::FSeekThru(FromTick, Tick)));
+		AddEvent(FMidiClockEvent(MyMidiClock->GetCurrentBlockFrameIndex(), MidiClockMessageTypes::FSeekThru(FromTick, Tick)));
 	}
 
 	void FMidiClock::FMidiClockEventCursor::AdvanceThruTick(int32 Tick, bool IsPreRoll)
 	{
-		const int32 FromTick = CurrentTick;
+		CurrentAdvanceStartTick = CurrentTick;
 		FMidiPlayCursor::AdvanceThruTick(Tick, IsPreRoll);
-		AddEvent(FMidiClockEvent(MyMidiClock->GetCurrentBlockFrameIndex(), FMidiClockMsg::FAdvanceThru(FromTick, Tick, IsPreRoll)));
+		AddEvent(FMidiClockEvent(MyMidiClock->GetCurrentBlockFrameIndex(), MidiClockMessageTypes::FAdvanceThru(CurrentAdvanceStartTick, Tick, IsPreRoll)));
 	}
 
 	void FMidiClock::FMidiClockEventCursor::OnTempo(int32 TrackIndex, int32 Tick, int32 Tempo, bool IsPreroll)
 	{
-		int32 BlockFrameIndex = MyMidiClock->CurrentBlockFrameIndex;
+		const int32 BlockFrameIndex = MyMidiClock->CurrentBlockFrameIndex;
 
 		check(BlockFrameIndex >= MyMidiClock->TempoChangesInBlock.Last().BlockSampleFrameIndex);
 		MyMidiClock->HasTempoChangeInBlock = true;
-		float Bpm = Harmonix::Midi::Constants::MidiTempoToBPM(Tempo);
+		const float Bpm = Harmonix::Midi::Constants::MidiTempoToBPM(Tempo);
 		if (MyMidiClock->TempoChangesInBlock.Last().BlockSampleFrameIndex == BlockFrameIndex)
 		{
 			MyMidiClock->TempoChangesInBlock.Last().Tempo = Bpm;
@@ -632,9 +632,32 @@ namespace HarmonixMetasound
 		{ 
 			MyMidiClock->TempoChangesInBlock.Add({BlockFrameIndex, 0.0f, Bpm});
 		}
+
+		// This gets fired off as a result of AdvanceThruTick, and we want the event to go in between advances,
+		// so we add an AdvanceThru event up to this tick. But if this tempo event is on the first tick, we don't need to do that.
+		if (Tick > CurrentAdvanceStartTick)
+		{
+			AddEvent(FMidiClockEvent(MyMidiClock->GetCurrentBlockFrameIndex(), MidiClockMessageTypes::FAdvanceThru(CurrentAdvanceStartTick, Tick, IsPreroll)));
+			CurrentAdvanceStartTick = Tick;
+		}
+		
+		AddEvent(FMidiClockEvent(MyMidiClock->GetCurrentBlockFrameIndex(), MidiClockMessageTypes::FTempoChange(Tick, Bpm)));
+	}
+
+	void FMidiClock::FMidiClockEventCursor::OnTimeSig(int32 TrackIndex, int32 Tick, int32 Numerator, int32 Denominator, bool IsPreroll)
+	{
+		// This gets fired off as a result of AdvanceThruTick, and we want the event to go in between advances,
+		// so we add an AdvanceThru event up to this tick. But if this tempo event is on the first tick, we don't need to do that.
+		if (Tick > CurrentAdvanceStartTick)
+		{
+			AddEvent(FMidiClockEvent(MyMidiClock->GetCurrentBlockFrameIndex(), MidiClockMessageTypes::FAdvanceThru(CurrentAdvanceStartTick, Tick, IsPreroll)));
+			CurrentAdvanceStartTick = Tick;
+		}
+		
+		AddEvent(FMidiClockEvent(MyMidiClock->GetCurrentBlockFrameIndex(), MidiClockMessageTypes::FTimeSignatureChange(Tick, { Numerator, Denominator })));
 	}
 
-	void FMidiClock::FMidiClockEventCursor::AddEvent(const FMidiClockEvent& InEvent)
+	void FMidiClock::FMidiClockEventCursor::AddEvent(const FMidiClockEvent& InEvent) const
 	{
 		TArray<FMidiClockEvent>& Events = MyMidiClock->MidiClockEventsInBlock;
 		if (!Events.IsEmpty())
diff --git a/Engine/Plugins/Runtime/Harmonix/Source/HarmonixMetasound/Private/HarmonixMetasound/DataTypes/MidiClockEvent.cpp b/Engine/Plugins/Runtime/Harmonix/Source/HarmonixMetasound/Private/HarmonixMetasound/DataTypes/MidiClockEvent.cpp
deleted file mode 100644
diff --git a/Engine/Plugins/Runtime/Harmonix/Source/HarmonixMetasound/Private/HarmonixMetasound/MidiOps/MidiStreamWriter.cpp b/Engine/Plugins/Runtime/Harmonix/Source/HarmonixMetasound/Private/HarmonixMetasound/MidiOps/MidiStreamWriter.cpp
--- a/Engine/Plugins/Runtime/Harmonix/Source/HarmonixMetasound/Private/HarmonixMetasound/MidiOps/MidiStreamWriter.cpp
+++ b/Engine/Plugins/Runtime/Harmonix/Source/HarmonixMetasound/Private/HarmonixMetasound/MidiOps/MidiStreamWriter.cpp
@@ -18,9 +18,14 @@ namespace Harmonix::Midi::Ops
 		{
 			for (const FMidiClockEvent& ClockEvent : MidiClock->GetMidiClockEventsInBlock())
 			{
-				if (ClockEvent.Msg.Type == FMidiClockMsg::EType::AdvanceThru && !ClockEvent.Msg.AsAdvanceThru().IsPreRoll)
+				if (ClockEvent.Msg.IsType<MidiClockMessageTypes::FAdvanceThru>())
 				{
-					Process(InStream, ClockEvent.Msg.FromTick(), ClockEvent.Msg.ThruTick());
+					const MidiClockMessageTypes::FAdvanceThru& AdvanceThru = ClockEvent.Msg.Get<MidiClockMessageTypes::FAdvanceThru>();
+					
+					if (!AdvanceThru.IsPreRoll)
+					{
+						Process(InStream, AdvanceThru.FromTick, AdvanceThru.ThruTick);
+					}
 				}
 			}
 		}
diff --git a/Engine/Plugins/Runtime/Harmonix/Source/HarmonixMetasound/Private/HarmonixMetasound/Nodes/MidiClockOffsetNode.cpp b/Engine/Plugins/Runtime/Harmonix/Source/HarmonixMetasound/Private/HarmonixMetasound/Nodes/MidiClockOffsetNode.cpp
--- a/Engine/Plugins/Runtime/Harmonix/Source/HarmonixMetasound/Private/HarmonixMetasound/Nodes/MidiClockOffsetNode.cpp
+++ b/Engine/Plugins/Runtime/Harmonix/Source/HarmonixMetasound/Private/HarmonixMetasound/Nodes/MidiClockOffsetNode.cpp
@@ -171,59 +171,58 @@ namespace HarmonixMetasound::Nodes::MidiClockOffset
 			for (int32 EventIndex = 0; EventIndex < MidiClockEvents.Num(); ++EventIndex)
 			{
 				const FMidiClockEvent& Event = MidiClockEvents[EventIndex];
-				switch (Event.Msg.Type)
+
+				if (Event.Msg.IsType<MidiClockMessageTypes::FSeekTo>())
+				{
+					const int32 Tick = GetTickWithOffset(Event.Msg.Get<MidiClockMessageTypes::FSeekTo>().ToTick, OffsetBars, OffsetBeats, OffsetMs);
+					FMusicSeekTarget SeekTarget;
+					SeekTarget.Type = ESeekPointType::Millisecond;
+					SeekTarget.Ms = MidiClockOut->GetSongMaps().TickToMs(Tick);
+								
+					MidiClockOut->SeekTo(Event.BlockFrameIndex, SeekTarget, PrerollBars);
+				}
+				else if (Event.Msg.IsType<MidiClockMessageTypes::FReset>())
+				{
+					const int32 Tick = GetTickWithOffset(Event.Msg.Get<MidiClockMessageTypes::FReset>().ToTick, OffsetBars, OffsetBeats, OffsetMs);
+					FMusicSeekTarget SeekTarget;
+					SeekTarget.Type = ESeekPointType::Millisecond;
+					SeekTarget.Ms = MidiClockOut->GetSongMaps().TickToMs(Tick);
+								
+					MidiClockOut->SeekTo(Event.BlockFrameIndex, SeekTarget, PrerollBars);
+				}
+				else if (Event.Msg.IsType<MidiClockMessageTypes::FSeekThru>())
 				{
-				case FMidiClockMsg::EType::SeekTo:
-				case FMidiClockMsg::EType::Reset:
+					const int32 Tick = GetTickWithOffset(Event.Msg.Get<MidiClockMessageTypes::FSeekThru>().ThruTick, OffsetBars, OffsetBeats, OffsetMs);
+					FMusicSeekTarget SeekTarget;
+					SeekTarget.Type = ESeekPointType::Millisecond;
+					SeekTarget.Ms = MidiClockOut->GetSongMaps().TickToMs(Tick + 1);
+							
+					MidiClockOut->SeekTo(Event.BlockFrameIndex, SeekTarget, PrerollBars);
+				}
+				else if (Event.Msg.IsType<MidiClockMessageTypes::FAdvanceThru>())
+				{
+					const int32 ThruTick = Event.Msg.Get<MidiClockMessageTypes::FAdvanceThru>().ThruTick;
+					const int32 Tick = GetTickWithOffset(ThruTick, OffsetBars, OffsetBeats, OffsetMs);
+						
+					// if our offset changed while we were advancing, seek to the new offset first
+					if (!FMath::IsNearlyEqual(PrevOffsetMs, OffsetMs) || !FMath::IsNearlyEqual(PrevOffsetBeats, OffsetBeats) || PrevOffsetBars != OffsetBars)
 					{
-						const int32 Tick = GetTickWithOffset(Event.Msg.ToTick(), OffsetBars, OffsetBeats, OffsetMs);
+						PrevOffsetMs = OffsetMs;
+						PrevOffsetBars = OffsetBars;
+						PrevOffsetBeats = OffsetBeats;
+
 						FMusicSeekTarget SeekTarget;
 						SeekTarget.Type = ESeekPointType::Millisecond;
 						SeekTarget.Ms = MidiClockOut->GetSongMaps().TickToMs(Tick);
 								
 						MidiClockOut->SeekTo(Event.BlockFrameIndex, SeekTarget, PrerollBars);
-						break;
-					}
-				case FMidiClockMsg::EType::SeekThru:
-					{
-						const int32 Tick = GetTickWithOffset(Event.Msg.ThruTick(), OffsetBars, OffsetBeats, OffsetMs);
-						FMusicSeekTarget SeekTarget;
-						SeekTarget.Type = ESeekPointType::Millisecond;
-						SeekTarget.Ms = MidiClockOut->GetSongMaps().TickToMs(Tick + 1);
-							
-						MidiClockOut->SeekTo(Event.BlockFrameIndex, SeekTarget, PrerollBars);
-						break;
 					}
-				case FMidiClockMsg::EType::AdvanceThru:
-					{
-						const int32 Tick = GetTickWithOffset(Event.Msg.ThruTick(), OffsetBars, OffsetBeats, OffsetMs);
-						
-						// if our offset changed while we were advancing, seek to the new offset first
-						if (!FMath::IsNearlyEqual(PrevOffsetMs, OffsetMs) || !FMath::IsNearlyEqual(PrevOffsetBeats, OffsetBeats) || PrevOffsetBars != OffsetBars)
-						{
-							PrevOffsetMs = OffsetMs;
-							PrevOffsetBars = OffsetBars;
-							PrevOffsetBeats = OffsetBeats;
 
-							FMusicSeekTarget SeekTarget;
-							SeekTarget.Type = ESeekPointType::Millisecond;
-							SeekTarget.Ms = MidiClockOut->GetSongMaps().TickToMs(Tick);
-								
-							MidiClockOut->SeekTo(Event.BlockFrameIndex, SeekTarget, PrerollBars);
-						}
-
-						const float AdvanceToMs = MidiClockOut->GetSongMaps().TickToMs(Tick);
-						const float ClockInSpeed = MidiClockIn->GetSpeedAtBlockSampleFrame(EventIndex);
-						const float AdvanceRatio = MidiClockIn->GetSongMaps().GetTempoAtTick(Event.Msg.ThruTick()) / MidiClockOut->GetSongMaps().GetTempoAtTick(Tick);
-						MidiClockOut->InformOfCurrentAdvanceRate(ClockInSpeed * AdvanceRatio);
-						MidiClockOut->AdvanceHiResToMs(Event.BlockFrameIndex, AdvanceToMs, true);
-						break;
-					}
-				case FMidiClockMsg::EType::Loop:
-					{
-						break;
-					}
-							
+					const float AdvanceToMs = MidiClockOut->GetSongMaps().TickToMs(Tick);
+					const float ClockInSpeed = MidiClockIn->GetSpeedAtBlockSampleFrame(EventIndex);
+					const float AdvanceRatio = MidiClockIn->GetSongMaps().GetTempoAtTick(ThruTick) / MidiClockOut->GetSongMaps().GetTempoAtTick(Tick);
+					MidiClockOut->InformOfCurrentAdvanceRate(ClockInSpeed * AdvanceRatio);
+					MidiClockOut->AdvanceHiResToMs(Event.BlockFrameIndex, AdvanceToMs, true);
 				}
 			}
 		}
diff --git a/Engine/Plugins/Runtime/Harmonix/Source/HarmonixMetasound/Private/HarmonixMetasound/Nodes/MidiClockSubdivisionTriggerNode.cpp b/Engine/Plugins/Runtime/Harmonix/Source/HarmonixMetasound/Private/HarmonixMetasound/Nodes/MidiClockSubdivisionTriggerNode.cpp
--- a/Engine/Plugins/Runtime/Harmonix/Source/HarmonixMetasound/Private/HarmonixMetasound/Nodes/MidiClockSubdivisionTriggerNode.cpp
+++ b/Engine/Plugins/Runtime/Harmonix/Source/HarmonixMetasound/Private/HarmonixMetasound/Nodes/MidiClockSubdivisionTriggerNode.cpp
@@ -240,19 +240,21 @@ namespace HarmonixMetasound
 		{
 			for (const FMidiClockEvent& ClockEvent : MidiClockInPin->GetMidiClockEventsInBlock())
 			{
-				if (ClockEvent.Msg.Type == FMidiClockMsg::EType::AdvanceThru)
+				if (ClockEvent.Msg.IsType<MidiClockMessageTypes::FAdvanceThru>())
 				{
-					if (ClockEvent.Msg.AsAdvanceThru().IsPreRoll)
+					const MidiClockMessageTypes::FAdvanceThru& AdvanceThru = ClockEvent.Msg.Get<MidiClockMessageTypes::FAdvanceThru>();
+					
+					if (AdvanceThru.IsPreRoll)
 					{
 						continue;
 					}
 					
-					int32 TickPreceedingThisAdvance = ClockEvent.Msg.FromTick() - GridOffsetTicks;
-					int32 LastTickProcessed = ClockEvent.Msg.ThruTick() - GridOffsetTicks;
+					int32 TickPreceedingThisAdvance = AdvanceThru.FromTick - GridOffsetTicks;
+					int32 LastTickProcessed = AdvanceThru.ThruTick - GridOffsetTicks;
 					const FSongMaps& SongMaps = MidiClockInPin->GetSongMaps();
 					
 					// calculate new GridSizeTicks and GridOffsetTicks at the tick we're evaluating based on any time signature changes
-					GridSizeTicks = SongMaps.SubdivisionToMidiTicks(GridSizeUnits, ClockEvent.Msg.ThruTick()) * GridSizeMultiplier;
+					GridSizeTicks = SongMaps.SubdivisionToMidiTicks(GridSizeUnits, AdvanceThru.ThruTick) * GridSizeMultiplier;
 					
 					// if this is one of the quantization types that is variable in size we have to do it a little different...
 					if (GridSizeUnits == EMidiClockSubdivisionQuantization::Bar ||
diff --git a/Engine/Plugins/Runtime/Harmonix/Source/HarmonixMetasound/Private/HarmonixMetasound/Nodes/MidiClockTimestampTriggerNode.cpp b/Engine/Plugins/Runtime/Harmonix/Source/HarmonixMetasound/Private/HarmonixMetasound/Nodes/MidiClockTimestampTriggerNode.cpp
--- a/Engine/Plugins/Runtime/Harmonix/Source/HarmonixMetasound/Private/HarmonixMetasound/Nodes/MidiClockTimestampTriggerNode.cpp
+++ b/Engine/Plugins/Runtime/Harmonix/Source/HarmonixMetasound/Private/HarmonixMetasound/Nodes/MidiClockTimestampTriggerNode.cpp
@@ -221,26 +221,29 @@ namespace HarmonixMetasound
 		{
 			for (const FMidiClockEvent& ClockEvent : MidiClockInPin->GetMidiClockEventsInBlock())
 			{
-				switch (ClockEvent.Msg.Type)
+				using namespace MidiClockMessageTypes;
+
+				bool ShouldTrigger = false;
+
+				if (ClockEvent.Msg.IsType<FAdvanceThru>())
 				{
-				case FMidiClockMsg::EType::AdvanceThru:
-					if (ClockEvent.Msg.AsAdvanceThru().IsPreRoll)
-					{
-						continue;
-					}
-				case FMidiClockMsg::EType::SeekThru:
+					const FAdvanceThru AdvanceThru = ClockEvent.Msg.Get<FAdvanceThru>();
+					ShouldTrigger = !AdvanceThru.IsPreRoll && AdvanceThru.FromTick < TriggerTick && AdvanceThru.ThruTick >= TriggerTick;
+				}
+				else if (ClockEvent.Msg.IsType<FSeekThru>())
+				{
+					if (*TriggerDuringSeekInPin)
 					{
-						const bool EventIsSeek = ClockEvent.Msg.Type == FMidiClockMsg::EType::SeekThru;
-
-						if (ClockEvent.Msg.FromTick() < TriggerTick && ClockEvent.Msg.ThruTick() >= TriggerTick)
-						{
-							if (!EventIsSeek || *TriggerDuringSeekInPin)
-							{
-								TriggerOutPin->TriggerFrame(ClockEvent.BlockFrameIndex);
-							}
-						}
+						const FSeekThru& SeekThru = ClockEvent.Msg.Get<FSeekThru>();
+						ShouldTrigger = SeekThru.FromTick < TriggerTick && SeekThru.ThruTick >= TriggerTick;
 					}
 				}
+				// TODO: SeekTo?
+				
+				if (ShouldTrigger)
+				{
+					TriggerOutPin->TriggerFrame(ClockEvent.BlockFrameIndex);
+				}
 			}
 		}
 	}
diff --git a/Engine/Plugins/Runtime/Harmonix/Source/HarmonixMetasound/Private/HarmonixMetasound/Nodes/StepSequencePlayerNode.cpp b/Engine/Plugins/Runtime/Harmonix/Source/HarmonixMetasound/Private/HarmonixMetasound/Nodes/StepSequencePlayerNode.cpp
--- a/Engine/Plugins/Runtime/Harmonix/Source/HarmonixMetasound/Private/HarmonixMetasound/Nodes/StepSequencePlayerNode.cpp
+++ b/Engine/Plugins/Runtime/Harmonix/Source/HarmonixMetasound/Private/HarmonixMetasound/Nodes/StepSequencePlayerNode.cpp
@@ -371,17 +371,18 @@ namespace HarmonixMetasound::Nodes::StepSequencePlayer
 
 						if (Event.BlockFrameIndex >= StartFrameIndex)
 						{
-							switch (Event.Msg.Type)
+							if (Event.Msg.IsType<MidiClockMessageTypes::FAdvanceThru>())
 							{
-							case FMidiClockMsg::EType::AdvanceThru:
-								AdvanceThruTick(Event.BlockFrameIndex, Event.Msg.ThruTick(), Event.Msg.AsAdvanceThru().IsPreRoll);
-								break;
-							case FMidiClockMsg::EType::SeekThru:
-								SeekThruTick(Event.BlockFrameIndex, Event.Msg.ThruTick());
-								break;
-							case FMidiClockMsg::EType::SeekTo:
-								SeekToTick(Event.BlockFrameIndex, Event.Msg.ToTick());
-								break;
+								const MidiClockMessageTypes::FAdvanceThru& AdvanceThru = Event.Msg.Get<MidiClockMessageTypes::FAdvanceThru>();
+								AdvanceThruTick(Event.BlockFrameIndex, AdvanceThru.ThruTick, AdvanceThru.IsPreRoll);
+							}
+							else if (Event.Msg.IsType<MidiClockMessageTypes::FSeekThru>())
+							{
+								SeekThruTick(Event.BlockFrameIndex, Event.Msg.Get<MidiClockMessageTypes::FSeekThru>().ThruTick);
+							}
+							else if (Event.Msg.IsType<MidiClockMessageTypes::FSeekTo>())
+							{
+								SeekToTick(Event.BlockFrameIndex, Event.Msg.Get<MidiClockMessageTypes::FSeekTo>().ToTick);
 							}
 						}
 					}
diff --git a/Engine/Plugins/Runtime/Harmonix/Source/HarmonixMetasound/Public/HarmonixMetasound/DataTypes/MidiClock.h b/Engine/Plugins/Runtime/Harmonix/Source/HarmonixMetasound/Public/HarmonixMetasound/DataTypes/MidiClock.h
--- a/Engine/Plugins/Runtime/Harmonix/Source/HarmonixMetasound/Public/HarmonixMetasound/DataTypes/MidiClock.h
+++ b/Engine/Plugins/Runtime/Harmonix/Source/HarmonixMetasound/Public/HarmonixMetasound/DataTypes/MidiClock.h
@@ -249,11 +249,13 @@ namespace HarmonixMetasound
 			virtual void SeekThruTick(int32 Tick) override;
 			virtual void AdvanceThruTick(int32 Tick, bool IsPreRoll) override;
 			virtual void OnTempo(int32 TrackIndex, int32 Tick, int32 Tempo, bool IsPreroll = false) override;
+			virtual void OnTimeSig(int32 TrackIndex, int32 Tick, int32 Numerator, int32 Denominator, bool IsPreroll) override;
 			//~ END FMidiPlayCursor Overrides
 			
-			void AddEvent(const FMidiClockEvent& InEvent);
+			void AddEvent(const FMidiClockEvent& InEvent) const;
 		private:
 			FMidiClock* MyMidiClock = nullptr;
+			int32 CurrentAdvanceStartTick = -1;
 		};
 		
 		FMidiClockEventCursor MidiClockEventCursor;
diff --git a/Engine/Plugins/Runtime/Harmonix/Source/HarmonixMetasound/Public/HarmonixMetasound/DataTypes/MidiClockEvent.h b/Engine/Plugins/Runtime/Harmonix/Source/HarmonixMetasound/Public/HarmonixMetasound/DataTypes/MidiClockEvent.h
--- a/Engine/Plugins/Runtime/Harmonix/Source/HarmonixMetasound/Public/HarmonixMetasound/DataTypes/MidiClockEvent.h
+++ b/Engine/Plugins/Runtime/Harmonix/Source/HarmonixMetasound/Public/HarmonixMetasound/DataTypes/MidiClockEvent.h
@@ -6,48 +6,31 @@
 
 namespace HarmonixMetasound
 {
-	struct HARMONIXMETASOUND_API FMidiClockMsg
+	namespace MidiClockMessageTypes
 	{
-		enum class EType : uint8
-		{
-			Reset,
-			Loop,
-			SeekTo,
-			SeekThru,
-			AdvanceThru
-		};
-		
 		struct FReset
 		{
 			int32 FromTick;
 			int32 ToTick;
 			bool ForceNoBroadcast;
-			FReset(int32 FromTick, int32 ToTick, bool ForceNoBroadcast)
-				: FromTick(FromTick), ToTick(ToTick), ForceNoBroadcast(ForceNoBroadcast) {}
 		};
 
 		struct FLoop
 		{
 			int32 LoopStartTick;
 			int32 LoopEndTick;
-			FLoop(int32 LoopStartTick, int32 LoopEndTick)
-				: LoopStartTick(LoopStartTick), LoopEndTick(LoopEndTick) {}
 		};
 
 		struct FSeekTo
 		{
 			int32 FromTick;
 			int32 ToTick;
-			FSeekTo(int32 FromTick, int32 ToTick)
-				: FromTick(FromTick), ToTick(ToTick) {}
 		};
 
 		struct FSeekThru
 		{
 			int32 FromTick;
 			int32 ThruTick;
-			FSeekThru(int32 FromTick, int32 ThruTick)
-				: FromTick(FromTick), ThruTick(ThruTick) {}
 		};
 
 		struct FAdvanceThru
@@ -55,55 +38,42 @@ namespace HarmonixMetasound
 			int32 FromTick;
 			int32 ThruTick;
 			bool IsPreRoll;
-			FAdvanceThru(int32 FromTick, int32 ThruTick, bool IsPreRoll)
-				: FromTick(FromTick), ThruTick(ThruTick), IsPreRoll(IsPreRoll) {}
 		};
 
-		FMidiClockMsg(const FReset& Reset) : Type(EType::Reset), Reset(Reset) {};
-		FMidiClockMsg(const FLoop& Loop) : Type(EType::Loop), Loop(Loop) {};
-		FMidiClockMsg(const FSeekTo& SeekTo) : Type(EType::SeekTo), SeekTo(SeekTo) {};
-		FMidiClockMsg(const FSeekThru& SeekThru) : Type(EType::SeekThru), SeekThru(SeekThru) {};
-		FMidiClockMsg(const FAdvanceThru AdvanceThru) : Type(EType::AdvanceThru), AdvanceThru(AdvanceThru) {};
-
-		// type safe way to access msg data
-		// asserts that is of the corresponding type
-		const FReset& AsReset() const;
-		const FLoop& AsLoop() const;
-		const FSeekTo& AsSeekTo() const;
-		const FSeekThru& AsSeekThru() const;
-		const FAdvanceThru& AsAdvanceThru() const;
-
-		// Asserts that this is a Reset, SeekTo, SeekThru, or AdvanceThru type
-		int32 FromTick() const;
-
-		// Asserts that this is a Reset, or SeekTo type
-		int32 ToTick() const;
-
-		// Asserts that this is a SeekThru or AdvanceThru type
-		int32 ThruTick() const;
+		struct FTempoChange
+		{
+			int32 Tick;
+			float Tempo;
+		};
 
-		const EType Type;
-	private:
-		union 
+		struct FTimeSignatureChange
 		{
-			const FReset Reset;
-			const FLoop Loop;
-			const FSeekTo SeekTo;
-			const FSeekThru SeekThru;
-			const FAdvanceThru AdvanceThru;
+			int32 Tick;
+			FTimeSignature TimeSignature;
 		};
-		
-	};
+	}
+
+	using FMidiClockMsg = TVariant<
+		MidiClockMessageTypes::FReset,
+		MidiClockMessageTypes::FLoop,
+		MidiClockMessageTypes::FSeekTo,
+		MidiClockMessageTypes::FSeekThru,
+		MidiClockMessageTypes::FAdvanceThru,
+		MidiClockMessageTypes::FTempoChange,
+		MidiClockMessageTypes::FTimeSignatureChange
+	>;
 	
 	struct HARMONIXMETASOUND_API FMidiClockEvent
 	{
 		const int32 BlockFrameIndex;
 		const FMidiClockMsg Msg;
-		
-		FMidiClockEvent(int32 InBlockFrameIndex, const FMidiClockMsg& Msg)
-			: BlockFrameIndex(InBlockFrameIndex), Msg(Msg)
-		{}
-		
+
+		template<typename T>
+		FMidiClockEvent(const int32 InBlockFrameIndex, T&& Msg)
+			: BlockFrameIndex(InBlockFrameIndex)
+			, Msg(TInPlaceType<T>(), MoveTemp(Msg))
+		{
+		}
 	};
 
 };
diff --git a/Engine/Plugins/Runtime/Harmonix/Source/HarmonixMetasoundTests/Private/DataTypes/MidiClock.spec.cpp b/Engine/Plugins/Runtime/Harmonix/Source/HarmonixMetasoundTests/Private/DataTypes/MidiClock.spec.cpp
--- a/Engine/Plugins/Runtime/Harmonix/Source/HarmonixMetasoundTests/Private/DataTypes/MidiClock.spec.cpp
+++ b/Engine/Plugins/Runtime/Harmonix/Source/HarmonixMetasoundTests/Private/DataTypes/MidiClock.spec.cpp
@@ -16,6 +16,7 @@ namespace HarmonixMetasoundTests::MidiClock
 
 	TUniquePtr<FMidiClock> TestClock;
 	TUniquePtr<FMidiClock> DrivingClock;
+	TSharedPtr<FMidiFileData> MidiFileData;
 	Metasound::FOperatorSettings OperatorSettings {48000, 100};
 
 	void AddStateAtFrame(EMusicPlayerTransportState State, int32 Frame) const
@@ -60,10 +61,10 @@ namespace HarmonixMetasoundTests::MidiClock
 			int32 Tick2 = (int32)(TestClock->GetTempoMap().MsToTick(Ms) + 0.5f);
 
 			TestEqual(FString::Printf(TEXT("Frame-%d: Event.BlockFrameIndex"), Event.BlockFrameIndex), Event.BlockFrameIndex, BlockFrameIndex);
-			TestEqual(FString::Printf(TEXT("Frame-%d: Event.Type"), Event.BlockFrameIndex), Event.Msg.Type, FMidiClockMsg::EType::AdvanceThru);
-			TestEqual(FString::Printf(TEXT("Frame-%d: Event.FromTick"), Event.BlockFrameIndex), Event.Msg.FromTick(), Tick1);
-			TestEqual(FString::Printf(TEXT("Frame-%d: Event.ThruTick"), Event.BlockFrameIndex), Event.Msg.ThruTick(), Tick2);
-			TestEqual(FString::Printf(TEXT("Frame-%d: Event.IsPreRoll"), Event.BlockFrameIndex), Event.Msg.AsAdvanceThru().IsPreRoll, false);
+			TestTrue(FString::Printf(TEXT("Frame-%d: Event.Type"), Event.BlockFrameIndex), Event.Msg.IsType<MidiClockMessageTypes::FAdvanceThru>());
+			TestEqual(FString::Printf(TEXT("Frame-%d: Event.FromTick"), Event.BlockFrameIndex), Event.Msg.Get<MidiClockMessageTypes::FAdvanceThru>().FromTick, Tick1);
+			TestEqual(FString::Printf(TEXT("Frame-%d: Event.ThruTick"), Event.BlockFrameIndex), Event.Msg.Get<MidiClockMessageTypes::FAdvanceThru>().ThruTick, Tick2);
+			TestEqual(FString::Printf(TEXT("Frame-%d: Event.IsPreRoll"), Event.BlockFrameIndex), Event.Msg.Get<MidiClockMessageTypes::FAdvanceThru>().IsPreRoll, false);
 			
 			BlockFrameIndex += FMidiClock::kMidiGranularity;
 		}
@@ -313,7 +314,7 @@ namespace HarmonixMetasoundTests::MidiClock
 
 				float Speed = 1.0f;
 				int32 PrerollBars = 8;
-				FMidiClockEvent Event = FMidiClockEvent(StartFrame, FMidiClockMsg::FAdvanceThru(0, Tick, false));
+				FMidiClockEvent Event = FMidiClockEvent(StartFrame, MidiClockMessageTypes::FAdvanceThru(0, Tick, false));
 
 				TestClock->PrepareBlock();
 				int32 OldEventsNum = TestClock->GetMidiClockEventsInBlock().Num();
@@ -323,8 +324,8 @@ namespace HarmonixMetasoundTests::MidiClock
 				{
 					return;
 				}
-				TestEqual("Last Clock Event in block", TestClock->GetMidiClockEventsInBlock().Last().Msg.Type, FMidiClockMsg::EType::AdvanceThru);
-				TestEqual("Clock Current Tick", TestClock->GetCurrentMidiTick(), Event.Msg.ThruTick() - 1);
+				TestTrue("Last Clock Event in block", TestClock->GetMidiClockEventsInBlock().Last().Msg.IsType<MidiClockMessageTypes::FAdvanceThru>());
+				TestEqual("Clock Current Tick", TestClock->GetCurrentMidiTick(), Event.Msg.Get<MidiClockMessageTypes::FAdvanceThru>().ThruTick - 1);
 			});
 
 			It("EventType::AdvanceThru.Looping", [&, this]
@@ -353,7 +354,7 @@ namespace HarmonixMetasoundTests::MidiClock
 
 				float Speed = 1.0f;
 				int32 PrerollBars = 8;
-				FMidiClockEvent Event = FMidiClockEvent(StartFrame, FMidiClockMsg::FAdvanceThru(0, Tick, false));
+				FMidiClockEvent Event = FMidiClockEvent(StartFrame, MidiClockMessageTypes::FAdvanceThru(0, Tick, false));
 
 				TestClock->PrepareBlock();
 				int32 OldEventsNum = TestClock->GetMidiClockEventsInBlock().Num();
@@ -363,7 +364,7 @@ namespace HarmonixMetasoundTests::MidiClock
 				{
 					return;
 				}
-				TestEqual("Last Clock Event in block", TestClock->GetMidiClockEventsInBlock().Last().Msg.Type, FMidiClockMsg::EType::AdvanceThru);
+				TestTrue("Last Clock Event in block", TestClock->GetMidiClockEventsInBlock().Last().Msg.IsType<MidiClockMessageTypes::FAdvanceThru>());
 				TestEqual("Clock Current Tick", TestClock->GetCurrentMidiTick(), ExpectedTick - 1);
 			});
 
@@ -373,7 +374,7 @@ namespace HarmonixMetasoundTests::MidiClock
 				int32 Tick = 1000;
 				float Speed = 1.0f;
 				int32 PrerollBars = 8;
-				FMidiClockEvent Event = FMidiClockEvent(StartFrame, FMidiClockMsg::FSeekTo(0, Tick));
+				FMidiClockEvent Event = FMidiClockEvent(StartFrame, MidiClockMessageTypes::FSeekTo(0, Tick));
 
 				TestClock->PrepareBlock();
 				int32 OldEventsNum = TestClock->GetMidiClockEventsInBlock().Num();
@@ -383,12 +384,99 @@ namespace HarmonixMetasoundTests::MidiClock
 				{
 					return;
 				}
-				TestEqual("Last Clock Event in block", TestClock->GetMidiClockEventsInBlock()[0].Msg.Type, FMidiClockMsg::EType::SeekThru);
+				TestTrue("Last Clock Event in block", TestClock->GetMidiClockEventsInBlock()[0].Msg.IsType<MidiClockMessageTypes::FSeekThru>());
 				TestEqual("Clock Current Tick", TestClock->GetCurrentMidiTick(), Tick - 1);
 			});
 			
 		});
 
+		Describe("Tempo Changes", [this]()
+		{
+			It("Without driving clock", [this]()
+			{
+				MidiFileData = FMidiClock::MakeClockConductorMidiData(123, 5, 8);
+				constexpr int32 TempoChangeTick = 234;
+				constexpr float TempoChangeTempo = 89;
+				MidiFileData->AddTempoChange(0, TempoChangeTick, TempoChangeTempo);
+				TestClock->AttachToMidiResource(MidiFileData);
+				TestClock->ResetAndStart(0);
+				
+				while (TestClock->GetCurrentMidiTick() < TempoChangeTick)
+				{
+					TestClock->PrepareBlock();
+					TestClock->WriteAdvance(0, OperatorSettings.GetNumFramesPerBlock());
+				}
+
+				bool HasTempoEvent = false;
+
+				for (const FMidiClockEvent& Event : TestClock->GetMidiClockEventsInBlock())
+				{
+					if (Event.Msg.IsType<MidiClockMessageTypes::FTempoChange>())
+					{
+						const MidiClockMessageTypes::FTempoChange& TempoChange = Event.Msg.Get<MidiClockMessageTypes::FTempoChange>();
+
+						if (!TestEqual("Tempo is correct", TempoChange.Tempo, TempoChangeTempo))
+						{
+							return;
+						}
+
+						if (!TestEqual("Tick is correct", TempoChange.Tick, TempoChangeTick))
+						{
+							return;
+						}
+						
+						HasTempoEvent = true;
+						break;
+					}
+				}
+
+				TestTrue("Got tempo event", HasTempoEvent);
+			});
+		});
+
+		Describe("Time signature Changes", [this]()
+		{
+			It("Without driving clock", [this]()
+			{
+				MidiFileData = FMidiClock::MakeClockConductorMidiData(123, 5, 8);
+				const int32 TimeSigChangeTick = MidiFileData->SongMaps.GetBarMap().BarBeatTickIncludingCountInToTick(2, 1, 0);
+				const FTimeSignature NewTimeSig{ 3, 4 };
+				MidiFileData->AddTimeSigChange(0, TimeSigChangeTick, NewTimeSig.Numerator, NewTimeSig.Denominator);
+				TestClock->AttachToMidiResource(MidiFileData);
+				TestClock->ResetAndStart(0);
+						
+				while (TestClock->GetCurrentMidiTick() < TimeSigChangeTick)
+				{
+					TestClock->PrepareBlock();
+					TestClock->WriteAdvance(0, OperatorSettings.GetNumFramesPerBlock());
+				}
+
+				bool HasTimeSignatureEvent = false;
+
+				for (const FMidiClockEvent& Event : TestClock->GetMidiClockEventsInBlock())
+				{
+					if (Event.Msg.IsType<MidiClockMessageTypes::FTimeSignatureChange>())
+					{
+						const MidiClockMessageTypes::FTimeSignatureChange& TimeSigChange = Event.Msg.Get<MidiClockMessageTypes::FTimeSignatureChange>();
+
+						if (!TestEqual("Time signature is correct", TimeSigChange.TimeSignature, NewTimeSig))
+						{
+							return;
+						}
+
+						if (!TestEqual("Tick is correct", TimeSigChange.Tick, TimeSigChangeTick))
+						{
+							return;
+						}
+								
+						HasTimeSignatureEvent = true;
+						break;
+					}
+				}
+
+				TestTrue("Got time sig event", HasTimeSignatureEvent);
+			});
+		});
 	}
 
 
