MidiStreamWriterNode

Only process MidiStreamWriter if the MidiStream has a clock.

#tests Metasound Tests

diff --git a/Engine/Plugins/Runtime/Harmonix/Source/HarmonixMetasound/Private/HarmonixMetasound/Nodes/MidiStreamWriterNode.cpp b/Engine/Plugins/Runtime/Harmonix/Source/HarmonixMetasound/Private/HarmonixMetasound/Nodes/MidiStreamWriterNode.cpp
--- a/Engine/Plugins/Runtime/Harmonix/Source/HarmonixMetasound/Private/HarmonixMetasound/Nodes/MidiStreamWriterNode.cpp
+++ b/Engine/Plugins/Runtime/Harmonix/Source/HarmonixMetasound/Private/HarmonixMetasound/Nodes/MidiStreamWriterNode.cpp
@@ -123,7 +123,7 @@ namespace HarmonixMetasound::Nodes::MidiStreamWriter
 		void Execute()
 		{
 			SetEnabled(*EnableInPin);
-			if (bEnabled && MidiStreamWriter.IsValid())
+			if (bEnabled && MidiStreamWriter.IsValid() && MidiStreamInPin->GetClock())
 			{
 				MidiStreamWriter->Process(*MidiStreamInPin);
 			}
