A quick and dirty tool for comparing two wave files using PSNR.

diff --git a/Engine/Plugins/Runtime/Harmonix/Source/HarmonixDsp/Private/HarmonixDsp/AudioAnalysis/WaveFileComparison.cpp b/Engine/Plugins/Runtime/Harmonix/Source/HarmonixDsp/Private/HarmonixDsp/AudioAnalysis/WaveFileComparison.cpp
new file mode 100644
--- /dev/null
+++ b/Engine/Plugins/Runtime/Harmonix/Source/HarmonixDsp/Private/HarmonixDsp/AudioAnalysis/WaveFileComparison.cpp
@@ -0,0 +1,112 @@
+// Copyright Epic Games, Inc. All Rights Reserved.
+
+#include "HarmonixDsp/AudioAnalysis/WaveFileComparison.h"
+
+#include "DSP/EnvelopeFollower.h"
+#include "GenericPlatform/GenericPlatformFile.h"
+#include "HAL/FileManagerGeneric.h"
+#include "HAL/PlatformFileManager.h"
+#include "HarmonixDsp/AudioAnalysis/AnalysisUtilities.h"
+
+DEFINE_LOG_CATEGORY_STATIC(LogWaveFileComparison, Log, Log);
+
+namespace Harmonix::Dsp::AudioAnalysis
+{
+	bool FWaveFileComparison::LoadForCompare(const FString& Wave1FilePath, const FString& Wave2FilePath)
+	{
+		IPlatformFile& PlatformFileApi = FPlatformFileManager::Get().GetPlatformFile();
+		IFileHandle* FileHandle = PlatformFileApi.OpenRead(*Wave1FilePath);
+		TSharedRef<FArchive> Wave1Archive = MakeShared<FArchiveFileReaderGeneric>(FileHandle, *Wave1FilePath, FileHandle->Size());
+
+		FileHandle = PlatformFileApi.OpenRead(*Wave2FilePath);
+		TSharedRef<FArchive> Wave2Archive = MakeShared<FArchiveFileReaderGeneric>(FileHandle, *Wave2FilePath, FileHandle->Size());
+		
+		return LoadForCompare(*Wave1Archive, *Wave2Archive);
+	}
+
+	bool FWaveFileComparison::LoadForCompare(FArchive& Wave1Archive, FArchive& Wave2Archive)
+	{
+		bOk = false;
+		
+		if (!Wave1.Load(Wave1Archive))
+		{
+			UE_LOG(LogWaveFileComparison, Log, TEXT("Failed to load Wave 1."));
+			return false;
+		}
+		
+		if (!Wave2.Load(Wave2Archive))
+		{
+			UE_LOG(LogWaveFileComparison, Log, TEXT("Failed to load Wave 2."));
+			return false;
+		}
+
+		if (*Wave1.Info.pChannels != *Wave2.Info.pChannels)
+		{
+			UE_LOG(LogWaveFileComparison, Log, TEXT("Can't compare... different channel counts."));
+			return false;
+		}
+
+		if (*Wave1.Info.pFormatTag != *Wave2.Info.pFormatTag)
+		{
+			UE_LOG(LogWaveFileComparison, Log, TEXT("Can't compare... different sample formats."));
+			return false;
+		}
+
+		if (*Wave1.Info.pBitsPerSample!= *Wave2.Info.pBitsPerSample)
+		{
+			UE_LOG(LogWaveFileComparison, Log, TEXT("Can't compare... different bits per sample."));
+			return false;
+		}
+
+		if (*Wave1.Info.pFormatTag != FWaveModInfo::WAVE_INFO_FORMAT_PCM && *Wave1.Info.pFormatTag != FWaveModInfo::WAVE_INFO_FORMAT_IEEE_FLOAT)
+		{
+			UE_LOG(LogWaveFileComparison, Log, TEXT("Can't compare... samples are not shorts or floats."));
+			return false;
+		}
+
+		if (*Wave1.Info.pFormatTag == FWaveModInfo::WAVE_INFO_FORMAT_PCM && *Wave1.Info.pBitsPerSample != 16)
+		{
+			UE_LOG(LogWaveFileComparison, Log, TEXT("Can't compare... pcm samples are not shorts."));
+			return false;
+		}
+
+		bOk = true;
+		return true;
+	}
+
+	float FWaveFileComparison::GetPSNR(bool bCommonSizeOnly) const
+	{
+		if (!bOk)
+		{
+			return std::numeric_limits<float>::max();
+		}
+
+		if (*Wave1.Info.pFormatTag == FWaveModInfo::WAVE_INFO_FORMAT_PCM)
+		{
+			return Harmonix::Dsp::AudioAnalysis::CalculatePSNR<int16>(reinterpret_cast<const int16*>(Wave1.Info.SampleDataStart),
+				reinterpret_cast<const int16*>(Wave2.Info.SampleDataStart),
+				*Wave1.Info.pChannels,
+				FMath::Min(Wave1.Info.GetNumSamples(), Wave2.Info.GetNumSamples()) / *Wave1.Info.pChannels);
+		}
+		return Harmonix::Dsp::AudioAnalysis::CalculatePSNR<float>(reinterpret_cast<const float*>(Wave1.Info.SampleDataStart),
+			reinterpret_cast<const float*>(Wave2.Info.SampleDataStart),
+			*Wave1.Info.pChannels,
+			FMath::Min(Wave1.Info.GetNumSamples(), Wave2.Info.GetNumSamples()) / *Wave1.Info.pChannels);
+	}
+
+	bool FWaveFileComparison::FOneWaveFile::Load(FArchive& Archive)
+	{
+		if (ensureAlways(Archive.IsLoading()))
+		{
+			BulkData.SetNumUninitialized(Archive.TotalSize());
+			Archive.Serialize(BulkData.GetData(), BulkData.Num());
+			if (!Info.ReadWaveInfo(BulkData.GetData(), BulkData.Num()))
+			{
+				return false;
+			}
+			return true;
+		}
+		return false;
+	}
+
+}
\ No newline at end of file
diff --git a/Engine/Plugins/Runtime/Harmonix/Source/HarmonixDsp/Public/HarmonixDsp/AudioAnalysis/AnalysisUtilities.h b/Engine/Plugins/Runtime/Harmonix/Source/HarmonixDsp/Public/HarmonixDsp/AudioAnalysis/AnalysisUtilities.h
--- a/Engine/Plugins/Runtime/Harmonix/Source/HarmonixDsp/Public/HarmonixDsp/AudioAnalysis/AnalysisUtilities.h
+++ b/Engine/Plugins/Runtime/Harmonix/Source/HarmonixDsp/Public/HarmonixDsp/AudioAnalysis/AnalysisUtilities.h
@@ -19,4 +19,51 @@ namespace Harmonix::Dsp::AudioAnalysis
 		float PreviousEnergyLinear,
 		float ElapsedMs,
 		const FHarmonixAudioAnalyzerOutputSettings& Settings);
+
+	/**
+	@brief Calculates the signal to noise. 96dB is perfect... no noise, all signal.	Big numbers good, small numbers bad
+	@param InterleavedInA - Pointer to the first sample in the first buffer. If multichannel, the channels must be interleaved
+	@param InterleavedInB - Pointer to the first sample in the second buffer. If multichannel, the channels must be interleaved
+	@param NumnChannels - Number of interleaved channels per sample frame
+	@param NumSampleFrames - Number of samples frames
+	@return Classic PSNR in dB. 96dB is perfect match. Smaller values indicate bigger difference.
+	*/
+	template <typename T>
+	float CalculatePSNR(const T* InterleavedInA, const T* InterleavedInB, const int32 NumChannels, const int32 NumSampleFrames)
+	{
+		// make sure if this is an integral type it is an int16...
+		static_assert(!std::is_integral<T>::value || sizeof(T) == 2);
+
+		double WorstSNR = 96.0f;
+		for (int32 ChannelIndex = 0; ChannelIndex < NumChannels; ++ChannelIndex)
+		{
+			double ChannelSumSquares = 0.0f;
+			const T* SampleSrc1 = InterleavedInA + ChannelIndex;
+			const T* SampleSrc2 = InterleavedInB + ChannelIndex;
+			for (int32 SampleIndex = 0; SampleIndex < NumSampleFrames; ++SampleIndex)
+			{
+				if constexpr (std::is_integral<T>::value)
+				{
+					double S1 = static_cast<double>(*SampleSrc1) / 32768.0;
+					double S2 = static_cast<double>(*SampleSrc2) / 32768.0;
+					ChannelSumSquares += (double)((S1 - S2) * (S1 - S2));
+				}
+				else
+				{
+					double S1 = static_cast<double>(*SampleSrc1);
+					double S2 = static_cast<double>(*SampleSrc2);
+					ChannelSumSquares += (double)((S1 - S2) * (S1 - S2));
+				}
+				SampleSrc1 += NumChannels;
+				SampleSrc2 += NumChannels;
+			}
+			ChannelSumSquares /= static_cast<double>(NumSampleFrames);
+			double SNR = ChannelSumSquares == 0.0 ? 96.0f : FMath::Clamp(10.0f * FMath::LogX(10.0, 4.0 / ChannelSumSquares), 0.0, 96.0);
+			if (SNR < WorstSNR)
+			{
+				WorstSNR = SNR;
+			}
+		}
+		return WorstSNR;
+	}
 }
diff --git a/Engine/Plugins/Runtime/Harmonix/Source/HarmonixDsp/Public/HarmonixDsp/AudioAnalysis/WaveFileComparison.h b/Engine/Plugins/Runtime/Harmonix/Source/HarmonixDsp/Public/HarmonixDsp/AudioAnalysis/WaveFileComparison.h
new file mode 100644
--- /dev/null
+++ b/Engine/Plugins/Runtime/Harmonix/Source/HarmonixDsp/Public/HarmonixDsp/AudioAnalysis/WaveFileComparison.h
@@ -0,0 +1,33 @@
+// Copyright Epic Games, Inc. All Rights Reserved.
+
+#pragma once
+#include "Audio.h"
+
+class FString;
+class FArchive;
+
+namespace Harmonix::Dsp::AudioAnalysis 
+{
+	class HARMONIXDSP_API FWaveFileComparison
+	{
+	public:
+		bool LoadForCompare(const FString& Wave1FilePath, const FString& Wave2FilePath);
+		bool LoadForCompare(FArchive& Wave1Archive, FArchive& Wave2Archive);
+
+		float GetPSNR(bool bCommonSizeOnly = true) const;
+
+	private:
+		struct FOneWaveFile
+		{
+			FWaveModInfo Info;
+			TArray<uint8> BulkData;
+
+			bool Load(FArchive& Archive);
+		};
+
+		FOneWaveFile Wave1;
+		FOneWaveFile Wave2;
+
+		bool bOk = false;
+	};
+}
\ No newline at end of file
