Fix non-Unity build

diff --git a/Engine/Plugins/Runtime/Harmonix/Source/HarmonixDsp/Public/HarmonixDsp/AudioAnalysis/AnalysisUtilities.h b/Engine/Plugins/Runtime/Harmonix/Source/HarmonixDsp/Public/HarmonixDsp/AudioAnalysis/AnalysisUtilities.h
--- a/Engine/Plugins/Runtime/Harmonix/Source/HarmonixDsp/Public/HarmonixDsp/AudioAnalysis/AnalysisUtilities.h
+++ b/Engine/Plugins/Runtime/Harmonix/Source/HarmonixDsp/Public/HarmonixDsp/AudioAnalysis/AnalysisUtilities.h
@@ -2,6 +2,9 @@
 
 #pragma once
 
+#include "HAL/Platform.h"
+#include "Math/UnrealMathUtility.h"
+
 struct FHarmonixAudioAnalyzerOutputSettings;
 
 namespace Harmonix::Dsp::AudioAnalysis
