StreamingAudioRendererV2

Handle the WaveProxyReader failing to create a decoder

Add "checks" to assert when we need the decoder to be valid. These methods are conditionally caled on the decoder being valid. If the decoder fails, we just produce silence and log.

#tests PIE with local cvar to force failure

diff --git a/Engine/Plugins/Runtime/Harmonix/Source/HarmonixDsp/Private/HarmonixDsp/AudioData/StreamingAudioRendererV2.cpp b/Engine/Plugins/Runtime/Harmonix/Source/HarmonixDsp/Private/HarmonixDsp/AudioData/StreamingAudioRendererV2.cpp
--- a/Engine/Plugins/Runtime/Harmonix/Source/HarmonixDsp/Private/HarmonixDsp/AudioData/StreamingAudioRendererV2.cpp
+++ b/Engine/Plugins/Runtime/Harmonix/Source/HarmonixDsp/Private/HarmonixDsp/AudioData/StreamingAudioRendererV2.cpp
@@ -38,7 +38,7 @@ void FStreamingAudioRendererV2::SetAudioData(TSharedRef<FSoundWaveProxy> InSound
 
 	check(WaveProxyReader.IsValid());
 
-	int32 WaveProxyNumChannels = WaveProxyReader->GetNumChannels();
+	int32 WaveProxyNumChannels = SoundWaveProxy->GetNumChannels();
 
 	int32 DecodeBufferSize = WaveProxyNumChannels * DeinterleaveBlockSizeInFrames;
 	DecodeBuffer.Reset(DecodeBufferSize);
@@ -76,7 +76,10 @@ void FStreamingAudioRendererV2::MigrateToSampler(const FFusionSampler* InSampler
 
 void FStreamingAudioRendererV2::SetFrame(uint32 InFrameNum)
 {
-	SeekSourceAudioToFrame(InFrameNum);
+	if (!WaveProxyReader->HasFailed())
+	{
+		SeekSourceAudioToFrame(InFrameNum);
+	}
 }
 
 double FStreamingAudioRendererV2::Render(TAudioBuffer<float>& OutBuffer, double InPos, int32 InMaxFrame, double InResampleInc, double InPitchShift, double InSpeed, bool MaintainPitchWhenSpeedChanges, bool InShouldHonorLoopPoints, const FGainMatrix& InGain)
@@ -410,6 +413,8 @@ void FStreamingAudioRendererV2::RenderMultiChannelRoutedUnshifted(TAudioBuffer<f
 
 void FStreamingAudioRendererV2::SeekSourceAudioToFrame(uint32 FrameIdx)
 {
+	check(!WaveProxyReader->HasFailed());
+	
 	// no need to seek
 	uint32 SourceFrameIndex = GetSourceAudioFrameIndex();
 
@@ -458,6 +463,8 @@ void FStreamingAudioRendererV2::SeekSourceAudioToFrame(uint32 FrameIdx)
 
 void FStreamingAudioRendererV2::DecodeSourceAudio(Audio::TCircularAudioBuffer<float>& OutBuffer)
 {
+	check(!WaveProxyReader->HasFailed());
+	
 	const int32 NumSamplesToGenerate = DeinterleaveBlockSizeInFrames * WaveProxyReader->GetNumChannels();
 	check(NumSamplesToGenerate == DecodeBuffer.Num());
 	
@@ -481,6 +488,7 @@ void FStreamingAudioRendererV2::DecodeSourceAudio(Audio::TCircularAudioBuffer<fl
 
 uint32 FStreamingAudioRendererV2::GetSourceAudioFrameIndex()
 {
+	check(!WaveProxyReader->HasFailed());
 	int32 ReaderFrameIndex = WaveProxyReader->GetFrameIndex();
 	int32 ReaderFramesInWave = WaveProxyReader->GetNumFramesInWave();
 	if (ReaderFramesInWave < ReaderFrameIndex)
@@ -610,10 +618,9 @@ void FStreamingAudioRendererV2::GenerateSourceAudio(uint32 StartFrameIndex, Audi
 void FStreamingAudioRendererV2::GenerateSourceAudioInternal(uint32 StartFrameIndex, float* OutAudioData, uint32 NumSamples)
 {
 	int32 NumSamplesRequested = NumSamples;
-	uint32 StartSampleIndex = StartFrameIndex * NumDeinterleaveChannels;
-
 	int32 BufferIdx = 0;
-	while (NumSamplesRequested > 0)
+	
+	while (NumSamplesRequested > 0 && !WaveProxyReader->HasFailed())
 	{
 		if (GetSourceAudioFrameIndex() != StartFrameIndex)
 		{
