Fix incorrect parameters to AdvanceTrackerByDeltaMs.

diff --git a/Engine/Plugins/Runtime/Harmonix/Source/HarmonixMidi/Private/HarmonixMidi/MidiPlayCursorMgr.cpp b/Engine/Plugins/Runtime/Harmonix/Source/HarmonixMidi/Private/HarmonixMidi/MidiPlayCursorMgr.cpp
--- a/Engine/Plugins/Runtime/Harmonix/Source/HarmonixMidi/Private/HarmonixMidi/MidiPlayCursorMgr.cpp
+++ b/Engine/Plugins/Runtime/Harmonix/Source/HarmonixMidi/Private/HarmonixMidi/MidiPlayCursorMgr.cpp
@@ -698,7 +698,7 @@ void FMidiPlayCursorMgr::AdvanceHiResToMs(float Ms, bool Broadcast)
 void FMidiPlayCursorMgr::AdvanceHiResByDeltaMs(float Ms, bool Broadcast)
 {
 	FMidiPlayCursorListLock HiResCursorListLock(GetHiResTracker().CursorListCS);
-	HiResLoopedSinceLastLoResUpdate = AdvanceTrackerByDeltaMs(Ms, GetHiResTracker(), Broadcast) || HiResLoopedSinceLastLoResUpdate;
+	HiResLoopedSinceLastLoResUpdate = AdvanceTrackerByDeltaMs(Ms, GetHiResTracker(), false, Broadcast) || HiResLoopedSinceLastLoResUpdate;
 	MsSinceLowResUpdate += Ms;
 }
 
