Harmonix plugin HarmonixMetasoundFunctionalTest
#tests Harmonix Functional Tests

#virtualized

diff --git a/Engine/Plugins/Runtime/Harmonix/Source/HarmonixMetasoundTests/HarmonixMetasoundTests.Build.cs b/Engine/Plugins/Runtime/Harmonix/Source/HarmonixMetasoundTests/HarmonixMetasoundTests.Build.cs
--- a/Engine/Plugins/Runtime/Harmonix/Source/HarmonixMetasoundTests/HarmonixMetasoundTests.Build.cs
+++ b/Engine/Plugins/Runtime/Harmonix/Source/HarmonixMetasoundTests/HarmonixMetasoundTests.Build.cs
@@ -13,7 +13,10 @@ public class HarmonixMetasoundTests : ModuleRules
 				"Core",
 				"CoreUObject",
 				"Engine",
+				"FunctionalTesting",
 				"UnrealEd",
+				"AudioExtensions",
+				"MetasoundEngine",
 				"MetasoundEngineTest",
 				"MetasoundFrontend",
 				"MetasoundGenerator",
@@ -22,6 +25,7 @@ public class HarmonixMetasoundTests : ModuleRules
 				"HarmonixDsp",
 				"HarmonixMetasound",
 				"HarmonixMidi",
+				"Projects",
 				"SignalProcessing"
 			}
 		);
diff --git a/Engine/Plugins/Runtime/Harmonix/Source/HarmonixMetasoundTests/Private/HarmonixMetasoundFunctionalTest.cpp b/Engine/Plugins/Runtime/Harmonix/Source/HarmonixMetasoundTests/Private/HarmonixMetasoundFunctionalTest.cpp
new file mode 100644
--- /dev/null
+++ b/Engine/Plugins/Runtime/Harmonix/Source/HarmonixMetasoundTests/Private/HarmonixMetasoundFunctionalTest.cpp
@@ -0,0 +1,361 @@
+﻿// Copyright Epic Games, Inc. All Rights Reserved.
+
+#include "HarmonixMetasoundFunctionalTest.h"
+
+#include "MetasoundGenerator.h"
+#include "Components/AudioComponent.h"
+#include "HarmonixMetasound/DataTypes/MidiStream.h"
+#include "HarmonixDsp/AudioAnalysis/AnalysisUtilities.h"
+
+#include "Audio.h"
+#include "Audio/SimpleWaveWriter.h"
+#include "Analysis/MetasoundFrontendVertexAnalyzerAudioBuffer.h"
+#include "HAL/FileManager.h"
+#include "Interfaces/IPluginManager.h"
+#include "Kismet/KismetSystemLibrary.h"
+#include "Misc/FileHelper.h"
+
+
+DEFINE_LOG_CATEGORY(LogHarmonixMetasoundTests)
+
+namespace HarmonixMetasoundTests
+{
+	FString MetasoundOutputValueAsString(const FMetaSoundOutput& Output)
+	{
+		if (float Value; Output.Get(Value))
+		{
+			return FString::Printf(TEXT("%f"), Value);
+		}
+		if (int32 Value; Output.Get(Value))
+		{
+			return FString::Printf(TEXT("%d"), Value);
+		}
+		if (bool Value; Output.Get(Value))
+		{
+			return FString::Printf(TEXT("%s"), Value ? TEXT("true") : TEXT("false"));
+		}
+		if (FString Value; Output.Get(Value))
+		{
+			return Value;
+		}
+		if (Metasound::FTime Value; Output.Get(Value))
+		{
+			return FString::Printf(TEXT("%f"), (float)Value.GetSeconds());
+		}
+		return FString::Printf(TEXT("Unsupported logging type: %s"), *Output.GetDataTypeName().ToString());
+	}
+}
+
+bool UHarmonixMetasoundFunctionalTestLibrary::AddOutputLogger(UMetasoundGeneratorHandle* GeneratorHandle, FName OutputName, EAudioParameterType Type)
+{
+	UE_LOG(LogHarmonixMetasoundTests, Log, TEXT("Adding output logger for output: %s"), *OutputName.ToString());
+	TWeakObjectPtr<const UObject> WeakContext = GeneratorHandle;
+	bool ReturnValue = GeneratorHandle->WatchOutput(
+		OutputName, FOnMetasoundOutputValueChangedNative::CreateLambda(
+			[WeakContext](FName OutputName, const FMetaSoundOutput& Output) 
+			{
+				FString Message = FString::Printf(TEXT("%s: %s"), *OutputName.ToString(), *HarmonixMetasoundTests::MetasoundOutputValueAsString(Output));
+				UKismetSystemLibrary::PrintString(nullptr, Message);
+			}
+		)
+	);
+	if (!ReturnValue)
+	{
+		UE_LOG(LogHarmonixMetasoundTests, Log, TEXT("Failed to add logger for output: %s"), *OutputName.ToString());
+	}
+	return ReturnValue;
+}
+
+bool UHarmonixMetasoundFunctionalTestLibrary::AddMidiStreamLogger(UMetasoundGeneratorHandle* GeneratorHandle, FName OutputName)
+{
+	using namespace HarmonixMetasound;
+	return GeneratorHandle->WatchOutput(
+		OutputName, FOnMetasoundOutputValueChangedNative::CreateLambda(
+				[](FName OutputName, const FMetaSoundOutput& Output)
+				{
+					if (FMidiStream MidiStream; Output.Get<FMidiStream>(MidiStream))
+					{
+						for (const FMidiStreamEvent& Event : MidiStream.GetEventsInBlock())
+						{
+							FString Message = FString::Printf(TEXT("%s: Event: BlockSampleFrameIndex=%d, MidiTick=%d, IsNoteMessage=%s, IsNoteOn=%s, Std1=%d"),
+								*OutputName.ToString(),
+								Event.BlockSampleFrameIndex,
+								Event.CurrentMidiTick,
+								Event.MidiMessage.IsNoteMessage() ? TEXT("true") : TEXT("false"),
+								Event.MidiMessage.IsNoteOn() ? TEXT("true") : TEXT("false"),
+								Event.MidiMessage.IsStd() ? Event.MidiMessage.GetStdData1() : 0);
+							UKismetSystemLibrary::PrintString(nullptr, Message);
+						}
+					}
+				}
+			)
+		);
+}
+
+FString UHarmonixMetasoundFunctionalTestLibrary::WriteAudioToFile(const FString& Filename, int32 SampleRate, int32 NumChannels, const Audio::FAlignedFloatBuffer& Audio)
+{
+	// write output to file in the "Audio Capture Directory"
+	static const FString RootPath = FPaths::AudioCaptureDir();
+	FString OutFilename = RootPath / Filename;
+	TUniquePtr<FArchive> Stream{ IFileManager::Get().CreateFileWriter(*OutFilename, IO_WRITE) };
+	// create the Wave Writer to write output to file
+	const TUniquePtr<Audio::FSimpleWaveWriter> Writer = MakeUnique<Audio::FSimpleWaveWriter>(MoveTemp(Stream), SampleRate, NumChannels, true);
+	Writer->Write(MakeArrayView(Audio.GetData(), Audio.Num()));
+	return OutFilename;
+}
+
+bool UHarmonixMetasoundFunctionalTestLibrary::ReadAudioFromFile(const FString& Filepath, Audio::FAlignedFloatBuffer& OutAudio, int32& OutSampleRate, int32& OutNumFrames, int32& OutNumChannels, uint16& OutFormatTag)
+{
+	OutAudio.Reset();
+	OutSampleRate = -1;
+	OutNumChannels = -1;
+	OutFormatTag = 0;
+	
+	if (!ensure(FPaths::FileExists(Filepath)))
+	{
+		UE_LOG(LogHarmonixMetasoundTests, Error, TEXT("Failed to read wave file %s: File does not exist"), *Filepath);
+		return false;
+	}
+
+	TArray64<uint8> FileData;
+	if (!ensure(FFileHelper::LoadFileToArray(FileData, *Filepath)))
+	{
+		UE_LOG(LogHarmonixMetasoundTests, Error, TEXT("Failed to read wave file %s: Unable to convert data to array"), *Filepath)
+		return false;
+	}
+
+	FWaveModInfo WaveInfo;
+	FString ErrorMessage;
+
+	if (!ensure(WaveInfo.ReadWaveInfo(FileData.GetData(), FileData.Num(), &ErrorMessage)))
+	{
+		UE_LOG(LogHarmonixMetasoundTests, Error, TEXT("Failed to read wave file %s: %s"), *Filepath, *ErrorMessage);
+		return false;
+	}
+	
+	// if (!ensure(*WaveInfo.pFormatTag == FWaveModInfo::WAVE_INFO_FORMAT_PCM || *WaveInfo.pFormatTag == FWaveModInfo::WAVE_INFO_FORMAT_IEEE_FLOAT))
+	// {
+	// 	UE_LOG(LogHarmonixMetasoundTests, Error, TEXT("Failed to read wave file %s: Unable to read format. Must be PCM short or Float"));
+	// 	return false;
+	// }
+
+	if (*WaveInfo.pFormatTag == FWaveModInfo::WAVE_INFO_FORMAT_PCM)
+	{
+		if (!ensure(*WaveInfo.pBitsPerSample == 16))
+		{
+			UE_LOG(LogHarmonixMetasoundTests, Error, TEXT("Failed to read wave file %s: Unable to read format. Must be 16 bit PCM"), *Filepath);
+			return false;
+		}
+		
+		OutFormatTag = *WaveInfo.pFormatTag;
+		OutSampleRate = *WaveInfo.pSamplesPerSec;
+		OutNumChannels= *WaveInfo.pChannels;
+		OutNumFrames = WaveInfo.SampleDataSize / OutNumChannels / (sizeof(int16));
+		OutAudio.AddUninitialized(OutNumFrames * OutNumChannels);
+		const int16* RawDataPtr = (const int16*)WaveInfo.SampleDataStart;
+		constexpr float Max16BitAsFloat = static_cast<float>(TNumericLimits<int16>::Max());
+		for (int32 Frame = 0; Frame < OutNumFrames; ++Frame)
+		{
+			for (int32 Channel = 0; Channel < OutNumChannels; ++Channel)
+			{
+				int32 SampleIdx = Frame * OutNumChannels + Channel;
+				OutAudio[SampleIdx] = RawDataPtr[SampleIdx] / Max16BitAsFloat;
+			}
+		}
+
+		return true;
+	}
+
+	if (*WaveInfo.pFormatTag == FWaveModInfo::WAVE_INFO_FORMAT_IEEE_FLOAT)
+	{
+		OutFormatTag = *WaveInfo.pFormatTag;
+		OutSampleRate = *WaveInfo.pSamplesPerSec;
+		OutNumChannels= *WaveInfo.pChannels;
+		OutNumFrames = WaveInfo.SampleDataSize / OutNumChannels / (sizeof(float));
+		OutAudio.AddUninitialized(OutNumFrames * OutNumChannels);
+		const float* RawDataPtr = (const float*)WaveInfo.SampleDataStart;
+		for (int32 Frame = 0; Frame < OutNumFrames; ++Frame)
+		{
+			for (int32 Channel = 0; Channel < OutNumChannels; ++Channel)
+			{
+				int32 SampleIdx = Frame * OutNumChannels + Channel;
+				OutAudio[SampleIdx] = RawDataPtr[SampleIdx];
+			}
+		}
+		return true;
+	}
+
+	UE_LOG(LogHarmonixMetasoundTests, Error, TEXT("Failed to read wave file %s: Unable to read format. Must be 16 bit PCM or IEEE float"), *Filepath);
+	return false;
+}
+
+
+AHarmonixMetasoundFunctionalTest::AHarmonixMetasoundFunctionalTest(const FObjectInitializer& ObjectInitializer)
+	:AFunctionalTest(ObjectInitializer)
+{
+
+	AudioComponent = CreateDefaultSubobject<UAudioComponent>(TEXT("AudioComponent"));
+	AudioComponent->SetupAttachment(RootComponent);
+	check(AudioComponent);
+}
+
+bool AHarmonixMetasoundFunctionalTest::IsReady_Implementation()
+{
+	bool OutIsReady = AFunctionalTest::IsReady_Implementation();
+
+	return OutIsReady && GeneratorHandle;
+}
+
+void AHarmonixMetasoundFunctionalTest::CompareResults()
+{
+	if (WavFilename_Expected.IsEmpty())
+	{
+		return;
+	}
+	
+	TSharedPtr<IPlugin> Plugin = IPluginManager::Get().FindPlugin(TEXT("Harmonix"));
+	check(Plugin);
+	FString TestAudioDir = Plugin->GetContentDir() / TEXT("Editor/Tests/Audio");
+	FString Filepath = TestAudioDir / WavFilename_Expected;
+	Audio::FAlignedFloatBuffer AudioData;
+	int32 SampleRate;
+	int32 NumChannels;
+	int32 NumFrames;
+	uint16 FormatTag;
+	bool Success = UHarmonixMetasoundFunctionalTestLibrary::ReadAudioFromFile(Filepath, AudioData, SampleRate, NumFrames, NumChannels, FormatTag);
+
+	if (!AssertTrue(Success, FString::Printf(TEXT("ReadAudioFromFile: %s"), *WavFilename_Expected), this))
+	{
+		return;
+	}
+
+	int32 NumChannels_Actual = 1;
+	if (!AssertEqual_Int(NumChannels_Actual, NumChannels, TEXT("NumChannels"), this))
+	{
+		return;
+	}
+
+	if (!AssertTrue(AudioCaptureSampleRate > 0.0f, TEXT("AudioCaptureSampleRate > 0"), this))
+	{
+		return;
+	}
+
+	float AudioCaptureDuration = AudioCaptureOutput.Num() / AudioCaptureSampleRate;
+	float AudioDuration = AudioData.Num() / (float)SampleRate;
+	if (!AssertEqual_Float(AudioCaptureDuration, AudioDuration, TEXT("AudioCaptureDuration"), 0.1f, this))
+	{
+		return;
+	}
+
+	int32 NumFramesToCompare = FMath::Min(AudioCaptureOutput.Num(), AudioData.Num());
+	float PSNR = Harmonix::Dsp::AudioAnalysis::CalculatePSNR(AudioCaptureOutput.GetData(), AudioData.GetData(), NumChannels, NumFramesToCompare);
+	AssertEqual_Float(PSNR, 96.0f, FString::Printf(TEXT("PSNR of Captured Audio and %s"), *WavFilename_Expected), 0.01f, this);
+}
+
+void AHarmonixMetasoundFunctionalTest::StartTest()
+{
+	AFunctionalTest::StartTest();
+
+	if (AudioComponent && AudioAutoStart)
+	{
+		AudioComponent->Play();
+	}
+}
+
+void AHarmonixMetasoundFunctionalTest::FinishTest(EFunctionalTestResult TestResult, const FString& Message)
+{
+	if (AudioComponent)
+	{
+		AudioComponent->Stop();
+	}
+
+	CompareResults();
+	
+	AFunctionalTest::FinishTest(TestResult, Message);
+}
+
+void AHarmonixMetasoundFunctionalTest::PrepareTest()
+{
+	AFunctionalTest::PrepareTest();
+	
+	if (!AudioComponent)
+	{
+		return;
+	}
+
+	if (!TestSound)
+	{
+		return;
+	}
+
+	AudioComponent->Sound = TestSound;
+	
+	GeneratorHandle = UMetasoundGeneratorHandle::CreateMetaSoundGeneratorHandle(AudioComponent);
+	
+	GeneratorHandle->OnGeneratorHandleAttached.AddLambda([this]()
+	{
+		if (!GeneratorHandle)
+		{
+			return;
+		}
+
+		TSharedPtr<Metasound::FMetasoundGenerator> Generator = GeneratorHandle->GetGenerator();
+		if (!Generator)
+		{
+			return;
+		}
+
+		AudioCaptureSampleRate = Generator->OperatorSettings.GetSampleRate();
+		
+		AudioCaptureOutput.Reset();
+		AudioOutAnalyzerAddress.DataType = Metasound::GetMetasoundDataTypeName<Metasound::FAudioBuffer>();
+		AudioOutAnalyzerAddress.InstanceID = 1234;
+		AudioOutAnalyzerAddress.OutputName = AudioOutName;
+		AudioOutAnalyzerAddress.AnalyzerName = Metasound::Frontend::FVertexAnalyzerAudioBuffer::GetAnalyzerName();
+		AudioOutAnalyzerAddress.AnalyzerInstanceID = FGuid::NewGuid();
+		AudioOutAnalyzerAddress.AnalyzerMemberName = Metasound::Frontend::FVertexAnalyzerAudioBuffer::FOutputs::GetValue().Name;
+			
+		Generator->AddOutputVertexAnalyzer(AudioOutAnalyzerAddress);
+
+		Generator->OnOutputChanged.AddLambda(
+			[this](const FName AnalyzerName, const FName OutputName, const FName AnalyzerOutputName, TSharedPtr<Metasound::IOutputStorage> OutputData)
+		{
+			if (AnalyzerName == AudioOutAnalyzerAddress.AnalyzerName
+				&& OutputName == AudioOutAnalyzerAddress.OutputName
+				&& AnalyzerOutputName == AudioOutAnalyzerAddress.AnalyzerMemberName)
+			{
+				const Metasound::FAudioBuffer& AudioBuffer = static_cast<Metasound::TOutputStorage<Metasound::FAudioBuffer>*>(OutputData.Get())->Get();
+				AudioCaptureOutput.Append(AudioBuffer.GetData(), AudioBuffer.Num());
+			}
+		});
+	
+	});
+	
+	OnTestFinished.AddDynamic(this, &AHarmonixMetasoundFunctionalTest::OnTestFinishedEvent);
+}
+
+void AHarmonixMetasoundFunctionalTest::OnTestFinishedEvent()
+{
+	OnTestFinished.RemoveDynamic(this, &AHarmonixMetasoundFunctionalTest::OnTestFinishedEvent);
+	
+	if (GeneratorHandle)
+	{
+		if (TSharedPtr<Metasound::FMetasoundGenerator> Generator = GeneratorHandle->GetGenerator())
+		{
+			
+#if WITH_EDITOR
+			// only write out the test results to a file when running in editor. 
+			if (!WavFilename_Output.IsEmpty() )
+			{
+				UHarmonixMetasoundFunctionalTestLibrary::WriteAudioToFile(WavFilename_Output, Generator->OperatorSettings.GetSampleRate(), 1, AudioCaptureOutput);
+			}
+#endif
+			
+			Generator->RemoveOutputVertexAnalyzer(AudioOutAnalyzerAddress);
+		}
+	}
+}
+
+
+
+
diff --git a/Engine/Plugins/Runtime/Harmonix/Source/HarmonixMetasoundTests/Private/HarmonixMetasoundFunctionalTest.h b/Engine/Plugins/Runtime/Harmonix/Source/HarmonixMetasoundTests/Private/HarmonixMetasoundFunctionalTest.h
new file mode 100644
--- /dev/null
+++ b/Engine/Plugins/Runtime/Harmonix/Source/HarmonixMetasoundTests/Private/HarmonixMetasoundFunctionalTest.h
@@ -0,0 +1,95 @@
+﻿// Copyright Epic Games, Inc. All Rights Reserved.
+
+#pragma once
+
+#include "FunctionalTest.h"
+
+#include "AudioParameter.h"
+#include "MetasoundGeneratorHandle.h"
+
+#include "HarmonixMetasoundFunctionalTest.generated.h"
+
+DECLARE_LOG_CATEGORY_EXTERN(LogHarmonixMetasoundTests, Log, All);
+
+namespace Audio
+{
+	class FSimpleWaveWriter;
+}
+
+UCLASS()
+class UHarmonixMetasoundFunctionalTestLibrary final : public UBlueprintFunctionLibrary
+{
+	GENERATED_BODY()
+public:
+
+	UFUNCTION(BlueprintCallable, Category="Harmonix Metasound", Meta=(ExpandBoolAsExecs="ReturnValue"))
+	static bool AddOutputLogger(UMetasoundGeneratorHandle* GeneratorHandle, FName OutputName, EAudioParameterType Type);
+
+	UFUNCTION(BlueprintCallable, Category="Harmonix Metasound", Meta=(ExpandBoolAsExecs="ReturnValue"))
+	static bool AddMidiStreamLogger(UMetasoundGeneratorHandle* GeneratorHandle, FName OutputName);
+	
+	static FString WriteAudioToFile(const FString& Filename, int32 SampleRate, int32 NumChannels, const Audio::FAlignedFloatBuffer& Audio);
+
+	// absolute filepath
+	// 
+	// reads audio data from audio file, converting audio to float
+	// audio is interleaved in the OutAudio buffer for multi-channel audio
+	// OutFormatTag is the format the audio file was saved in
+	static bool ReadAudioFromFile(const FString& Filepath, Audio::FAlignedFloatBuffer& OutAudio, int32& OutSampleRate, int32& OutNumFrames, int32& OutNumChannels, uint16& OutFormatTag);
+};
+
+UCLASS(Blueprintable)
+class AHarmonixMetasoundFunctionalTest : public AFunctionalTest
+{
+	GENERATED_BODY()
+public:
+
+	AHarmonixMetasoundFunctionalTest(const FObjectInitializer& ObjectInitializer);
+
+	UPROPERTY(EditAnywhere, Category = "Functional Testing")
+	TObjectPtr<USoundBase> TestSound;
+
+	UPROPERTY(EditAnywhere, Category = "Functional Testing")
+	FName AudioOutName = "AudioOut";
+
+	UPROPERTY(EditAnywhere, Category = "Functional Testing")
+	bool AudioAutoStart = true;
+
+	// name of wave file with the expected audio output.
+	// (searches the Content/Tests/AudioSource directory)
+	UPROPERTY(EditAnywhere, Category = "Functional Testing")
+	FString WavFilename_Expected;
+
+	// optionally write the output to a file
+	UPROPERTY(EditAnywhere, Category = "Functional Testing")
+	FString WavFilename_Output;
+
+	virtual void FinishTest(EFunctionalTestResult TestResult, const FString& Message) override;
+
+protected:
+	
+	virtual void PrepareTest() override;
+
+	virtual void StartTest() override;
+
+	virtual bool IsReady_Implementation() override;
+
+	void CompareResults();
+
+	UFUNCTION()
+	void OnTestFinishedEvent();
+
+	UPROPERTY()
+	UAudioComponent* AudioComponent;
+
+	UPROPERTY(Transient)
+	UMetasoundGeneratorHandle* GeneratorHandle;
+
+	// the audio captured from the metasound
+	Audio::FAlignedFloatBuffer AudioCaptureOutput;
+	float AudioCaptureSampleRate = 0.0f;
+	
+	Metasound::Frontend::FAnalyzerAddress AudioOutAnalyzerAddress;
+	
+	FTimerHandle TimerHandle;
+};
\ No newline at end of file
