Remove references to FMidiPlayCursor and instead use the MidiClock event stream.

diff --git a/Engine/Plugins/Runtime/Harmonix/Source/HarmonixMetasound/Private/HarmonixMetasound/MidiOps/PulseGenerator.cpp b/Engine/Plugins/Runtime/Harmonix/Source/HarmonixMetasound/Private/HarmonixMetasound/MidiOps/PulseGenerator.cpp
--- a/Engine/Plugins/Runtime/Harmonix/Source/HarmonixMetasound/Private/HarmonixMetasound/MidiOps/PulseGenerator.cpp
+++ b/Engine/Plugins/Runtime/Harmonix/Source/HarmonixMetasound/Private/HarmonixMetasound/MidiOps/PulseGenerator.cpp
@@ -9,161 +9,118 @@ namespace Harmonix::Midi::Ops
 		Enabled = bEnable;
 	}
 
-	void FPulseGenerator::SetClock(const TSharedPtr<const HarmonixMetasound::FMidiClock, ESPMode::NotThreadSafe>& Clock)
+	void FPulseGenerator::SetClock(const TSharedPtr<const HarmonixMetasound::FMidiClock, ESPMode::NotThreadSafe>& NewClock)
 	{
-		Cursor.SetClock(Clock);
+		Clock = NewClock;
 	}
 
 	void FPulseGenerator::SetInterval(const FMusicTimeInterval& NewInterval)
 	{
-		Cursor.SetInterval(NewInterval);
+		Interval = NewInterval;
+
+		// Multiplier should be >= 1
+		Interval.IntervalMultiplier = FMath::Max(Interval.IntervalMultiplier, static_cast<uint16>(1));
 	}
 
 	void FPulseGenerator::Process(HarmonixMetasound::FMidiStream& OutStream)
 	{
+		using namespace HarmonixMetasound;
+		using namespace HarmonixMetasound::MidiClockMessageTypes;
+
 		OutStream.PrepareBlock();
 		
-		// Keep draining the queue if disabled, so we get the next note off,
-		// and so we stay in phase if we toggle off and back on
-		FCursor::FPulseTime NextPulse;
-		while (Cursor.Pop(NextPulse))
+		const auto PinnedClock = Clock.Pin();
+
+		if (!PinnedClock)
 		{
-			int32 NoteOnSample = NextPulse.AudioBlockFrame;
+			return;
+		}
 
-			// Note off if there was a previous note on
-			if (LastNoteOn.IsSet())
+		for (const FMidiClockEvent& ClockEvent : PinnedClock->GetMidiClockEventsInBlock())
+		{
+			if (const FAdvanceThru* AsAdvanceThru = ClockEvent.TryGet<FAdvanceThru>())
 			{
-				check(LastNoteOn->MidiMessage.IsNoteOn());
-				
-				// Trigger the note off one sample before the note on
-				const int32 NoteOffSample = NextPulse.AudioBlockFrame > 0 ? NextPulse.AudioBlockFrame - 1 : NextPulse.AudioBlockFrame;
-				NoteOnSample = NoteOffSample + 1;
-
-				// Trigger the note off one tick before the note on
-				const int32 NoteOffTick = NextPulse.MidiTick - 1;
-
-				FMidiMsg Msg{ FMidiMsg::CreateNoteOff(LastNoteOn->MidiMessage.GetStdChannel(), LastNoteOn->MidiMessage.GetStdData1()) };
-				HarmonixMetasound::FMidiStreamEvent Event{ &VoiceGenerator, Msg };
-				Event.BlockSampleFrameIndex = NoteOffSample;
-				Event.AuthoredMidiTick = NoteOffTick;
-				Event.CurrentMidiTick = NoteOffTick;
-				Event.TrackIndex = LastNoteOn->TrackIndex;
-				OutStream.InsertMidiEvent(Event);
-
-				LastNoteOn.Reset();
+				HandleAdvanceThru(ClockEvent.BlockFrameIndex, AsAdvanceThru, OutStream);
 			}
-			
-			// Note on
-			if (Enabled)
+			else if (const FTimeSignatureChange* AsTimeSigChange = ClockEvent.TryGet<FTimeSignatureChange>())
 			{
-				FMidiMsg Msg{ FMidiMsg::CreateNoteOn(Channel - 1, NoteNumber, Velocity) };
-				HarmonixMetasound::FMidiStreamEvent Event{ &VoiceGenerator, Msg };
-				Event.BlockSampleFrameIndex = NoteOnSample;
-				Event.AuthoredMidiTick = NextPulse.MidiTick;
-				Event.CurrentMidiTick = NextPulse.MidiTick;
-				Event.TrackIndex = Track;
-				OutStream.InsertMidiEvent(Event);
-
-				LastNoteOn.Emplace(MoveTemp(Event));
+				HandleTimeSignatureChange(AsTimeSigChange);
 			}
 		}
 	}
 
-	FPulseGenerator::FCursor::FCursor()
-	{
-		SetMessageFilter(EFilterPassFlags::TimeSig);
-	}
-
-	FPulseGenerator::FCursor::~FCursor()
-	{
-		// NB: If we let ~FMidiPlayCursor do this, we get warnings for potentially bad access
-		SetClock(nullptr);
-	}
-
-	void FPulseGenerator::FCursor::SetClock(const TSharedPtr<const HarmonixMetasound::FMidiClock, ESPMode::NotThreadSafe>& NewClock)
+	void FPulseGenerator::HandleAdvanceThru(const int32 BlockFrameIndex, const HarmonixMetasound::MidiClockMessageTypes::FAdvanceThru* AdvanceThru, HarmonixMetasound::FMidiStream& OutStream)
 	{
-		if (Clock == NewClock)
+		if (!NextPulseTimestamp.IsValid())
 		{
 			return;
 		}
 
-		// unregister if we already have a clock
-		if (const auto PinnedClock = Clock.Pin())
-		{
-			PinnedClock->UnregisterPlayCursor(this);
-		}
+		check(Clock.IsValid()); // if we're here and we don't have a clock, we have problems
+		const auto PinnedClock = Clock.Pin();
+		int32 NextPulseTick = PinnedClock->GetBarMap().MusicTimestampToTick(NextPulseTimestamp);
 
-		// register the new clock if we were given one
-		if (NewClock.IsValid())
+		while (AdvanceThru->ThruTick >= NextPulseTick)
 		{
-			NewClock->RegisterHiResPlayCursor(this);
-		}
+			DoPulse(BlockFrameIndex, NextPulseTick, OutStream);
 
-		Clock = NewClock;
+			IncrementTimestampByInterval(NextPulseTimestamp, Interval, CurrentTimeSignature);
+
+			NextPulseTick = PinnedClock->GetBarMap().MusicTimestampToTick(NextPulseTimestamp);
+		}
 	}
 
-	void FPulseGenerator::FCursor::SetInterval(const FMusicTimeInterval& NewInterval)
+	void FPulseGenerator::HandleTimeSignatureChange(const HarmonixMetasound::MidiClockMessageTypes::FTimeSignatureChange* TimeSigChange)
 	{
-		Interval = NewInterval;
+		CurrentTimeSignature = TimeSigChange->TimeSignature;
 
-		// Multiplier should be >= 1
-		Interval.IntervalMultiplier = FMath::Max(Interval.IntervalMultiplier, static_cast<uint16>(1));
+		check(Clock.IsValid()); // if we're here and we don't have a clock, we have problems
+		const auto PinnedClock = Clock.Pin();
+		// Time sig changes will come on the downbeat, and if we change time signature,
+		// we want to reset the pulse, so the next pulse is now plus the offset
+		NextPulseTimestamp = PinnedClock->GetBarMap().TickToMusicTimestamp(TimeSigChange->Tick);
+		IncrementTimestampByOffset(NextPulseTimestamp, Interval, CurrentTimeSignature);
 	}
 
-	bool FPulseGenerator::FCursor::Pop(FPulseTime& PulseTime)
+	void FPulseGenerator::DoPulse(const int32 BlockFrameIndex, const int32 PulseTick, HarmonixMetasound::FMidiStream& OutStream)
 	{
-		if (PulsesSinceLastProcess.Dequeue(PulseTime))
+		int32 NoteOnSample = BlockFrameIndex;
+
+		// Note off if there was a previous note on
+		if (LastNoteOn.IsSet())
 		{
-			--QueueSize;
-			return true;
-		}
+			check(LastNoteOn->MidiMessage.IsNoteOn());
 
-		return false;
-	}
+			// Trigger the note off one sample before the note on
+			const int32 NoteOffSample = BlockFrameIndex > 0 ? BlockFrameIndex - 1 : BlockFrameIndex;
+			NoteOnSample = NoteOffSample + 1;
 
-	void FPulseGenerator::FCursor::Push(FPulseTime&& PulseTime)
-	{
-		constexpr int32 MaxQueueSize = 1024; // just to keep the queue from infinitely growing if we stop draining it
-		if (QueueSize < MaxQueueSize)
-		{
-			PulsesSinceLastProcess.Enqueue(MoveTemp(PulseTime));
-			++QueueSize;
-		}
-	}
+			// Trigger the note off one tick before the note on
+			const int32 NoteOffTick = PulseTick - 1;
 
-	void FPulseGenerator::FCursor::AdvanceThruTick(int32 Tick, bool IsPreRoll)
-	{
-		FMidiPlayCursor::AdvanceThruTick(Tick, IsPreRoll);
+			FMidiMsg Msg{ FMidiMsg::CreateNoteOff(LastNoteOn->MidiMessage.GetStdChannel(), LastNoteOn->MidiMessage.GetStdData1()) };
+			HarmonixMetasound::FMidiStreamEvent Event{ &VoiceGenerator, Msg };
+			Event.BlockSampleFrameIndex = NoteOffSample;
+			Event.AuthoredMidiTick = NoteOffTick;
+			Event.CurrentMidiTick = NoteOffTick;
+			Event.TrackIndex = LastNoteOn->TrackIndex;
+			OutStream.InsertMidiEvent(Event);
 
-		if (!NextPulseTimestamp.IsValid())
-		{
-			return;
+			LastNoteOn.Reset();
 		}
-		
-		check(Clock.IsValid()); // if we're here and we don't have a clock, we have problems
-		const auto PinnedClock = Clock.Pin();
-		int32 NextPulseTick = PinnedClock->GetBarMap().MusicTimestampToTick(NextPulseTimestamp);
 
-		while (CurrentTick >= NextPulseTick)
+		// Note on
+		if (Enabled)
 		{
-			Push({ PinnedClock->GetCurrentBlockFrameIndex(), NextPulseTick });
-			
-			IncrementTimestampByInterval(NextPulseTimestamp, Interval, CurrentTimeSignature);
-
-			NextPulseTick = PinnedClock->GetBarMap().MusicTimestampToTick(NextPulseTimestamp);
+			FMidiMsg Msg{ FMidiMsg::CreateNoteOn(Channel - 1, NoteNumber, Velocity) };
+			HarmonixMetasound::FMidiStreamEvent Event{ &VoiceGenerator, Msg };
+			Event.BlockSampleFrameIndex = NoteOnSample;
+			Event.AuthoredMidiTick = PulseTick;
+			Event.CurrentMidiTick = PulseTick;
+			Event.TrackIndex = Track;
+			OutStream.InsertMidiEvent(Event);
+
+			LastNoteOn.Emplace(MoveTemp(Event));
 		}
 	}
-
-	void FPulseGenerator::FCursor::OnTimeSig(int32 TrackIndex, int32 Tick, int32 Numerator, int32 Denominator, bool IsPreroll)
-	{
-		CurrentTimeSignature.Numerator = Numerator;
-		CurrentTimeSignature.Denominator = Denominator;
-		
-		check(Clock.IsValid()); // if we're here and we don't have a clock, we have problems
-		const auto PinnedClock = Clock.Pin();
-		// Time sig changes will come on the downbeat, and if we change time signature,
-		// we want to reset the pulse, so the next pulse is now plus the offset
-		NextPulseTimestamp = PinnedClock->GetBarMap().TickToMusicTimestamp(Tick);
-		IncrementTimestampByOffset(NextPulseTimestamp, Interval, CurrentTimeSignature);
-	}
 }
diff --git a/Engine/Plugins/Runtime/Harmonix/Source/HarmonixMetasound/Public/HarmonixMetasound/DataTypes/MidiClockEvent.h b/Engine/Plugins/Runtime/Harmonix/Source/HarmonixMetasound/Public/HarmonixMetasound/DataTypes/MidiClockEvent.h
--- a/Engine/Plugins/Runtime/Harmonix/Source/HarmonixMetasound/Public/HarmonixMetasound/DataTypes/MidiClockEvent.h
+++ b/Engine/Plugins/Runtime/Harmonix/Source/HarmonixMetasound/Public/HarmonixMetasound/DataTypes/MidiClockEvent.h
@@ -111,6 +111,12 @@ namespace HarmonixMetasound
 			, Msg(TInPlaceType<T>(), MoveTemp(Msg))
 		{
 		}
+
+		template<typename T>
+		bool IsType() const	{ return Msg.IsType<T>(); }
+
+		template<typename T>
+		const T* TryGet() const { return Msg.TryGet<T>(); }
 	};
 
 };
diff --git a/Engine/Plugins/Runtime/Harmonix/Source/HarmonixMetasound/Public/HarmonixMetasound/MidiOps/PulseGenerator.h b/Engine/Plugins/Runtime/Harmonix/Source/HarmonixMetasound/Public/HarmonixMetasound/MidiOps/PulseGenerator.h
--- a/Engine/Plugins/Runtime/Harmonix/Source/HarmonixMetasound/Public/HarmonixMetasound/MidiOps/PulseGenerator.h
+++ b/Engine/Plugins/Runtime/Harmonix/Source/HarmonixMetasound/Public/HarmonixMetasound/MidiOps/PulseGenerator.h
@@ -31,51 +31,21 @@ namespace Harmonix::Midi::Ops
 		
 		void SetInterval(const FMusicTimeInterval& NewInterval);
 
-		FMusicTimeInterval GetInterval() const { return Cursor.GetInterval(); }
+		FMusicTimeInterval GetInterval() const { return Interval; }
 
 		void Process(HarmonixMetasound::FMidiStream& OutStream);
 
 	private:
+		void HandleAdvanceThru(const int32 BlockFrameIndex, const HarmonixMetasound::MidiClockMessageTypes::FAdvanceThru* AdvanceThru, HarmonixMetasound::FMidiStream& OutStream);
+		void HandleTimeSignatureChange(const HarmonixMetasound::MidiClockMessageTypes::FTimeSignatureChange* TimeSigChange);
+		void DoPulse(const int32 BlockFrameIndex, const int32 PulseTick, HarmonixMetasound::FMidiStream& OutStream);
+
 		FMidiVoiceGeneratorBase VoiceGenerator{};
 		TOptional<HarmonixMetasound::FMidiStreamEvent> LastNoteOn;
 		bool Enabled{ true };
-
-		class FCursor final : public FMidiPlayCursor
-		{
-		public:
-			FCursor();
-
-			virtual ~FCursor() override;
-
-			void SetClock(const TSharedPtr<const HarmonixMetasound::FMidiClock, ESPMode::NotThreadSafe>& NewClock);
-
-			void SetInterval(const FMusicTimeInterval& NewInterval);
-
-			FMusicTimeInterval GetInterval() const { return Interval; }
-
-			struct FPulseTime
-			{
-				int32 AudioBlockFrame;
-				int32 MidiTick;
-			};
-
-			bool Pop(FPulseTime& PulseTime);
-
-		private:
-			void Push(FPulseTime&& PulseTime);
-			
-			virtual void AdvanceThruTick(int32 Tick, bool IsPreRoll) override;
-			
-			virtual void OnTimeSig(int32 TrackIndex, int32 Tick, int32 Numerator, int32 Denominator, bool IsPreroll) override;
-
-			TWeakPtr<const HarmonixMetasound::FMidiClock, ESPMode::NotThreadSafe> Clock;
-			FMusicTimeInterval Interval{};
-			TSpscQueue<FPulseTime> PulsesSinceLastProcess;
-			int32 QueueSize{0};
-			FTimeSignature CurrentTimeSignature{};
-			FMusicTimestamp NextPulseTimestamp{ -1, -1 };
-		};
-		
-		FCursor Cursor;
+		TWeakPtr<const HarmonixMetasound::FMidiClock, ESPMode::NotThreadSafe> Clock;
+		FMusicTimeInterval Interval{};
+		FTimeSignature CurrentTimeSignature{};
+		FMusicTimestamp NextPulseTimestamp{ -1, -1 };
 	};
 }
