Be sure tempo change events happen before the midi clock advance for the tick that they appear on.

diff --git a/Engine/Plugins/Runtime/Harmonix/Source/HarmonixMetasound/Private/HarmonixMetasound/DataTypes/MidiClock.cpp b/Engine/Plugins/Runtime/Harmonix/Source/HarmonixMetasound/Private/HarmonixMetasound/DataTypes/MidiClock.cpp
--- a/Engine/Plugins/Runtime/Harmonix/Source/HarmonixMetasound/Private/HarmonixMetasound/DataTypes/MidiClock.cpp
+++ b/Engine/Plugins/Runtime/Harmonix/Source/HarmonixMetasound/Private/HarmonixMetasound/DataTypes/MidiClock.cpp
@@ -633,12 +633,18 @@ namespace HarmonixMetasound
 			MyMidiClock->TempoChangesInBlock.Add({BlockFrameIndex, 0.0f, Bpm});
 		}
 
+		// CurrentAdvanceStartTick is the tick ALREADY processed in a previous advance. So the first new tick 
+		// we are processing is CurrentAdvanceStartTick + 1.
+		// 
 		// This gets fired off as a result of AdvanceThruTick, and we want the event to go in between advances,
-		// so we add an AdvanceThru event up to this tick. But if this tempo event is on the first tick, we don't need to do that.
-		if (Tick > CurrentAdvanceStartTick)
+		// so if this tempo change is not on the first new tick of this advance we need to "cap off" the 
+		// currently in progress advance and start a new one. 
+		const int32 FirstNewTickThisAdvance = CurrentAdvanceStartTick + 1;
+		const int32 TickPriorToThisEvent = Tick - 1;
+		if (Tick > FirstNewTickThisAdvance)
 		{
-			AddEvent(FMidiClockEvent(MyMidiClock->GetCurrentBlockFrameIndex(), MidiClockMessageTypes::FAdvanceThru(CurrentAdvanceStartTick, Tick, IsPreroll)));
-			CurrentAdvanceStartTick = Tick;
+			AddEvent(FMidiClockEvent(MyMidiClock->GetCurrentBlockFrameIndex(), MidiClockMessageTypes::FAdvanceThru(CurrentAdvanceStartTick, TickPriorToThisEvent, IsPreroll)));
+			CurrentAdvanceStartTick = TickPriorToThisEvent;
 		}
 		
 		AddEvent(FMidiClockEvent(MyMidiClock->GetCurrentBlockFrameIndex(), MidiClockMessageTypes::FTempoChange(Tick, Bpm)));
@@ -646,12 +652,18 @@ namespace HarmonixMetasound
 
 	void FMidiClock::FMidiClockEventCursor::OnTimeSig(int32 TrackIndex, int32 Tick, int32 Numerator, int32 Denominator, bool IsPreroll)
 	{
+		// CurrentAdvanceStartTick is the tick ALREADY processed in a previous advance. So the first new tick 
+		// we are processing is CurrentAdvanceStartTick + 1.
+		// 
 		// This gets fired off as a result of AdvanceThruTick, and we want the event to go in between advances,
-		// so we add an AdvanceThru event up to this tick. But if this tempo event is on the first tick, we don't need to do that.
-		if (Tick > CurrentAdvanceStartTick)
+		// so if this time signature change is not on the first new tick of this advance we need to "cap off" the 
+		// currently in progress advance and start a new one. 
+		const int32 FirstNewTickThisAdvance = CurrentAdvanceStartTick + 1;
+		const int32 TickPriorToThisEvent = Tick - 1;
+		if (Tick > FirstNewTickThisAdvance)
 		{
-			AddEvent(FMidiClockEvent(MyMidiClock->GetCurrentBlockFrameIndex(), MidiClockMessageTypes::FAdvanceThru(CurrentAdvanceStartTick, Tick, IsPreroll)));
-			CurrentAdvanceStartTick = Tick;
+			AddEvent(FMidiClockEvent(MyMidiClock->GetCurrentBlockFrameIndex(), MidiClockMessageTypes::FAdvanceThru(CurrentAdvanceStartTick, TickPriorToThisEvent, IsPreroll)));
+			CurrentAdvanceStartTick = TickPriorToThisEvent;
 		}
 		
 		AddEvent(FMidiClockEvent(MyMidiClock->GetCurrentBlockFrameIndex(), MidiClockMessageTypes::FTimeSignatureChange(Tick, { Numerator, Denominator })));
