[Harmonix] Remove DistortionV1

diff --git a/Engine/Plugins/Runtime/Harmonix/Source/HarmonixDsp/Private/HarmonixDsp/Effects/DistortionV1.cpp b/Engine/Plugins/Runtime/Harmonix/Source/HarmonixDsp/Private/HarmonixDsp/Effects/DistortionV1.cpp
deleted file mode 100644
diff --git a/Engine/Plugins/Runtime/Harmonix/Source/HarmonixDsp/Public/HarmonixDsp/Effects/DistortionV1.h b/Engine/Plugins/Runtime/Harmonix/Source/HarmonixDsp/Public/HarmonixDsp/Effects/DistortionV1.h
deleted file mode 100644
diff --git a/Engine/Plugins/Runtime/Harmonix/Source/HarmonixMetasound/Private/HarmonixMetasound/Nodes/DistortionNode.cpp b/Engine/Plugins/Runtime/Harmonix/Source/HarmonixMetasound/Private/HarmonixMetasound/Nodes/DistortionNode.cpp
--- a/Engine/Plugins/Runtime/Harmonix/Source/HarmonixMetasound/Private/HarmonixMetasound/Nodes/DistortionNode.cpp
+++ b/Engine/Plugins/Runtime/Harmonix/Source/HarmonixMetasound/Private/HarmonixMetasound/Nodes/DistortionNode.cpp
@@ -16,7 +16,6 @@
 #include "HarmonixDsp/Effects/Settings/BiquadFilterSettings.h"
 #include "HarmonixDsp/Effects/BiquadFilter.h"
 #include "HarmonixDsp/Effects/Settings/DistortionSettings.h"
-#include "HarmonixDsp/Effects/DistortionV1.h"
 #include "HarmonixDsp/Effects/DistortionV2.h"
 #include "HarmonixMetasound/Common.h"
 
