[Harmonix] FFT Analyzer: Use FAlignedFloatBuffer instead of TAudioBuffer<float>

diff --git a/Engine/Plugins/Runtime/Harmonix/Source/HarmonixDsp/Private/HarmonixDsp/AudioAnalysis/FFTAnalyzer.cpp b/Engine/Plugins/Runtime/Harmonix/Source/HarmonixDsp/Private/HarmonixDsp/AudioAnalysis/FFTAnalyzer.cpp
--- a/Engine/Plugins/Runtime/Harmonix/Source/HarmonixDsp/Private/HarmonixDsp/AudioAnalysis/FFTAnalyzer.cpp
+++ b/Engine/Plugins/Runtime/Harmonix/Source/HarmonixDsp/Private/HarmonixDsp/AudioAnalysis/FFTAnalyzer.cpp
@@ -52,17 +52,10 @@ namespace Harmonix::Dsp::AudioAnalysis
 		State = FState{};
 	}
 
-	void FFFTAnalyzer::Process(const TAudioBuffer<float>& InBuffer, FHarmonixFFTAnalyzerResults& InOutResults)
+	void FFFTAnalyzer::Process(const Audio::FAlignedFloatBuffer& InBuffer, FHarmonixFFTAnalyzerResults& InOutResults)
 	{
 		TRACE_CPUPROFILER_EVENT_SCOPE(Harmonix::Dsp::AudioAnalysis::FFFTAnalyzer::Process);
 		
-		const int32 NumChannels = InBuffer.GetNumValidChannels();
-
-		if (NumChannels < 1)
-		{
-			return;
-		}
-		
 		// Make a copy of the settings so we don't hold up another thread
 		FHarmonixFFTAnalyzerSettings SettingsCopy;
 		{
@@ -92,18 +85,8 @@ namespace Harmonix::Dsp::AudioAnalysis
 
 		// Process however many windows we need to process
 		{
-			// copy the input buffer
-			{
-				State.InputBuffer.SetNumUninitialized(InBuffer.GetNumValidFrames());
-				TDynamicStridePtr<float> ChannelPtr = InBuffer.GetStridingChannelDataPointer(0);
-				for (int32 i = 0; i < InBuffer.GetNumValidFrames(); ++i)
-				{
-					State.InputBuffer[i] = ChannelPtr[i];
-				}
-			}
-
 			// make the sliding window
-			Audio::TAutoSlidingWindow<float> SlidingWindow(*State.SlidingBuffer, State.InputBuffer, State.WindowedBuffer);
+			Audio::TAutoSlidingWindow<float> SlidingWindow(*State.SlidingBuffer, InBuffer, State.WindowedBuffer);
 
 			// process each window
 			for (TArray<float>& Window : SlidingWindow)
diff --git a/Engine/Plugins/Runtime/Harmonix/Source/HarmonixDsp/Public/HarmonixDsp/AudioAnalysis/FFTAnalyzer.h b/Engine/Plugins/Runtime/Harmonix/Source/HarmonixDsp/Public/HarmonixDsp/AudioAnalysis/FFTAnalyzer.h
--- a/Engine/Plugins/Runtime/Harmonix/Source/HarmonixDsp/Public/HarmonixDsp/AudioAnalysis/FFTAnalyzer.h
+++ b/Engine/Plugins/Runtime/Harmonix/Source/HarmonixDsp/Public/HarmonixDsp/AudioAnalysis/FFTAnalyzer.h
@@ -4,7 +4,6 @@
 
 #include "DSP/AudioFFT.h"
 #include "DSP/SlidingWindow.h"
-#include "HarmonixDsp/AudioBuffer.h"
 #include "HarmonixDsp/AudioAnalysis/OutputSettings.h"
 
 #include "FFTAnalyzer.generated.h"
@@ -82,7 +81,7 @@ namespace Harmonix::Dsp::AudioAnalysis
 		 * @param InBuffer - The audio buffer to analyze
 		 * @param InOutResults - The results from the analysis
 		 */
-		void Process(const TAudioBuffer<float>& InBuffer, FHarmonixFFTAnalyzerResults& InOutResults);
+		void Process(const Audio::FAlignedFloatBuffer& InBuffer, FHarmonixFFTAnalyzerResults& InOutResults);
 		
 	private:
 		const float SampleRate;
@@ -95,7 +94,6 @@ namespace Harmonix::Dsp::AudioAnalysis
 		{
 			TUniquePtr<Audio::IFFTAlgorithm> FFT;
 			TUniquePtr<Audio::TSlidingBuffer<float>> SlidingBuffer;
-			TArray<float> InputBuffer;
 			TUniquePtr<Audio::FWindow> Window;
 			TArray<float> WindowedBuffer;
 			Audio::AlignedFloatBuffer FFTOutput;
diff --git a/Engine/Plugins/Runtime/Harmonix/Source/HarmonixDspTests/Private/AudioAnalysis/FFTAnalyzerTests.cpp b/Engine/Plugins/Runtime/Harmonix/Source/HarmonixDspTests/Private/AudioAnalysis/FFTAnalyzerTests.cpp
--- a/Engine/Plugins/Runtime/Harmonix/Source/HarmonixDspTests/Private/AudioAnalysis/FFTAnalyzerTests.cpp
+++ b/Engine/Plugins/Runtime/Harmonix/Source/HarmonixDspTests/Private/AudioAnalysis/FFTAnalyzerTests.cpp
@@ -1,5 +1,6 @@
 // Copyright Epic Games, Inc. All Rights Reserved.
 
+#include "HarmonixDsp/Generate.h"
 #include "HarmonixDsp/AudioAnalysis/FFTAnalyzer.h"
 #include "Misc/AutomationTest.h"
 
@@ -13,7 +14,6 @@ namespace Harmonix::Dsp::AudioAnalysis::FFTAnalyzer::Tests
 		EAutomationTestFlags::EditorContext | EAutomationTestFlags::EngineFilter)
 	bool FHarmonixFFTAnalyzerBasicTest::RunTest(const FString&)
 	{
-		constexpr int32 NumChannels = 2;
 		constexpr int32 NumFramesPerBlock = 256;
 		constexpr float SampleRate = 48000;
 		
@@ -25,18 +25,15 @@ namespace Harmonix::Dsp::AudioAnalysis::FFTAnalyzer::Tests
 		FHarmonixFFTAnalyzerResults Results;
 
 		// process a few times so we can ensure we have some energy in the spectrum
-		TAudioBuffer<float> Buffer{ NumChannels, NumFramesPerBlock, EAudioBufferCleanupMode::Delete };
+		Audio::FAlignedFloatBuffer Buffer;
+		Buffer.SetNumUninitialized(NumFramesPerBlock);
 		constexpr float Frequency = 440;
 		float Phase = 0;
 		constexpr int32 NumBlocks = 200;
 		for (int i = 0; i < NumBlocks; ++i)
 		{
 			// Fill the buffer with a sine
-			HarmonixDsp::GenerateSine(Buffer.GetValidChannelData(0), NumFramesPerBlock, Frequency, SampleRate, Phase);
-			for (int32 c = 1; c < NumChannels; ++c)
-			{
-				FMemory::Memcpy(Buffer.GetValidChannelData(c), Buffer.GetValidChannelData(0), NumFramesPerBlock * sizeof(float));
-			}
+			HarmonixDsp::GenerateSine(Buffer.GetData(), NumFramesPerBlock, Frequency, SampleRate, Phase);
 			Analyzer.Process(Buffer, Results);
 		}
 		
diff --git a/Engine/Plugins/Runtime/Harmonix/Source/HarmonixMetasound/Private/HarmonixMetasound/Nodes/FFTAnalyzerNode.cpp b/Engine/Plugins/Runtime/Harmonix/Source/HarmonixMetasound/Private/HarmonixMetasound/Nodes/FFTAnalyzerNode.cpp
--- a/Engine/Plugins/Runtime/Harmonix/Source/HarmonixMetasound/Private/HarmonixMetasound/Nodes/FFTAnalyzerNode.cpp
+++ b/Engine/Plugins/Runtime/Harmonix/Source/HarmonixMetasound/Private/HarmonixMetasound/Nodes/FFTAnalyzerNode.cpp
@@ -72,7 +72,6 @@ namespace HarmonixMetasound
 		FFFTAnalyzerOperator(const FBuildOperatorParams& CreateOperatorParams, FInputs&& Inputs)
 			: Inputs(MoveTemp(Inputs))
 			, FFTAnalyzer(CreateOperatorParams.OperatorSettings.GetSampleRate())
-			, InputBufferAlias(NumChannels, CreateOperatorParams.OperatorSettings.GetNumFramesPerBlock(), EAudioBufferCleanupMode::DontDelete)
 			, ResultsOut(FHarmonixFFTAnalyzerResultsWriteRef::CreateNew())
 		{
 			Reset(CreateOperatorParams);
@@ -204,8 +203,7 @@ namespace HarmonixMetasound
 			}
 			
 			RefreshParams();
-			InputBufferAlias.Alias(Inputs.Audio->GetData(), Inputs.Audio->Num(), NumChannels);
-			FFTAnalyzer.Process(InputBufferAlias, *ResultsOut);
+			FFTAnalyzer.Process(*Inputs.Audio, *ResultsOut);
 		}
 
 	private:
@@ -258,7 +256,6 @@ namespace HarmonixMetasound
 		FInputs Inputs;
 		FHarmonixFFTAnalyzerSettings FFTSettings;
 		FFFTAnalyzer FFTAnalyzer;
-		TAudioBuffer<float> InputBufferAlias;
 
 		// OUTPUT
 		FHarmonixFFTAnalyzerResultsWriteRef ResultsOut;
diff --git a/Engine/Plugins/Runtime/Harmonix/Source/HarmonixMetasoundTests/Private/Nodes/FFTAnalyzerNode.spec.cpp b/Engine/Plugins/Runtime/Harmonix/Source/HarmonixMetasoundTests/Private/Nodes/FFTAnalyzerNode.spec.cpp
--- a/Engine/Plugins/Runtime/Harmonix/Source/HarmonixMetasoundTests/Private/Nodes/FFTAnalyzerNode.spec.cpp
+++ b/Engine/Plugins/Runtime/Harmonix/Source/HarmonixMetasoundTests/Private/Nodes/FFTAnalyzerNode.spec.cpp
@@ -36,8 +36,9 @@ namespace HarmonixMetasoundTests::FFTAnalyzerNode
 				Harmonix::Dsp::AudioAnalysis::FFFTAnalyzer FFTAnalyzerForComparison{ Generator->OperatorSettings.GetSampleRate() };
 
 				// Make some noise and copy it to the node's input
-				TAudioBuffer<float> InputBuffer{ 1, Generator->OperatorSettings.GetNumFramesPerBlock(), EAudioBufferCleanupMode::Delete };
-				HarmonixDsp::GenerateWhiteNoise(InputBuffer.GetValidChannelData(0), InputBuffer.GetNumValidFrames());
+				Audio::FAlignedFloatBuffer InputBuffer;
+				InputBuffer.SetNumUninitialized(Generator->OperatorSettings.GetNumFramesPerBlock());
+				HarmonixDsp::GenerateWhiteNoise(InputBuffer.GetData(), InputBuffer.Num());
 
 				auto GeneratorInputBuffer = Generator->GetInputWriteReference<Metasound::FAudioBuffer>("In");
 
@@ -49,15 +50,16 @@ namespace HarmonixMetasoundTests::FFTAnalyzerNode
 				FHarmonixFFTAnalyzerResults RawResults;
 
 				constexpr int32 NumBlocksToRender = 4;
+				Metasound::FAudioBuffer OutputBuffer{ InputBuffer.Num() };
+				
 				for (int32 i = 0; i < NumBlocksToRender; ++i)
 				{
 					FMemory::Memcpy(
 						(*GeneratorInputBuffer)->GetData(),
-						InputBuffer.GetValidChannelData(0),
-						InputBuffer.GetNumValidFrames() * sizeof(float));
+						InputBuffer.GetData(),
+						InputBuffer.Num() * sizeof(float));
 				
 					// Process the node
-					Metasound::FAudioBuffer OutputBuffer{ InputBuffer.GetNumValidFrames() };
 					Generator->OnGenerateAudio(OutputBuffer.GetData(), OutputBuffer.Num());
 
 					// Process the raw DSP
