Seek MIDIClock upon attach to clock authority to the authority's Tick
- Prior to this change, if a clock driver was hooked up dynamically while it was playing, it would advance with an unexpected variable offset relative to the driver
#tests PIE, MIDIClock test suite

diff --git a/Engine/Plugins/Runtime/Harmonix/Source/HarmonixMetasound/Private/HarmonixMetasound/DataTypes/MidiClock.cpp b/Engine/Plugins/Runtime/Harmonix/Source/HarmonixMetasound/Private/HarmonixMetasound/DataTypes/MidiClock.cpp
--- a/Engine/Plugins/Runtime/Harmonix/Source/HarmonixMetasound/Private/HarmonixMetasound/DataTypes/MidiClock.cpp
+++ b/Engine/Plugins/Runtime/Harmonix/Source/HarmonixMetasound/Private/HarmonixMetasound/DataTypes/MidiClock.cpp
@@ -213,6 +213,25 @@ namespace HarmonixMetasound
 		return MidiClockEventsInBlock;
 	}
 
+	const FMidiClockEvent* FMidiClock::FindLastMidiClockEventAtBlockSampleFrame(int32 FrameIndex) const
+	{
+		int32 Index = Algo::LowerBoundBy(MidiClockEventsInBlock, FrameIndex, &FMidiClockEvent::BlockFrameIndex);
+		if (MidiClockEventsInBlock.IsValidIndex(Index))
+		{
+			while (MidiClockEventsInBlock.IsValidIndex(Index + 1) && MidiClockEventsInBlock[Index + 1].BlockFrameIndex == FrameIndex)
+			{
+				++Index;
+			}
+			return &MidiClockEventsInBlock[Index];
+		}
+
+		if (!MidiClockEventsInBlock.IsEmpty())
+		{
+			return &MidiClockEventsInBlock.Last();
+		}
+		return nullptr;
+	}
+
 	EMusicPlayerTransportState FMidiClock::GetTransportStateAtBlockSampleFrame(int32 FrameIndex) const
 	{
 		return GetTransportTimestampForBlockSampleFrame(FrameIndex).TransportState;
@@ -306,6 +325,7 @@ namespace HarmonixMetasound
 	void FMidiClock::AttachToTimeAuthority(const FMidiClock& MidiClockRef)
 	{
 		DrivingMidiPlayCursorMgr->AttachToTimeAuthority(MidiClockRef.DrivingMidiPlayCursorMgr);
+		bSeekToAuthorityOnNextProcess = true;
 	}
 
 	void FMidiClock::DetachFromTimeAuthority()
@@ -418,6 +438,41 @@ namespace HarmonixMetasound
 	void FMidiClock::Process(const FMidiClock& DrivingClock, int32 StartFrame, int32 NumFrames, int32 PrerollBars, float Speed)
 	{
 		int32 EndFrame = StartFrame + NumFrames;
+
+		// When we are driven by an external clock, we want to be sure that upon hookup we are synced to the same tick as our Driver
+		// Otherwise, we move in lockstep with our Driver but potentially at an offset beat/measure/etc
+		if (bSeekToAuthorityOnNextProcess)
+		{
+			bSeekToAuthorityOnNextProcess = false;
+
+			if (const FMidiClockEvent* Event = DrivingClock.FindLastMidiClockEventAtBlockSampleFrame(StartFrame))
+			{
+				int32 SeekTick = 0;
+				switch (Event->Msg.Type)
+				{
+				case FMidiClockMsg::EType::SeekTo:
+				case FMidiClockMsg::EType::Reset:
+					{
+						SeekTick = Event->Msg.ToTick();
+						break;
+					}
+				case FMidiClockMsg::EType::SeekThru:
+				case FMidiClockMsg::EType::AdvanceThru:
+					{
+						SeekTick = Event->Msg.ThruTick() + 1;
+						break;
+					}
+				case FMidiClockMsg::EType::Loop:
+					{
+						SeekTick = Event->Msg.AsLoop().LoopStartTick;
+						break;
+					}
+				}
+				SeekTo(StartFrame, CalculateMappedTick(SeekTick), PrerollBars);
+				StartFrame = Event->BlockFrameIndex + 1;
+			}
+		}
+
 		const TArray<FMidiClockEvent>& ClockEvents = DrivingClock.GetMidiClockEventsInBlock();
 		int32 Index = Algo::LowerBoundBy(ClockEvents, StartFrame, &FMidiClockEvent::BlockFrameIndex);
 		while (ClockEvents.IsValidIndex(Index))
diff --git a/Engine/Plugins/Runtime/Harmonix/Source/HarmonixMetasound/Public/HarmonixMetasound/DataTypes/MidiClock.h b/Engine/Plugins/Runtime/Harmonix/Source/HarmonixMetasound/Public/HarmonixMetasound/DataTypes/MidiClock.h
--- a/Engine/Plugins/Runtime/Harmonix/Source/HarmonixMetasound/Public/HarmonixMetasound/DataTypes/MidiClock.h
+++ b/Engine/Plugins/Runtime/Harmonix/Source/HarmonixMetasound/Public/HarmonixMetasound/DataTypes/MidiClock.h
@@ -92,6 +92,7 @@ namespace HarmonixMetasound
 		}
 
 		const TArray<FMidiClockEvent>& GetMidiClockEventsInBlock() const;
+		const FMidiClockEvent* FindLastMidiClockEventAtBlockSampleFrame(int32 FrameIndex) const;
 
 		EMusicPlayerTransportState GetTransportStateAtBlockSampleFrame(int32 FrameIndex) const;
 		EMusicPlayerTransportState GetTransportStateAtEndOfBlock() const;
@@ -284,6 +285,7 @@ namespace HarmonixMetasound
 
 		bool SmoothingEnabled = false;
 		TSharedRef<FMidiPlayCursorMgr> DrivingMidiPlayCursorMgr;
+		bool bSeekToAuthorityOnNextProcess = false;
 	};
 
 	// Declare aliases IN the namespace...
