Fix CIS compile error caused by CL 32871520 -- New types for MIDI Event types
#tests Local compile

diff --git a/Engine/Plugins/Runtime/Harmonix/Source/HarmonixMetasound/Private/HarmonixMetasound/DataTypes/MidiClock.cpp b/Engine/Plugins/Runtime/Harmonix/Source/HarmonixMetasound/Private/HarmonixMetasound/DataTypes/MidiClock.cpp
--- a/Engine/Plugins/Runtime/Harmonix/Source/HarmonixMetasound/Private/HarmonixMetasound/DataTypes/MidiClock.cpp
+++ b/Engine/Plugins/Runtime/Harmonix/Source/HarmonixMetasound/Private/HarmonixMetasound/DataTypes/MidiClock.cpp
@@ -448,25 +448,25 @@ namespace HarmonixMetasound
 			if (const FMidiClockEvent* Event = DrivingClock.FindLastMidiClockEventAtBlockSampleFrame(StartFrame))
 			{
 				int32 SeekTick = 0;
-				switch (Event->Msg.Type)
+				if (Event->Msg.IsType<MidiClockMessageTypes::FSeekTo>())
 				{
-				case FMidiClockMsg::EType::SeekTo:
-				case FMidiClockMsg::EType::Reset:
-					{
-						SeekTick = Event->Msg.ToTick();
-						break;
-					}
-				case FMidiClockMsg::EType::SeekThru:
-				case FMidiClockMsg::EType::AdvanceThru:
-					{
-						SeekTick = Event->Msg.ThruTick() + 1;
-						break;
-					}
-				case FMidiClockMsg::EType::Loop:
-					{
-						SeekTick = Event->Msg.AsLoop().LoopStartTick;
-						break;
-					}
+					SeekTick = Event->Msg.Get<MidiClockMessageTypes::FSeekTo>().ToTick;
+				}
+				else if (Event->Msg.IsType<MidiClockMessageTypes::FReset>())
+				{
+					SeekTick = Event->Msg.Get<MidiClockMessageTypes::FReset>().ToTick;
+				}
+				else if (Event->Msg.IsType<MidiClockMessageTypes::FSeekThru>())
+				{
+					SeekTick = Event->Msg.Get<MidiClockMessageTypes::FSeekThru>().ThruTick + 1;
+				}
+				else if (Event->Msg.IsType<MidiClockMessageTypes::FAdvanceThru>())
+				{
+					SeekTick = Event->Msg.Get<MidiClockMessageTypes::FAdvanceThru>().ThruTick + 1;
+				}
+				else if (Event->Msg.IsType<MidiClockMessageTypes::FLoop>())
+				{
+					SeekTick = Event->Msg.Get<MidiClockMessageTypes::FLoop>().LoopStartTick;
 				}
 				SeekTo(StartFrame, CalculateMappedTick(SeekTick), PrerollBars);
 				StartFrame = Event->BlockFrameIndex + 1;
