Account for two new MIDI Event Types in Main when seeking to clock authority that were not present in original branch

diff --git a/Engine/Plugins/Runtime/Harmonix/Source/HarmonixMetasound/Private/HarmonixMetasound/DataTypes/MidiClock.cpp b/Engine/Plugins/Runtime/Harmonix/Source/HarmonixMetasound/Private/HarmonixMetasound/DataTypes/MidiClock.cpp
--- a/Engine/Plugins/Runtime/Harmonix/Source/HarmonixMetasound/Private/HarmonixMetasound/DataTypes/MidiClock.cpp
+++ b/Engine/Plugins/Runtime/Harmonix/Source/HarmonixMetasound/Private/HarmonixMetasound/DataTypes/MidiClock.cpp
@@ -464,10 +464,20 @@ namespace HarmonixMetasound
 				{
 					SeekTick = Event->Msg.Get<MidiClockMessageTypes::FAdvanceThru>().ThruTick + 1;
 				}
+				// These lower three branches shouldn't be hit in practice, but putting them here for completeness
 				else if (Event->Msg.IsType<MidiClockMessageTypes::FLoop>())
 				{
 					SeekTick = Event->Msg.Get<MidiClockMessageTypes::FLoop>().LoopStartTick;
 				}
+				else if (Event->Msg.IsType<MidiClockMessageTypes::FTempoChange>())
+				{
+					SeekTick = Event->Msg.Get<MidiClockMessageTypes::FTempoChange>().Tick;
+				}
+				else if (Event->Msg.IsType<MidiClockMessageTypes::FTimeSignatureChange>())
+				{
+					SeekTick = Event->Msg.Get<MidiClockMessageTypes::FTimeSignatureChange>().Tick;
+				}
+
 				SeekTo(StartFrame, CalculateMappedTick(SeekTick), PrerollBars);
 				StartFrame = Event->BlockFrameIndex + 1;
 			}
