HarmonixFunctionalTest & FunctionalTestaAction

HarmonixFunctionalTest
- Reusable & instanceable functional test actor providing the basic behavior of playing a Metasound, capturing the output, and comparing it to some expected output file
- Has a list of "Functional Test Actions" which you add to to define the "actions" that happen during the test.

FunctionalTestAction
- An abstract base class, allowing you to define your own test "actions" for atest
- defines a number of useful test actions, and "collection" type actions for running actions in series or in parallel.
- added a metasound specific functional test action for setting the parameter on a metasound

#tests Harmonix Funcitonal Tests

#virtualized

diff --git a/Engine/Plugins/Runtime/Harmonix/Source/HarmonixMetasoundTests/Private/HarmonixFunctionalTestAction.cpp b/Engine/Plugins/Runtime/Harmonix/Source/HarmonixMetasoundTests/Private/HarmonixFunctionalTestAction.cpp
new file mode 100644
--- /dev/null
+++ b/Engine/Plugins/Runtime/Harmonix/Source/HarmonixMetasoundTests/Private/HarmonixFunctionalTestAction.cpp
@@ -0,0 +1,197 @@
+﻿// Copyright Epic Games, Inc. All Rights Reserved.
+
+#include "HarmonixFunctionalTestAction.h"
+
+void UHarmonixFunctionalTestActionSequence::Prepare_Implementation(AFunctionalTest* Test)
+{
+	ActionStack.Reserve(ActionSequence.Num());
+	for (int32 Idx = ActionSequence.Num() - 1; Idx >= 0; --Idx)
+	{
+		ActionStack.Add(ActionSequence[Idx]);
+	}
+
+	if (ActionStack.Num() > 0)
+	{
+		CurrentAction = ActionStack.Pop();
+	}
+
+	if (CurrentAction)
+	{
+		CurrentAction->Prepare(Test);
+	}
+}
+
+void UHarmonixFunctionalTestActionSequence::OnStart_Implementation(AFunctionalTest* Test)
+{
+	if (IsFinished())
+	{
+		return;
+	}
+
+	if (CurrentAction)
+	{
+		CurrentAction->OnStart(Test);
+		if (CurrentAction->IsFinished())
+		{
+			if (!CurrentAction->ShouldContinue())
+			{
+				Finish(false);
+			}
+			CurrentAction = nullptr;
+		}
+	}
+}
+
+void UHarmonixFunctionalTestActionSequence::Tick_Implementation(AFunctionalTest* Test, float DeltaSeconds)
+{
+	if (IsFinished())
+	{
+		return;
+	}
+		
+	if (CurrentAction == nullptr && ActionStack.Num() > 0)
+	{
+		CurrentAction = ActionStack.Pop();
+		CurrentAction->Prepare(Test);
+		CurrentAction->OnStart(Test);
+	}
+		
+	if (CurrentAction)
+	{
+		CurrentAction->Tick(Test, DeltaSeconds);
+		if (CurrentAction->IsFinished())
+		{
+			if (!CurrentAction->ShouldContinue())
+			{
+				Finish(false);
+			}
+			CurrentAction = nullptr;
+		}
+	}
+
+	if (!IsFinished() && CurrentAction == nullptr && ActionStack.IsEmpty())
+	{
+		Finish(true);
+	}
+}
+
+void UHarmonixFunctionalTestActionSequence::OnFinished_Implementation()
+{
+	CurrentAction = nullptr;
+	ActionStack.Empty();
+}
+
+void UHarmonixFunctionalTestActionParallel::Prepare_Implementation(AFunctionalTest* Test)
+{
+	ActionStack.Reserve(ParallelActions.Num());
+	for (int32 Idx = ParallelActions.Num() - 1; Idx >= 0; --Idx)
+	{
+		ActionStack.Add(ParallelActions[Idx]);
+	}
+
+	int32 ActionIdx = ActionStack.Num() - 1;
+	while (ActionIdx >= 0)
+	{
+		ActionStack[ActionIdx]->Prepare(Test);
+		--ActionIdx;
+	}
+}
+
+void UHarmonixFunctionalTestActionParallel::OnStart_Implementation(AFunctionalTest* Test)
+{
+	int32 ActionIdx = ActionStack.Num() - 1;
+	while (ActionIdx >= 0)
+	{
+		TObjectPtr<UHarmonixFunctionalTestAction> Action = ActionStack[ActionIdx];
+
+		Action->OnStart(Test);
+		if (Action->IsFinished())
+		{
+			if (!Action->ShouldContinue())
+			{
+				ActionIdx = -1;
+				FinishAllActions(false);
+				break;
+			}
+			
+			ActionStack.RemoveAt(ActionIdx);
+		}
+		--ActionIdx;
+	}
+}
+
+void UHarmonixFunctionalTestActionParallel::Tick_Implementation(AFunctionalTest* Test, float DeltaSeconds)
+{
+	if (IsFinished())
+	{
+		return;
+	}
+
+	bool bContinue = true;
+	int32 ActionIdx = ActionStack.Num() - 1;
+	while (ActionIdx >= 0)
+	{
+		TObjectPtr<UHarmonixFunctionalTestAction> Action = ActionStack[ActionIdx];
+
+		Action->Tick(Test, DeltaSeconds);
+		if (Action->IsFinished())
+		{
+			if (!Action->ShouldContinue())
+			{
+				ActionIdx = -1;
+				bContinue = false;
+				FinishAllActions(bContinue);
+				break;
+			}
+			
+			ActionStack.RemoveAt(ActionIdx);
+		}
+		--ActionIdx;
+	}
+
+	if (!IsFinished() && ActionStack.IsEmpty())
+	{
+		Finish(bContinue);
+	}
+}
+
+void UHarmonixFunctionalTestActionParallel::OnFinished_Implementation()
+{
+	FinishAllActions(true);
+}
+
+void UHarmonixFunctionalTestActionParallel::FinishAllActions(bool bContinue)
+{
+	while (ActionStack.Num() > 0)
+	{
+		
+		ActionStack.Pop()->Finish(bContinue);
+	}
+}
+
+void UHarmonixFunctionalTestActionDelay::OnStart_Implementation(AFunctionalTest* Test)
+{
+	TotalTime = 0.0f;
+}
+
+void UHarmonixFunctionalTestActionDelay::Tick_Implementation(AFunctionalTest* Test, float DeltaSeconds)
+{
+	if (IsFinished())
+	{
+		return;
+	}
+	
+	TotalTime += DeltaSeconds;
+
+	if (TotalTime > DelaySeconds)
+	{
+		Finish(true);
+	}
+}
+
+void UHarmonixFunctionalTestActionFinishTest::OnStart_Implementation(AFunctionalTest* Test)
+{
+	Test->FinishTest(Result, Message);
+	Finish(true);
+}
+
diff --git a/Engine/Plugins/Runtime/Harmonix/Source/HarmonixMetasoundTests/Private/HarmonixFunctionalTestAction.h b/Engine/Plugins/Runtime/Harmonix/Source/HarmonixMetasoundTests/Private/HarmonixFunctionalTestAction.h
new file mode 100644
--- /dev/null
+++ b/Engine/Plugins/Runtime/Harmonix/Source/HarmonixMetasoundTests/Private/HarmonixFunctionalTestAction.h
@@ -0,0 +1,160 @@
+﻿// Copyright Epic Games, Inc. All Rights Reserved.
+
+#pragma once
+
+#include "UObject/Object.h"
+#include "FunctionalTest.h"
+
+#include "HarmonixFunctionalTestAction.generated.h"
+
+class AFunctionalTest;
+
+UCLASS(Blueprintable, EditInlineNew, CollapseCategories, HideCategories=(Object), Abstract)
+class UHarmonixFunctionalTestAction : public UObject
+{
+	GENERATED_BODY()
+	
+public:
+
+	virtual void Prepare_Implementation(AFunctionalTest* Test) {};
+
+	virtual void OnStart_Implementation(AFunctionalTest* Test) {};
+	
+	virtual void OnFinished_Implementation() {};
+
+	virtual void Tick_Implementation(AFunctionalTest* Test, float DeltaSeconds) {};
+
+	bool ShouldContinue() const
+	{
+		return bShouldContinue;
+	}
+	
+	UFUNCTION(BlueprintNativeEvent, Category="Functional Testing")
+	void Prepare(AFunctionalTest* Test);
+
+	UFUNCTION(BlueprintNativeEvent, Category="Functional Testing")
+	void OnStart(AFunctionalTest* Test);
+
+	UFUNCTION(BlueprintNativeEvent, Category="Functional Testing")
+	void OnFinished();
+
+	UFUNCTION(BlueprintNativeEvent, Category = "Functional Testing")
+	void Tick(AFunctionalTest* Test, float DeltaSeconds);
+
+	// called by the user when finished with the step
+	UFUNCTION(BlueprintCallable, Category = "Functional Testing")
+	void Finish(bool bContinue = true)
+	{
+		if (IsFinished())
+		{
+			return;
+		}
+		
+		bShouldContinue = bContinue;
+		bIsFinished = true;
+		OnFinished();
+	}
+
+	UFUNCTION(BlueprintPure, Category = "Functional Testing")
+	bool IsFinished() const { return bIsFinished; }
+	
+private:
+	
+	bool bIsFinished = false;
+	bool bShouldContinue = false;
+};
+
+UCLASS(Blueprintable, meta=(DisplayName = "Action Sequence"))
+class UHarmonixFunctionalTestActionSequence : public UHarmonixFunctionalTestAction
+{
+	GENERATED_BODY()
+
+public:
+	
+	virtual void Prepare_Implementation(AFunctionalTest* Test) override;
+
+	virtual void OnStart_Implementation(AFunctionalTest* Test) override;
+
+	virtual void Tick_Implementation(AFunctionalTest* Test, float DeltaSeconds) override;
+
+	virtual void OnFinished_Implementation() override;
+
+	UPROPERTY(EditAnywhere, Instanced, Category = "Functional Testing")
+	TArray<TObjectPtr<UHarmonixFunctionalTestAction>> ActionSequence;
+
+private:
+	
+	UPROPERTY(Transient)
+	TArray<TObjectPtr<UHarmonixFunctionalTestAction>> ActionStack;
+
+	UPROPERTY(Transient)
+	TObjectPtr<UHarmonixFunctionalTestAction> CurrentAction;
+};
+
+UCLASS(Blueprintable, meta=(DisplayName = "Parallel Actions"))
+class UHarmonixFunctionalTestActionParallel : public UHarmonixFunctionalTestAction
+{
+	GENERATED_BODY()
+
+public:
+	
+	virtual void Prepare_Implementation(AFunctionalTest* Test) override;
+
+	virtual void OnStart_Implementation(AFunctionalTest* Test) override;
+
+	virtual void Tick_Implementation(AFunctionalTest* Test, float DeltaSeconds) override;
+
+	virtual void OnFinished_Implementation() override;
+
+	UPROPERTY(EditAnywhere, Instanced, Category = "Functional Testing")
+	TArray<TObjectPtr<UHarmonixFunctionalTestAction>> ParallelActions;
+
+private:
+	
+	void FinishAllActions(bool bContinue);
+	
+	UPROPERTY(Transient)
+	TArray<TObjectPtr<UHarmonixFunctionalTestAction>> ActionStack;
+};
+
+UCLASS(NotBlueprintable, meta=(DisplayName="Delay"))
+class UHarmonixFunctionalTestActionDelay final : public UHarmonixFunctionalTestAction
+{
+public:
+
+	GENERATED_BODY()
+	
+	virtual void OnStart_Implementation(AFunctionalTest* Test) override;
+	
+	virtual void Tick_Implementation(AFunctionalTest* Test, float DeltaSeconds) override;
+
+	UPROPERTY(EditAnywhere, Category = "Functional Testing")
+	float DelaySeconds = 1.0f;
+
+private:
+
+	float TotalTime = 0.0f;
+};
+
+UCLASS(NotBlueprintable, meta=(DisplayName="Wait For Timeout"))
+class UHarmonixFunctionalTestActionWaitForTimeout final : public UHarmonixFunctionalTestAction
+{
+	GENERATED_BODY()
+};
+
+UCLASS(NotBlueprintable, meta=(DisplayName="Finish Test"))
+class UHarmonixFunctionalTestActionFinishTest : public UHarmonixFunctionalTestAction
+{
+	GENERATED_BODY()
+
+public:
+	
+	virtual void OnStart_Implementation(AFunctionalTest* Test) override;
+
+	UPROPERTY(EditAnywhere, Category = "Functional Testing")
+	EFunctionalTestResult Result = EFunctionalTestResult::Default;
+
+	UPROPERTY(EditAnywhere, Category = "Functional Testing")
+	FString Message = TEXT("Finish Test");
+	
+};
diff --git a/Engine/Plugins/Runtime/Harmonix/Source/HarmonixMetasoundTests/Private/HarmonixMetasoundFunctionalTest.cpp b/Engine/Plugins/Runtime/Harmonix/Source/HarmonixMetasoundTests/Private/HarmonixMetasoundFunctionalTest.cpp
--- a/Engine/Plugins/Runtime/Harmonix/Source/HarmonixMetasoundTests/Private/HarmonixMetasoundFunctionalTest.cpp
+++ b/Engine/Plugins/Runtime/Harmonix/Source/HarmonixMetasoundTests/Private/HarmonixMetasoundFunctionalTest.cpp
@@ -1,6 +1,7 @@
 ﻿// Copyright Epic Games, Inc. All Rights Reserved.
 
 #include "HarmonixMetasoundFunctionalTest.h"
+#include "HarmonixFunctionalTestAction.h"
 
 #include "MetasoundGenerator.h"
 #include "Components/AudioComponent.h"
@@ -11,6 +12,7 @@
 #include "Audio/SimpleWaveWriter.h"
 #include "Analysis/MetasoundFrontendVertexAnalyzerAudioBuffer.h"
 #include "HAL/FileManager.h"
+#include "HarmonixDsp/AudioAnalysis/WaveFileComparison.h"
 #include "Interfaces/IPluginManager.h"
 #include "Kismet/KismetSystemLibrary.h"
 #include "Misc/FileHelper.h"
@@ -111,7 +113,7 @@ bool UHarmonixMetasoundFunctionalTestLibrary::ReadAudioFromFile(const FString& F
 	OutNumChannels = -1;
 	OutFormatTag = 0;
 	
-	if (!ensure(FPaths::FileExists(Filepath)))
+	if (!FPaths::FileExists(Filepath))
 	{
 		UE_LOG(LogHarmonixMetasoundTests, Error, TEXT("Failed to read wave file %s: File does not exist"), *Filepath);
 		return false;
@@ -201,9 +203,9 @@ AHarmonixMetasoundFunctionalTest::AHarmonixMetasoundFunctionalTest(const FObject
 
 bool AHarmonixMetasoundFunctionalTest::IsReady_Implementation()
 {
-	bool OutIsReady = AFunctionalTest::IsReady_Implementation();
-
-	return OutIsReady && GeneratorHandle;
+	bool OutIsReady = AFunctionalTest::IsReady_Implementation() && GeneratorHandle;
+	UE_LOG(LogHarmonixMetasoundTests, Log, TEXT("%s -- Is Ready: %d"), *TestLabel, OutIsReady);
+	return OutIsReady;
 }
 
 void AHarmonixMetasoundFunctionalTest::CompareResults()
@@ -256,14 +258,20 @@ void AHarmonixMetasoundFunctionalTest::StartTest()
 {
 	AFunctionalTest::StartTest();
 
+	UE_LOG(LogHarmonixMetasoundTests, Log, TEXT("%s -- StartTest"), *TestLabel);
+	
 	if (AudioComponent && AudioAutoStart)
 	{
 		AudioComponent->Play();
 	}
+
+	ActionSequence->OnStart(this);
 }
 
 void AHarmonixMetasoundFunctionalTest::FinishTest(EFunctionalTestResult TestResult, const FString& Message)
 {
+	ActionSequence->Finish(true);
+	
 	if (AudioComponent)
 	{
 		AudioComponent->Stop();
@@ -274,9 +282,26 @@ void AHarmonixMetasoundFunctionalTest::FinishTest(EFunctionalTestResult TestResu
 	AFunctionalTest::FinishTest(TestResult, Message);
 }
 
+void AHarmonixMetasoundFunctionalTest::Tick(float DeltaSeconds)
+{
+	if (IsRunning())
+	{
+		UE_LOG(LogHarmonixMetasoundTests, Log, TEXT("%s -- Ticking"), *TestLabel);
+		ActionSequence->Tick(this, DeltaSeconds);
+		if (ActionSequence && ActionSequence->IsFinished())
+		{
+			FinishTest(EFunctionalTestResult::Default, TEXT("Test completed"));
+		}
+	}
+	AFunctionalTest::Tick(DeltaSeconds);
+}
+
+
 void AHarmonixMetasoundFunctionalTest::PrepareTest()
 {
 	AFunctionalTest::PrepareTest();
+
+	UE_LOG(LogHarmonixMetasoundTests, Log, TEXT("%s -- PrepareTest"), *TestLabel);
 	
 	if (!AudioComponent)
 	{
@@ -288,6 +313,10 @@ void AHarmonixMetasoundFunctionalTest::PrepareTest()
 		return;
 	}
 
+	ActionSequence = NewObject<UHarmonixFunctionalTestActionSequence>(this);
+	ActionSequence->ActionSequence = FunctionalTestActions;
+	ActionSequence->Prepare(this);
+
 	AudioComponent->Sound = TestSound;
 	
 	GeneratorHandle = UMetasoundGeneratorHandle::CreateMetaSoundGeneratorHandle(AudioComponent);
@@ -336,6 +365,10 @@ void AHarmonixMetasoundFunctionalTest::PrepareTest()
 
 void AHarmonixMetasoundFunctionalTest::OnTestFinishedEvent()
 {
+	UE_LOG(LogHarmonixMetasoundTests, Log, TEXT("%s -- OnTestFinished"), *TestLabel);
+	
+	ActionSequence = nullptr;
+	
 	OnTestFinished.RemoveDynamic(this, &AHarmonixMetasoundFunctionalTest::OnTestFinishedEvent);
 	
 	if (GeneratorHandle)
@@ -354,6 +387,24 @@ void AHarmonixMetasoundFunctionalTest::OnTestFinishedEvent()
 			Generator->RemoveOutputVertexAnalyzer(AudioOutAnalyzerAddress);
 		}
 	}
+
+	if (!WavFilename_Expected.IsEmpty() && !WavFilename_Output.IsEmpty())
+	{
+			
+		TSharedPtr<IPlugin> Plugin = IPluginManager::Get().FindPlugin(TEXT("Harmonix"));
+		check(Plugin);
+		FString TestAudioDir = Plugin->GetContentDir() / TEXT("Editor/Tests/Audio");
+		FString Filepath_Expected = TestAudioDir / WavFilename_Expected;
+
+		static const FString RootPath = FPaths::AudioCaptureDir();
+		FString Filepath_Output = RootPath / WavFilename_Output;
+		
+		Harmonix::Dsp::AudioAnalysis::FWaveFileComparison FileComparison;
+		FileComparison.LoadForCompare(Filepath_Expected, Filepath_Output);
+		float PSNR =  FileComparison.GetPSNR();
+
+		UE_LOG(LogHarmonixMetasoundTests, Log, TEXT("PSNR of files is: %.2f"), PSNR);
+	}
 }
 
 
diff --git a/Engine/Plugins/Runtime/Harmonix/Source/HarmonixMetasoundTests/Private/HarmonixMetasoundFunctionalTest.h b/Engine/Plugins/Runtime/Harmonix/Source/HarmonixMetasoundTests/Private/HarmonixMetasoundFunctionalTest.h
--- a/Engine/Plugins/Runtime/Harmonix/Source/HarmonixMetasoundTests/Private/HarmonixMetasoundFunctionalTest.h
+++ b/Engine/Plugins/Runtime/Harmonix/Source/HarmonixMetasoundTests/Private/HarmonixMetasoundFunctionalTest.h
@@ -16,6 +16,9 @@ namespace Audio
 	class FSimpleWaveWriter;
 }
 
+class UHarmonixFunctionalTestAction;
+class UHarmonixFunctionalTestActionSequence;
+
 UCLASS()
 class UHarmonixMetasoundFunctionalTestLibrary final : public UBlueprintFunctionLibrary
 {
@@ -64,8 +67,13 @@ public:
 	UPROPERTY(EditAnywhere, Category = "Functional Testing")
 	FString WavFilename_Output;
 
+	UPROPERTY(EditAnywhere, Instanced, Category = "Functional Testing")
+	TArray<TObjectPtr<UHarmonixFunctionalTestAction>> FunctionalTestActions;
+
 	virtual void FinishTest(EFunctionalTestResult TestResult, const FString& Message) override;
 
+	virtual void Tick(float DeltaSeconds) override;
+
 protected:
 	
 	virtual void PrepareTest() override;
@@ -85,6 +93,9 @@ protected:
 	UPROPERTY(Transient)
 	UMetasoundGeneratorHandle* GeneratorHandle;
 
+	UPROPERTY(Transient)
+	TObjectPtr<UHarmonixFunctionalTestActionSequence> ActionSequence;
+
 	// the audio captured from the metasound
 	Audio::FAlignedFloatBuffer AudioCaptureOutput;
 	float AudioCaptureSampleRate = 0.0f;
diff --git a/Engine/Plugins/Runtime/Harmonix/Source/HarmonixMetasoundTests/Private/HarmonixMetasoundFunctionalTestAction.cpp b/Engine/Plugins/Runtime/Harmonix/Source/HarmonixMetasoundTests/Private/HarmonixMetasoundFunctionalTestAction.cpp
new file mode 100644
--- /dev/null
+++ b/Engine/Plugins/Runtime/Harmonix/Source/HarmonixMetasoundTests/Private/HarmonixMetasoundFunctionalTestAction.cpp
@@ -0,0 +1,22 @@
+﻿// Copyright Epic Games, Inc. All Rights Reserved.
+
+#include "HarmonixMetasoundFunctionalTestAction.h"
+#include "Components/AudioComponent.h"
+
+#include "FunctionalTest.h"
+
+void UHarmonixMetasoundFunctionalTestActionSetAudioParameter::OnStart_Implementation(AFunctionalTest* Test)
+{
+	if (Test)
+	{
+		if (TObjectPtr<UAudioComponent> AudioComponent = Test->FindComponentByClass<UAudioComponent>())
+		{
+			AudioComponent->SetParameter(FAudioParameter(AudioParameter));
+			Finish(true);
+		}
+		else
+		{
+			Finish(false);
+		}
+	}
+}
diff --git a/Engine/Plugins/Runtime/Harmonix/Source/HarmonixMetasoundTests/Private/HarmonixMetasoundFunctionalTestAction.h b/Engine/Plugins/Runtime/Harmonix/Source/HarmonixMetasoundTests/Private/HarmonixMetasoundFunctionalTestAction.h
new file mode 100644
--- /dev/null
+++ b/Engine/Plugins/Runtime/Harmonix/Source/HarmonixMetasoundTests/Private/HarmonixMetasoundFunctionalTestAction.h
@@ -0,0 +1,19 @@
+﻿// Copyright Epic Games, Inc. All Rights Reserved.
+
+#pragma once
+
+#include "HarmonixFunctionalTestAction.h"
+#include "AudioParameter.h"
+
+#include "HarmonixMetasoundFunctionalTestAction.generated.h"
+
+UCLASS(NotBlueprintable, Meta=(DisplayName="Set Audio Parameter"))
+class UHarmonixMetasoundFunctionalTestActionSetAudioParameter : public UHarmonixFunctionalTestAction
+{
+	GENERATED_BODY()
+public:
+	virtual void OnStart_Implementation(AFunctionalTest* Test) override;
+
+	UPROPERTY(EditAnywhere, Category = "Functional Testing")
+	FAudioParameter AudioParameter;
+};
\ No newline at end of file
