WaveFileComparison: null check file handles in case file doesn't exist

diff --git a/Engine/Plugins/Runtime/Harmonix/Source/HarmonixDsp/Private/HarmonixDsp/AudioAnalysis/WaveFileComparison.cpp b/Engine/Plugins/Runtime/Harmonix/Source/HarmonixDsp/Private/HarmonixDsp/AudioAnalysis/WaveFileComparison.cpp
--- a/Engine/Plugins/Runtime/Harmonix/Source/HarmonixDsp/Private/HarmonixDsp/AudioAnalysis/WaveFileComparison.cpp
+++ b/Engine/Plugins/Runtime/Harmonix/Source/HarmonixDsp/Private/HarmonixDsp/AudioAnalysis/WaveFileComparison.cpp
@@ -16,9 +16,17 @@ namespace Harmonix::Dsp::AudioAnalysis
 	{
 		IPlatformFile& PlatformFileApi = FPlatformFileManager::Get().GetPlatformFile();
 		IFileHandle* FileHandle = PlatformFileApi.OpenRead(*Wave1FilePath);
+		if (!FileHandle)
+		{
+			return false;
+		}
 		TSharedRef<FArchive> Wave1Archive = MakeShared<FArchiveFileReaderGeneric>(FileHandle, *Wave1FilePath, FileHandle->Size());
 
 		FileHandle = PlatformFileApi.OpenRead(*Wave2FilePath);
+		if (!FileHandle)
+		{
+			return false;
+		}
 		TSharedRef<FArchive> Wave2Archive = MakeShared<FArchiveFileReaderGeneric>(FileHandle, *Wave2FilePath, FileHandle->Size());
 		
 		return LoadForCompare(*Wave1Archive, *Wave2Archive);
