Harmonix Functional Test no longer cause timeout during preflight

- fix test parameters
- Fix Functional tests to "wait for Audio Finished" instead of waiting for a delay.
- Also added another test while I was at it

#tests Harmonix Functional Tests

#virtualized

diff --git a/Engine/Plugins/Runtime/Harmonix/Source/HarmonixMetasoundTests/Private/HarmonixMetasoundFunctionalTest.cpp b/Engine/Plugins/Runtime/Harmonix/Source/HarmonixMetasoundTests/Private/HarmonixMetasoundFunctionalTest.cpp
--- a/Engine/Plugins/Runtime/Harmonix/Source/HarmonixMetasoundTests/Private/HarmonixMetasoundFunctionalTest.cpp
+++ b/Engine/Plugins/Runtime/Harmonix/Source/HarmonixMetasoundTests/Private/HarmonixMetasoundFunctionalTest.cpp
@@ -63,7 +63,7 @@ bool UHarmonixMetasoundFunctionalTestLibrary::AddOutputLogger(UMetasoundGenerato
 	);
 	if (!ReturnValue)
 	{
-		UE_LOG(LogHarmonixMetasoundTests, Log, TEXT("Failed to add logger for output: %s"), *OutputName.ToString());
+		UE_LOG(LogHarmonixMetasoundTests, Error, TEXT("Failed to add logger for output: %s"), *OutputName.ToString());
 	}
 	return ReturnValue;
 }
@@ -373,17 +373,17 @@ void AHarmonixMetasoundFunctionalTest::OnTestFinishedEvent()
 	
 	if (GeneratorHandle)
 	{
-		if (TSharedPtr<Metasound::FMetasoundGenerator> Generator = GeneratorHandle->GetGenerator())
-		{
-			
+
 #if WITH_EDITOR
-			// only write out the test results to a file when running in editor. 
-			if (!WavFilename_Output.IsEmpty() )
-			{
-				UHarmonixMetasoundFunctionalTestLibrary::WriteAudioToFile(WavFilename_Output, Generator->OperatorSettings.GetSampleRate(), 1, AudioCaptureOutput);
-			}
+		// only write out the test results to a file when running in editor. 
+		if (!WavFilename_Output.IsEmpty() )
+		{
+			UHarmonixMetasoundFunctionalTestLibrary::WriteAudioToFile(WavFilename_Output, AudioCaptureSampleRate, 1, AudioCaptureOutput);
+		}
 #endif
-			
+
+		if (TSharedPtr<Metasound::FMetasoundGenerator> Generator = GeneratorHandle->GetGenerator())
+		{
 			Generator->RemoveOutputVertexAnalyzer(AudioOutAnalyzerAddress);
 		}
 	}
diff --git a/Engine/Plugins/Runtime/Harmonix/Source/HarmonixMetasoundTests/Private/HarmonixMetasoundFunctionalTestAction.cpp b/Engine/Plugins/Runtime/Harmonix/Source/HarmonixMetasoundTests/Private/HarmonixMetasoundFunctionalTestAction.cpp
--- a/Engine/Plugins/Runtime/Harmonix/Source/HarmonixMetasoundTests/Private/HarmonixMetasoundFunctionalTestAction.cpp
+++ b/Engine/Plugins/Runtime/Harmonix/Source/HarmonixMetasoundTests/Private/HarmonixMetasoundFunctionalTestAction.cpp
@@ -20,3 +20,19 @@ void UHarmonixMetasoundFunctionalTestActionSetAudioParameter::OnStart_Implementa
 		}
 	}
 }
+
+void UHarmonixMetasoundFunctionalTestActionSetWaitForAudioFinished::OnStart_Implementation(AFunctionalTest* Test)
+{
+	if (Test)
+	{
+		if (TObjectPtr<UAudioComponent> AudioComponent = Test->FindComponentByClass<UAudioComponent>())
+		{
+			AudioComponent->OnAudioFinished.AddDynamic(this, &UHarmonixMetasoundFunctionalTestActionSetWaitForAudioFinished::OnAudioFinished);
+		}
+	}
+}
+
+void UHarmonixMetasoundFunctionalTestActionSetWaitForAudioFinished::OnAudioFinished()
+{
+	Finish(true);
+}
diff --git a/Engine/Plugins/Runtime/Harmonix/Source/HarmonixMetasoundTests/Private/HarmonixMetasoundFunctionalTestAction.h b/Engine/Plugins/Runtime/Harmonix/Source/HarmonixMetasoundTests/Private/HarmonixMetasoundFunctionalTestAction.h
--- a/Engine/Plugins/Runtime/Harmonix/Source/HarmonixMetasoundTests/Private/HarmonixMetasoundFunctionalTestAction.h
+++ b/Engine/Plugins/Runtime/Harmonix/Source/HarmonixMetasoundTests/Private/HarmonixMetasoundFunctionalTestAction.h
@@ -16,4 +16,15 @@ public:
 
 	UPROPERTY(EditAnywhere, Category = "Functional Testing")
 	FAudioParameter AudioParameter;
+};
+
+UCLASS(NotBlueprintable, Meta=(DisplayName="Wait For Audio Finished"))
+class UHarmonixMetasoundFunctionalTestActionSetWaitForAudioFinished : public UHarmonixFunctionalTestAction
+{
+	GENERATED_BODY()
+public:
+	virtual void OnStart_Implementation(AFunctionalTest* Test) override;
+protected:
+	UFUNCTION()
+	void OnAudioFinished();
 };
\ No newline at end of file
