More StepSequencePlayerNode tests

- Test Octave Changes
- Test Page Changes

#tests Harmonix Functional Tests

#virtualized

diff --git a/Engine/Plugins/Runtime/Harmonix/Source/HarmonixMetasoundTests/Private/HarmonixMetasoundFunctionalTest.cpp b/Engine/Plugins/Runtime/Harmonix/Source/HarmonixMetasoundTests/Private/HarmonixMetasoundFunctionalTest.cpp
--- a/Engine/Plugins/Runtime/Harmonix/Source/HarmonixMetasoundTests/Private/HarmonixMetasoundFunctionalTest.cpp
+++ b/Engine/Plugins/Runtime/Harmonix/Source/HarmonixMetasoundTests/Private/HarmonixMetasoundFunctionalTest.cpp
@@ -242,6 +242,11 @@ void AHarmonixMetasoundFunctionalTest::CompareResults()
 		return;
 	}
 
+	if (!AssertEqual_Float(AudioCaptureSampleRate, SampleRate, TEXT("SampleRate")))
+	{
+		return;
+	}
+
 	float AudioCaptureDuration = AudioCaptureOutput.Num() / AudioCaptureSampleRate;
 	float AudioDuration = AudioData.Num() / (float)SampleRate;
 	if (!AssertEqual_Float(AudioCaptureDuration, AudioDuration, TEXT("AudioCaptureDuration"), 0.1f, this))
@@ -286,7 +291,6 @@ void AHarmonixMetasoundFunctionalTest::Tick(float DeltaSeconds)
 {
 	if (IsRunning())
 	{
-		UE_LOG(LogHarmonixMetasoundTests, Log, TEXT("%s -- Ticking"), *TestLabel);
 		ActionSequence->Tick(this, DeltaSeconds);
 		if (ActionSequence && ActionSequence->IsFinished())
 		{
