Unregister MidiPlayerOperator from MidiClock on destruct

#tests Harmonix Tests

diff --git a/Engine/Plugins/Runtime/Harmonix/Source/HarmonixMetasound/Private/HarmonixMetasound/Nodes/MidiPlayerNode.cpp b/Engine/Plugins/Runtime/Harmonix/Source/HarmonixMetasound/Private/HarmonixMetasound/Nodes/MidiPlayerNode.cpp
--- a/Engine/Plugins/Runtime/Harmonix/Source/HarmonixMetasound/Private/HarmonixMetasound/Nodes/MidiPlayerNode.cpp
+++ b/Engine/Plugins/Runtime/Harmonix/Source/HarmonixMetasound/Private/HarmonixMetasound/Nodes/MidiPlayerNode.cpp
@@ -45,6 +45,7 @@ namespace HarmonixMetasound
                             const int32 InPrerollBars,
 							const bool bInKillVoicesOnSeek,
 							const bool bInKillVoicesOnMidiChange);
+		~FMidiPlayerOperator();
 
 		virtual void BindInputs(FInputVertexInterfaceData& InVertexData) override;
 		virtual void BindOutputs(FOutputVertexInterfaceData& InVertexData) override;
@@ -257,6 +258,12 @@ namespace HarmonixMetasound
 		MidiOutPin->SetClock(*MidiClockOut);
 	}
 
+	FMidiPlayerOperator::~FMidiPlayerOperator()
+	{
+		MidiClockOut->UnregisterPlayCursor(this);
+	}
+
+
 	FExternallyClockedMidiPlayerOperator::FExternallyClockedMidiPlayerOperator(const FOperatorSettings& InSettings,
 																			   const FMidiAssetReadRef& InMidiAsset, 
 																			   const FMusicTransportEventStreamReadRef& InTransport,
