always use the current directory as the samples directory when importing

save off the directory that the fusion patch samples were imported to on the fusion patch itself for easier reimporting.

Also simplify some logic when importing/reimporting fusion patches.
#tests HarmonixDev import/reimport fusion patches

diff --git a/Engine/Plugins/Runtime/Harmonix/Source/HarmonixDsp/Public/HarmonixDsp/FusionSampler/FusionPatch.h b/Engine/Plugins/Runtime/Harmonix/Source/HarmonixDsp/Public/HarmonixDsp/FusionSampler/FusionPatch.h
--- a/Engine/Plugins/Runtime/Harmonix/Source/HarmonixDsp/Public/HarmonixDsp/FusionSampler/FusionPatch.h
+++ b/Engine/Plugins/Runtime/Harmonix/Source/HarmonixDsp/Public/HarmonixDsp/FusionSampler/FusionPatch.h
@@ -88,6 +88,10 @@ public:
 	UPROPERTY(VisibleAnywhere, Instanced, Category = ImportSettings)
 	TObjectPtr<class UAssetImportData> AssetImportData;
 
+	// save off the samples dir used when importing samples
+	UPROPERTY()
+	FString SamplesImportDir;
+
 	
 	virtual void PostEditChangeProperty(FPropertyChangedEvent& PropertyChangedEvent) override;
 	virtual void PostEditChangeChainProperty(FPropertyChangedChainEvent& PropertyChangedChainEvent) override;
diff --git a/Engine/Plugins/Runtime/Harmonix/Source/HarmonixDspEditor/Private/FusionPatchAssetFactory.cpp b/Engine/Plugins/Runtime/Harmonix/Source/HarmonixDspEditor/Private/FusionPatchAssetFactory.cpp
--- a/Engine/Plugins/Runtime/Harmonix/Source/HarmonixDspEditor/Private/FusionPatchAssetFactory.cpp
+++ b/Engine/Plugins/Runtime/Harmonix/Source/HarmonixDspEditor/Private/FusionPatchAssetFactory.cpp
@@ -12,6 +12,7 @@
 
 #include "JsonImporterHelper.h"
 #include "FusionPatchJsonImporter.h"
+#include "AssetRegistry/AssetRegistryModule.h"
 
 #include "Misc/FeedbackContext.h"
 #include "Misc/MessageDialog.h"
@@ -126,7 +127,7 @@ bool UFusionPatchAssetFactory::GetReplaceExistingSamplesResponse(const FString&
 
 	const FText ReplaceExistingTitle = NSLOCTEXT("FusionPatchImporter", "ReplaceExistingSamplesTitle", "Replace Existing Samples");
 	const FText ReplaceExistingMessage = FText::Format(NSLOCTEXT("FusionPatchImporter", "ReplaceExistingSamplesMsg", 
-		"Would you like to reimport and replace existing Sound Wave Assets in the directory with Samples referenced by this Fusion Fatch?" 
+		"Would you like to reimport and replace existing Sound Wave Assets in the directory with Samples referenced by this Fusion Patch?" 
 		"\n\nPatch Name: {0}"
 		"\n\nYes. Reimport and replace existing Samples with the new samples."
 		"\n\nNo. New Samples will be imported, but existing Samples will be unchanged. The Fusion Patch will reference any existing Samples in the directory with matching names."), 
@@ -165,31 +166,34 @@ void UFusionPatchAssetFactory::UpdateFusionPatchImportNotificationItem(TSharedPt
 
 UObject* UFusionPatchAssetFactory::FactoryCreateText(UClass* InClass, UObject* InParent, FName InName, EObjectFlags Flags, UObject* Context, const TCHAR* Type, const TCHAR*& Buffer, const TCHAR* BufferEnd, FFeedbackContext* Warn)
 {
+	// get the existing fusion patch if we're reimporting
+	UFusionPatch* FusionPatch = FindObject<UFusionPatch>(InParent, *InName.ToString());
+	
 	const FString LongPackagePath = FPackageName::GetLongPackagePath(InParent->GetOutermost()->GetPathName());
-
-	const UFusionPatchImportOptions* ImportOptions = nullptr;
 	
-	if (!ApplyOptionsToAllImport)
+	UFusionPatchImportOptions::FArgs Args;
+	Args.PatchName = InName;
+	// If we're reimporting and the fusion patch has saved off the samples directory
+	if (FusionPatch && !FusionPatch->SamplesImportDir.IsEmpty())
 	{
-		bool WasOkayPressed = false;
-		UFusionPatchImportOptions::FArgs Args;
-		Args.Directory = LongPackagePath;
-		ImportOptions = UFusionPatchImportOptions::GetWithDialog(MoveTemp(Args), WasOkayPressed);
-		if (!WasOkayPressed)
-		{
-			// import cancelled by user
-			return nullptr;
-		}
-		ApplyOptionsToAllImport = true;
+		Args.Directory = FusionPatch->SamplesImportDir;
 	}
 	else
 	{
-		UFusionPatchImportOptions* MutableOptions = GetMutableDefault<UFusionPatchImportOptions>();
-		if (MutableOptions->SamplesImportDir.Path.IsEmpty())
-		{
-			MutableOptions->SamplesImportDir.Path = LongPackagePath;	
-		}
-		ImportOptions = MutableOptions;
+		Args.Directory = LongPackagePath;
+	}
+
+	bool WasOkayPressed = false;
+	const UFusionPatchImportOptions* ImportOptions = UFusionPatchImportOptions::GetWithDialog(MoveTemp(Args), WasOkayPressed);
+	if (!WasOkayPressed)
+	{
+		// import cancelled by user
+		return nullptr;
+	}
+
+	if (!ensure(ImportOptions))
+	{
+		return nullptr;
 	}
 
 	if (Warn->ReceivedUserCancel())
@@ -210,7 +214,6 @@ UObject* UFusionPatchAssetFactory::FactoryCreateText(UClass* InClass, UObject* I
 	bool bImportSuccessful = false;
 	if (JsonObj.IsValid())
 	{
-		UFusionPatch* FusionPatch = FindObject<UFusionPatch>(InParent, *InName.ToString());
 		if (!FusionPatch)
 		{
 			FusionPatch = NewObject<UFusionPatch>(InParent, InName, Flags);
@@ -241,6 +244,9 @@ UObject* UFusionPatchAssetFactory::FactoryCreateText(UClass* InClass, UObject* I
 				FusionPatch->AssetImportData->Update(SourceFile);
 			}
 
+			// save off the samples dest path for simplifying reimporting
+			FusionPatch->SamplesImportDir = ImportArgs.SamplesDestPath;
+
 			bImportSuccessful = true;
 			UpdateFusionPatchImportNotificationItem(ImportNotificationItem, bImportSuccessful, InName);
 			return FusionPatch;
diff --git a/Engine/Plugins/Runtime/Harmonix/Source/HarmonixDspEditor/Private/FusionPatchImportOptions.cpp b/Engine/Plugins/Runtime/Harmonix/Source/HarmonixDspEditor/Private/FusionPatchImportOptions.cpp
--- a/Engine/Plugins/Runtime/Harmonix/Source/HarmonixDspEditor/Private/FusionPatchImportOptions.cpp
+++ b/Engine/Plugins/Runtime/Harmonix/Source/HarmonixDspEditor/Private/FusionPatchImportOptions.cpp
@@ -41,13 +41,10 @@ const FMidiScale FMidiScale::MajorScale(TArray<int8>({0, 2, 4, 5, 7, 9, 11}));
 const UFusionPatchImportOptions* UFusionPatchImportOptions::GetWithDialog(FArgs&& Args, bool& OutWasOkayPressed)
 {
 	UFusionPatchImportOptions* Options = GetMutableDefault<UFusionPatchImportOptions>();
-	// prompt on the first asset and update user settings
-	if (Options->SamplesImportDir.Path.IsEmpty())
-	{
-		// otherwise, import audio samples into the destination path provided
-		Options->SamplesImportDir.Path = MoveTemp(Args.Directory);
-	}
-	const FText ImportDialogTitle = NSLOCTEXT("FusionPatchImportOptions", "FusionPatchImportOptionsTitle", "Fusion Patch Import Options");
+	Options->SamplesImportDir.Path = MoveTemp(Args.Directory);
+	FFormatNamedArguments FormatArgs;
+	FormatArgs.Add(TEXT("PatchName"), FText::FromName(Args.PatchName));
+	const FText ImportDialogTitle = FText::Format(NSLOCTEXT("FusionPatchImportOptions", "FusionPatchImportOptionsTitle", "Fusion Patch Import Options: {PatchName}"), FormatArgs);
 	OutWasOkayPressed = UEditorDialogLibrary::ShowObjectDetailsView(ImportDialogTitle, Options);
 	return Options;
 }
@@ -66,8 +63,11 @@ const UFusionPatchCreateOptions* UFusionPatchCreateOptions::GetWithDialog(FArgs&
 	Options->StagedSoundWaves.Empty();
 
 	// apply some default settings that will work out of the box
-	Options->FusionPatchSettings.Adsr[0].IsEnabled = true;
-	Options->FusionPatchSettings.Adsr[0].SustainLevel = 1.0f;
+	FAdsrSettings& VolumeAdsr = Options->FusionPatchSettings.Adsr[0];
+	VolumeAdsr.Target = EAdsrTarget::Volume;
+	VolumeAdsr.IsEnabled = true;
+	VolumeAdsr.SustainLevel = 1.0f;
+	VolumeAdsr.Depth = 1.0f;
 	
 	return Options;
 }
diff --git a/Engine/Plugins/Runtime/Harmonix/Source/HarmonixDspEditor/Private/FusionPatchImportOptions.h b/Engine/Plugins/Runtime/Harmonix/Source/HarmonixDspEditor/Private/FusionPatchImportOptions.h
--- a/Engine/Plugins/Runtime/Harmonix/Source/HarmonixDspEditor/Private/FusionPatchImportOptions.h
+++ b/Engine/Plugins/Runtime/Harmonix/Source/HarmonixDspEditor/Private/FusionPatchImportOptions.h
@@ -10,7 +10,7 @@
 #include "HarmonixMidi/MidiConstants.h"
 #include "FusionPatchImportOptions.generated.h"
 
-UCLASS(config = Editor)
+UCLASS(config = EditorPerProjectUserSettings)
 class UFusionPatchImportOptions : public UObject
 {
 	GENERATED_BODY()
@@ -18,15 +18,16 @@ class UFusionPatchImportOptions : public UObject
 public:
 	struct FArgs
 	{
+		FName PatchName;
 		FString Directory;
 	};
 
 	static const UFusionPatchImportOptions* GetWithDialog(FArgs&& Args, bool& OutWasOkayPressed);
 
 	/** The directory to save samples to */
-	UPROPERTY(config, EditAnywhere, Category = "Import Options", Meta = (DisplayName = "Sound Waves Import Folder", ContentDir))
+	UPROPERTY(EditAnywhere, Category = "Import Options", Meta = (DisplayName = "Sound Waves Import Folder", ContentDir))
 	FDirectoryPath SamplesImportDir;
-
+	
 	/** The loading behavior to apply to the imported samples */
 	UPROPERTY(EditAnywhere, Category = "Import Options", Meta = (DisplayName = "Sound Wave Loading Behavior"))
 	ESoundWaveLoadingBehavior SampleLoadingBehavior = ESoundWaveLoadingBehavior::RetainOnLoad;
@@ -34,7 +35,6 @@ public:
 	/** The compression type to apply to the imported samples */
 	UPROPERTY(EditAnywhere, Category = "Import Options", Meta = (DisplayName = "Sound Wave Compression Type"))
 	ESoundAssetCompressionType SampleCompressionType = ESoundAssetCompressionType::BinkAudio;
-
 	
 #if WITH_EDITOR
 	/** UObject interface */
