MidiPlayerNodeTests

#tests Harmonix Functional Tests

#virtualized

diff --git a/Engine/Plugins/Runtime/Harmonix/Source/HarmonixMetasoundTests/Private/HarmonixMetasoundFunctionalTest.cpp b/Engine/Plugins/Runtime/Harmonix/Source/HarmonixMetasoundTests/Private/HarmonixMetasoundFunctionalTest.cpp
--- a/Engine/Plugins/Runtime/Harmonix/Source/HarmonixMetasoundTests/Private/HarmonixMetasoundFunctionalTest.cpp
+++ b/Engine/Plugins/Runtime/Harmonix/Source/HarmonixMetasoundTests/Private/HarmonixMetasoundFunctionalTest.cpp
@@ -22,6 +22,15 @@ DEFINE_LOG_CATEGORY(LogHarmonixMetasoundTests)
 
 namespace HarmonixMetasoundTests
 {
+	static int32 WriteOutputToFileCVar = 0;
+	FAutoConsoleVariableRef CVarWriteOutputToFile(
+		TEXT("harmonix.tests.WriteOutputToFile"),
+		WriteOutputToFileCVar,
+		TEXT("Whether to write the output of the unit tests to \".wav\" files store in the AudioCapture directory \\[ProjectDirectory]\\Saved\\AudioCaptures\\/\n")
+		TEXT("0: Disabled 1: Always write output 2: Only write output on error"),
+		ECVF_Default);
+
+	
 	FString MetasoundOutputValueAsString(const FMetaSoundOutput& Output)
 	{
 		if (float Value; Output.Get(Value))
@@ -291,10 +300,13 @@ void AHarmonixMetasoundFunctionalTest::Tick(float DeltaSeconds)
 {
 	if (IsRunning())
 	{
-		ActionSequence->Tick(this, DeltaSeconds);
-		if (ActionSequence && ActionSequence->IsFinished())
+		if (ActionSequence && !ActionSequence->IsFinished())
 		{
-			FinishTest(EFunctionalTestResult::Default, TEXT("Test completed"));
+			ActionSequence->Tick(this, DeltaSeconds);
+			if (ActionSequence->IsFinished())
+			{
+				FinishTest(EFunctionalTestResult::Default, TEXT("Test completed"));
+			}
 		}
 	}
 	AFunctionalTest::Tick(DeltaSeconds);
@@ -371,20 +383,14 @@ void AHarmonixMetasoundFunctionalTest::OnTestFinishedEvent()
 {
 	UE_LOG(LogHarmonixMetasoundTests, Log, TEXT("%s -- OnTestFinished"), *TestLabel);
 	
-	ActionSequence = nullptr;
-	
 	OnTestFinished.RemoveDynamic(this, &AHarmonixMetasoundFunctionalTest::OnTestFinishedEvent);
 	
 	if (GeneratorHandle)
 	{
-
-#if WITH_EDITOR
-		// only write out the test results to a file when running in editor. 
-		if (!WavFilename_Output.IsEmpty() )
+		if (HarmonixMetasoundTests::WriteOutputToFileCVar && !WavFilename_Output.IsEmpty() )
 		{
 			UHarmonixMetasoundFunctionalTestLibrary::WriteAudioToFile(WavFilename_Output, AudioCaptureSampleRate, 1, AudioCaptureOutput);
 		}
-#endif
 
 		if (TSharedPtr<Metasound::FMetasoundGenerator> Generator = GeneratorHandle->GetGenerator())
 		{
