Fusion Patch Create Options

Add "Octave Adjust" to allow user to adjust the resulting octave based on the note names or selected scale.

Also "Hide" scale option when the sort option won't allow it.

#tests Local Editor Tests

diff --git a/Engine/Plugins/Runtime/Harmonix/Source/HarmonixDspEditor/Private/FusionPatchImportOptions.cpp b/Engine/Plugins/Runtime/Harmonix/Source/HarmonixDspEditor/Private/FusionPatchImportOptions.cpp
--- a/Engine/Plugins/Runtime/Harmonix/Source/HarmonixDspEditor/Private/FusionPatchImportOptions.cpp
+++ b/Engine/Plugins/Runtime/Harmonix/Source/HarmonixDspEditor/Private/FusionPatchImportOptions.cpp
@@ -114,6 +114,14 @@ void UFusionPatchCreateOptions::UpdateKeyzonesWithSettings()
 		}
 	}
 
+	// apply the octave offset after applying the initial values
+	for (FKeyzoneSettings& Keyzone : Keyzones)
+	{
+		Keyzone.RootNote = FMath::Clamp(Keyzone.RootNote + OctaveAdjust * 12, 0, 127);
+		Keyzone.MinNote = FMath::Clamp(Keyzone.MinNote + OctaveAdjust * 12, 0, 127);
+		Keyzone.MaxNote = FMath::Clamp(Keyzone.MaxNote + OctaveAdjust * 12, 0, 127);
+	}
+
 	// if the min and max notes haven't already been determined from the file names, lay them out
 	if (ELockedNoteFlag::None ==  (LockedNotesMask & ELockedNoteFlag::MinMax))
 	{
diff --git a/Engine/Plugins/Runtime/Harmonix/Source/HarmonixDspEditor/Private/FusionPatchImportOptions.h b/Engine/Plugins/Runtime/Harmonix/Source/HarmonixDspEditor/Private/FusionPatchImportOptions.h
--- a/Engine/Plugins/Runtime/Harmonix/Source/HarmonixDspEditor/Private/FusionPatchImportOptions.h
+++ b/Engine/Plugins/Runtime/Harmonix/Source/HarmonixDspEditor/Private/FusionPatchImportOptions.h
@@ -117,15 +117,15 @@ public:
 	EFusionPatchKeyzoneSortOption SortOption = EFusionPatchKeyzoneSortOption::NoteNumber;
 
 	// how to assign the min and max notes for each keyzone based on ordering and note mapping
-	UPROPERTY(config, EditAnywhere, Category = "Create Patch", meta=(DisplayName="Keyzone Layout"))
+	UPROPERTY(config, EditAnywhere, Category = "Create Patch", meta=(DisplayName="Keyzone Layout", EditConditionHides, EditCondition="LockedNotesMask != ELockedNoteFlag::MinRootMax"))
 	EFusionPatchKeyzoneNoteLayoutOption LayoutOption = EFusionPatchKeyzoneNoteLayoutOption::Distribute;
 
 	// whether to assign the notes along a scale, or evenly spaced
-	UPROPERTY(Config, EditAnywhere, Category = "Create Patch", meta=(DisplayName="Scale Mapping"))
+	UPROPERTY(Config, EditAnywhere, Category = "Create Patch", meta=(DisplayName="Scale Mapping", EditConditionHides, EditCondition="SortOption == EFusionPatchKeyzoneSortOption::Lexical || SortOption == EFusionPatchKeyzoneSortOption::Index"))
 	EFusionPatchKeyzoneNoteScaleOption ScaleOption = EFusionPatchKeyzoneNoteScaleOption::None;
 	
 	// how to assign the root note based on the min and max notes
-	UPROPERTY(config, EditAnywhere, Category = "Create Patch", meta=(DisplayName="Note Alignment"))
+	UPROPERTY(config, EditAnywhere, Category = "Create Patch", meta=(DisplayName="Note Alignment", EditConditionHides, EditCondition="LockedNotesMask != ELockedNoteFlag::MinRootMax"))
 	EFusionPatchKeyzoneRootNoteOption RootNoteOption = EFusionPatchKeyzoneRootNoteOption::Min;
 
 	UPROPERTY(config, EditAnywhere, Category = "Create Patch", meta=(ClampMin = 0, ClampMax = 127, UIMin = 0, UIMax = 127))
@@ -134,6 +134,12 @@ public:
 	UPROPERTY(config, EditAnywhere, Category = "Create Patch", meta=(ClampMin = 0, ClampMax = 127, UIMin = 0, UIMax = 127))
 	int8 MaxNote = 127;
 
+	// Apply an additional octave, positive or negative, to the note values.
+	// This can be used to adjust the note values parsed from the sample file names
+	// or similarly adjust the octave when setting the notes by scale
+	UPROPERTY(Config, EditAnywhere, Category = "Create Patch")
+	int32 OctaveAdjust = 0;
+
 	UPROPERTY(EditAnywhere, Category = "Create Patch")
 	TArray<FKeyzoneSettings> Keyzones;
 
@@ -142,6 +148,7 @@ public:
 
 	TArray<TWeakObjectPtr<USoundWave>> StagedSoundWaves;
 
+	UPROPERTY()
 	ELockedNoteFlag LockedNotesMask = ELockedNoteFlag::None;
 
 #if WITH_EDITOR
