Improvements to FusionPatch import/reimport dialog
- Only prompt to "reimport existing assets" if the destination is _not_ empty
- Prompt to "Save" newly imported sound waves to make that work flow easier
#tests Local Editor Tests

diff --git a/Engine/Plugins/Runtime/Harmonix/Source/HarmonixDspEditor/HarmonixDspEditor.Build.cs b/Engine/Plugins/Runtime/Harmonix/Source/HarmonixDspEditor/HarmonixDspEditor.Build.cs
--- a/Engine/Plugins/Runtime/Harmonix/Source/HarmonixDspEditor/HarmonixDspEditor.Build.cs
+++ b/Engine/Plugins/Runtime/Harmonix/Source/HarmonixDspEditor/HarmonixDspEditor.Build.cs
@@ -31,6 +31,7 @@ public class HarmonixDspEditor : ModuleRules
 				"HarmonixDsp",
 				"UnrealEd",
 				"Json",
+				"AssetRegistry",
 				"AssetDefinition",
 				"AssetTools",
 				"PropertyEditor",
diff --git a/Engine/Plugins/Runtime/Harmonix/Source/HarmonixDspEditor/Private/FusionPatchAssetFactory.cpp b/Engine/Plugins/Runtime/Harmonix/Source/HarmonixDspEditor/Private/FusionPatchAssetFactory.cpp
--- a/Engine/Plugins/Runtime/Harmonix/Source/HarmonixDspEditor/Private/FusionPatchAssetFactory.cpp
+++ b/Engine/Plugins/Runtime/Harmonix/Source/HarmonixDspEditor/Private/FusionPatchAssetFactory.cpp
@@ -19,6 +19,7 @@
 #include "Misc/Paths.h"
 #include "Misc/PackageName.h"
 #include "EditorFramework/AssetImportData.h"
+#include "FileHelpers.h"
 #include "UObject/Package.h"
 #include "HAL/FileManager.h"
 
@@ -127,12 +128,12 @@ bool UFusionPatchAssetFactory::GetReplaceExistingSamplesResponse(const FString&
 
 	const FText ReplaceExistingTitle = NSLOCTEXT("FusionPatchImporter", "ReplaceExistingSamplesTitle", "Replace Existing Samples");
 	const FText ReplaceExistingMessage = FText::Format(NSLOCTEXT("FusionPatchImporter", "ReplaceExistingSamplesMsg", 
-		"Would you like to reimport and replace existing Sound Wave Assets in the directory with Samples referenced by this Fusion Patch?" 
+		"Existing samples were detected in the selected Sample directory. Would you like to reimport existing Sound Wave Assets?" 
 		"\n\nPatch Name: {0}"
-		"\n\nYes. Reimport and replace existing Samples with the new samples."
-		"\n\nNo. New Samples will be imported, but existing Samples will be unchanged. The Fusion Patch will reference any existing Samples in the directory with matching names."), 
+		"\n\nYes. Reimport existing Samples. *If you made changes to any samples*, you will want to do this."
+		"\n\nNo.  Don't reimport existing Samples."), 
 		FText::FromString(InName));
-	ReplaceExistingSamplesResponse = UEditorDialogLibrary::ShowMessage(ReplaceExistingTitle, ReplaceExistingMessage, EAppMsgType::YesNoYesAllNoAll, ReplaceExistingSamplesResponse, EAppMsgCategory::Info);
+	ReplaceExistingSamplesResponse = UEditorDialogLibrary::ShowMessage(ReplaceExistingTitle, ReplaceExistingMessage, EAppMsgType::YesNo, ReplaceExistingSamplesResponse, EAppMsgCategory::Info);
 	
 
 	switch (ReplaceExistingSamplesResponse)
@@ -180,7 +181,8 @@ UObject* UFusionPatchAssetFactory::FactoryCreateText(UClass* InClass, UObject* I
 	}
 	else
 	{
-		Args.Directory = LongPackagePath;
+		// Default samples directory to subdirectory in current directory: [CurrentDirectory] / [PatchName]
+		Args.Directory = LongPackagePath / InName.ToString();
 	}
 
 	bool WasOkayPressed = false;
@@ -201,7 +203,10 @@ UObject* UFusionPatchAssetFactory::FactoryCreateText(UClass* InClass, UObject* I
 		return nullptr;
 	}
 
-	const bool ReplaceExistingSamples = GetReplaceExistingSamplesResponse(InName.ToString());
+	const FAssetRegistryModule& AssetRegistryModule = FModuleManager::GetModuleChecked<FAssetRegistryModule>("AssetRegistry");
+	bool DestinationHasAssets = AssetRegistryModule.Get().HasAssets(FName(ImportOptions->SamplesImportDir.Path));
+	// if the destination is empty, then we don't need to prompt the user to replace existing samples since there are no existing samples to replace
+	const bool ReplaceExistingSamples = DestinationHasAssets ? GetReplaceExistingSamplesResponse(InName.ToString()) : false;
 	
 	const FString SourceFile = GetCurrentFilename();
 	AdditionalImportedObjects.Empty();
@@ -247,6 +252,13 @@ UObject* UFusionPatchAssetFactory::FactoryCreateText(UClass* InClass, UObject* I
 			// save off the samples dest path for simplifying reimporting
 			FusionPatch->SamplesImportDir = ImportArgs.SamplesDestPath;
 
+			// save imported keyzones "SoundWaves" after successfully creating the fusion patch
+			{
+				TArray<UPackage*> SoundWavePackages;
+				Algo::Transform(FusionPatch->GetKeyzones(), SoundWavePackages, [](const FKeyzoneSettings& Keyzone) { return Keyzone.SoundWave->GetPackage(); });
+				UEditorLoadingAndSavingUtils::SavePackagesWithDialog(SoundWavePackages, true);
+			}
+			
 			bImportSuccessful = true;
 			UpdateFusionPatchImportNotificationItem(ImportNotificationItem, bImportSuccessful, InName);
 			return FusionPatch;
