Ask whether to apply options to all on import.
Collect "AdditonalImportedAssets" during import.
Save and reset import state in "CleanUp() override"

#tests Bulk Imported Fusion Patches In Editor

diff --git a/Engine/Plugins/Runtime/Harmonix/Source/HarmonixDspEditor/Private/FusionPatchAssetFactory.cpp b/Engine/Plugins/Runtime/Harmonix/Source/HarmonixDspEditor/Private/FusionPatchAssetFactory.cpp
--- a/Engine/Plugins/Runtime/Harmonix/Source/HarmonixDspEditor/Private/FusionPatchAssetFactory.cpp
+++ b/Engine/Plugins/Runtime/Harmonix/Source/HarmonixDspEditor/Private/FusionPatchAssetFactory.cpp
@@ -108,24 +108,24 @@ EReimportResult::Type UFusionPatchAssetFactory::Reimport(UObject* Obj)
 	return EReimportResult::Failed;
 }
 
-void UFusionPatchAssetFactory::PostImportCleanUp()
+void UFusionPatchAssetFactory::CleanUp()
 {
-	ApplyOptionsToAllImport = false;
-	ReplaceExistingSamplesResponse = EAppReturnType::No;
+	ImportCounter = 0;
+	ApplyOptionsToAllImport = EApplyAllOption::Unset;
+	ReplaceExistingSamples = false;
+
+	// prompt to save the imported objects collected during entire import
+	TArray<UPackage*> PackagesToSave;
+	Algo::Transform(ImportedObjects, PackagesToSave, [](const UObject* Object) { return Object->GetPackage(); });
+    UEditorLoadingAndSavingUtils::SavePackagesWithDialog(PackagesToSave, true);
+    ImportedObjects.Empty();
+
+	UFactory::CleanUp();
 }
 
+
 bool UFusionPatchAssetFactory::GetReplaceExistingSamplesResponse(const FString& InName)
 {
-	if (ReplaceExistingSamplesResponse == EAppReturnType::YesAll)
-	{
-		return true;
-	}
-
-	if (ReplaceExistingSamplesResponse == EAppReturnType::NoAll)
-	{
-		return false;
-	}
-
 	const FText ReplaceExistingTitle = NSLOCTEXT("FusionPatchImporter", "ReplaceExistingSamplesTitle", "Replace Existing Samples");
 	const FText ReplaceExistingMessage = FText::Format(NSLOCTEXT("FusionPatchImporter", "ReplaceExistingSamplesMsg", 
 		"Existing samples were detected in the selected Sample directory. Would you like to reimport existing Sound Wave Assets?" 
@@ -133,16 +133,14 @@ bool UFusionPatchAssetFactory::GetReplaceExistingSamplesResponse(const FString&
 		"\n\nYes. Reimport existing Samples. *If you made changes to any samples*, you will want to do this."
 		"\n\nNo.  Don't reimport existing Samples."), 
 		FText::FromString(InName));
-	ReplaceExistingSamplesResponse = UEditorDialogLibrary::ShowMessage(ReplaceExistingTitle, ReplaceExistingMessage, EAppMsgType::YesNo, ReplaceExistingSamplesResponse, EAppMsgCategory::Info);
+	EAppReturnType::Type ReplaceExistingSamplesResponse = UEditorDialogLibrary::ShowMessage(ReplaceExistingTitle, ReplaceExistingMessage, EAppMsgType::YesNo, EAppReturnType::No, EAppMsgCategory::Info);
 	
 
 	switch (ReplaceExistingSamplesResponse)
 	{
 	case EAppReturnType::Yes:
-	case EAppReturnType::YesAll:
 		return true;
 	case EAppReturnType::No:
-	case EAppReturnType::NoAll:
 		return false;
 	default:
 		ensureMsgf(false, TEXT("Unexpected response! default behavior is to NOT replace existing samples"));
@@ -150,6 +148,25 @@ bool UFusionPatchAssetFactory::GetReplaceExistingSamplesResponse(const FString&
 	}
 }
 
+bool UFusionPatchAssetFactory::GetApplyOptionsToAllImportResponse()
+{
+	const FText ReplaceExistingTitle = NSLOCTEXT("FusionPatchImporter", "ApplyOptionsToAllTitle", "Apply Options to All");
+	const FText ReplaceExistingMessage = NSLOCTEXT("FusionPatchImporter", "ApplyOptionsToALlMsg", "Would you like to apply the selected options to all files being imported?");
+	EAppReturnType::Type Response = UEditorDialogLibrary::ShowMessage(ReplaceExistingTitle, ReplaceExistingMessage, EAppMsgType::YesNo, EAppReturnType::No, EAppMsgCategory::Info);
+
+	switch (Response)
+	{
+	case EAppReturnType::Yes:
+		return true;
+	case EAppReturnType::No:
+		return false;
+	default:
+		ensureMsgf(false, TEXT("Unexpected response! default behavior is to NOT apply all"));
+		return false;
+	}
+}
+
+
 void UFusionPatchAssetFactory::UpdateFusionPatchImportNotificationItem(TSharedPtr<SNotificationItem> InItem, bool bImportSuccessful, FName InName)
 {
 	if (bImportSuccessful)
@@ -167,49 +184,63 @@ void UFusionPatchAssetFactory::UpdateFusionPatchImportNotificationItem(TSharedPt
 
 UObject* UFusionPatchAssetFactory::FactoryCreateText(UClass* InClass, UObject* InParent, FName InName, EObjectFlags Flags, UObject* Context, const TCHAR* Type, const TCHAR*& Buffer, const TCHAR* BufferEnd, FFeedbackContext* Warn)
 {
+	AdditionalImportedObjects.Empty();
+	
 	// get the existing fusion patch if we're reimporting
 	UFusionPatch* FusionPatch = FindObject<UFusionPatch>(InParent, *InName.ToString());
 	
 	const FString LongPackagePath = FPackageName::GetLongPackagePath(InParent->GetOutermost()->GetPathName());
-	
-	UFusionPatchImportOptions::FArgs Args;
-	Args.PatchName = InName;
-	// If we're reimporting and the fusion patch has saved off the samples directory
-	if (FusionPatch && !FusionPatch->SamplesImportDir.IsEmpty())
+
+	const UFusionPatchImportOptions* ImportOptions = GetDefault<UFusionPatchImportOptions>();
+
+	// detect when we're importing another file so we can ask if we would like to apply the previously set settings to this and all other files
+	// import counter gets reset after all files have been imported
+	++ImportCounter;
+	if (ApplyOptionsToAllImport == EApplyAllOption::Unset && ImportCounter > 1)
 	{
-		Args.Directory = FusionPatch->SamplesImportDir;
+		ApplyOptionsToAllImport = GetApplyOptionsToAllImportResponse() ? EApplyAllOption::Yes : EApplyAllOption::No;
 	}
-	else
+	
+	if (ApplyOptionsToAllImport != EApplyAllOption::Yes)
 	{
-		// Default samples directory to subdirectory in current directory: [CurrentDirectory] / [PatchName]
-		Args.Directory = LongPackagePath / InName.ToString();
-	}
+		UFusionPatchImportOptions::FArgs Args;
+		Args.PatchName = InName;
+		// If we're reimporting and the fusion patch has saved off the samples directory
+		if (FusionPatch && !FusionPatch->SamplesImportDir.IsEmpty())
+		{
+			Args.Directory = FusionPatch->SamplesImportDir;
+		}
+		else
+		{
+			// Default samples directory to subdirectory in current directory: [CurrentDirectory] / [PatchName]
+			Args.Directory = LongPackagePath / InName.ToString();
+		}
 
-	bool WasOkayPressed = false;
-	const UFusionPatchImportOptions* ImportOptions = UFusionPatchImportOptions::GetWithDialog(MoveTemp(Args), WasOkayPressed);
-	if (!WasOkayPressed)
-	{
-		// import cancelled by user
-		return nullptr;
-	}
+		bool WasOkayPressed = false;
+		ImportOptions = UFusionPatchImportOptions::GetWithDialog(MoveTemp(Args), WasOkayPressed);
+		if (!WasOkayPressed)
+		{
+			// import cancelled by user
+			return nullptr;
+		}
 
-	if (!ensure(ImportOptions))
-	{
-		return nullptr;
-	}
+		if (!ensure(ImportOptions))
+		{
+			return nullptr;
+		}
 
-	if (Warn->ReceivedUserCancel())
-	{
-		return nullptr;
-	}
+		if (Warn->ReceivedUserCancel())
+		{
+			return nullptr;
+		}
 
-	const FAssetRegistryModule& AssetRegistryModule = FModuleManager::GetModuleChecked<FAssetRegistryModule>("AssetRegistry");
-	bool DestinationHasAssets = AssetRegistryModule.Get().HasAssets(FName(ImportOptions->SamplesImportDir.Path));
-	// if the destination is empty, then we don't need to prompt the user to replace existing samples since there are no existing samples to replace
-	const bool ReplaceExistingSamples = DestinationHasAssets ? GetReplaceExistingSamplesResponse(InName.ToString()) : false;
+		const FAssetRegistryModule& AssetRegistryModule = FModuleManager::GetModuleChecked<FAssetRegistryModule>("AssetRegistry");
+		bool DestinationHasAssets = AssetRegistryModule.Get().HasAssets(FName(ImportOptions->SamplesImportDir.Path));
+		// if the destination is empty, then we don't need to prompt the user to replace existing samples since there are no existing samples to replace
+		ReplaceExistingSamples = DestinationHasAssets ? GetReplaceExistingSamplesResponse(InName.ToString()) : false;
+	}
 	
 	const FString SourceFile = GetCurrentFilename();
-	AdditionalImportedObjects.Empty();
 	FString JsonString;
 	const FString DtaString(BufferEnd - Buffer, Buffer);
 
@@ -240,7 +271,7 @@ UObject* UFusionPatchAssetFactory::FactoryCreateText(UClass* InClass, UObject* I
 		ImportArgs.SampleCompressionType = ImportOptions->SampleCompressionType;
 
 		TArray<FString> ImportErrors;
-		if (FFusionPatchJsonImporter::TryParseJson(JsonObj, FusionPatch, ImportArgs, ImportErrors))
+		if (FFusionPatchJsonImporter::TryParseJson(JsonObj, FusionPatch, AdditionalImportedObjects, ImportArgs, ImportErrors))
 		{
 			UE_LOG(LogFusionPatchAssetFactory, Log, TEXT("Successfully imported FusionPatch asset"));
 
@@ -251,16 +282,11 @@ UObject* UFusionPatchAssetFactory::FactoryCreateText(UClass* InClass, UObject* I
 
 			// save off the samples dest path for simplifying reimporting
 			FusionPatch->SamplesImportDir = ImportArgs.SamplesDestPath;
-
-			// save imported keyzones "SoundWaves" after successfully creating the fusion patch
-			{
-				TArray<UPackage*> SoundWavePackages;
-				Algo::Transform(FusionPatch->GetKeyzones(), SoundWavePackages, [](const FKeyzoneSettings& Keyzone) { return Keyzone.SoundWave->GetPackage(); });
-				UEditorLoadingAndSavingUtils::SavePackagesWithDialog(SoundWavePackages, true);
-			}
 			
 			bImportSuccessful = true;
 			UpdateFusionPatchImportNotificationItem(ImportNotificationItem, bImportSuccessful, InName);
+			ImportedObjects.Append(AdditionalImportedObjects);
+			ImportedObjects.Add(FusionPatch);
 			return FusionPatch;
 		}
 		else
diff --git a/Engine/Plugins/Runtime/Harmonix/Source/HarmonixDspEditor/Private/FusionPatchAssetFactory.h b/Engine/Plugins/Runtime/Harmonix/Source/HarmonixDspEditor/Private/FusionPatchAssetFactory.h
--- a/Engine/Plugins/Runtime/Harmonix/Source/HarmonixDspEditor/Private/FusionPatchAssetFactory.h
+++ b/Engine/Plugins/Runtime/Harmonix/Source/HarmonixDspEditor/Private/FusionPatchAssetFactory.h
@@ -27,7 +27,7 @@ public:
 	virtual bool CanReimport(UObject* Obj, TArray<FString>& OutFilenames) override;
 	virtual void SetReimportPaths(UObject* Obj, const TArray<FString>& NewReimportPaths) override;;
 	virtual EReimportResult::Type Reimport(UObject* Obj) override;
-	virtual void PostImportCleanUp() override;
+	virtual void CleanUp() override;
 
 	virtual UObject* FactoryCreateText(UClass* InClass, UObject* InParent, FName InName, EObjectFlags Flags, UObject* Context, const TCHAR* Type, const TCHAR*& Buffer, const TCHAR* BufferEnd, FFeedbackContext* Warn) override;
 	virtual UObject* FactoryCreateNew(UClass* InClass, UObject* InParent, FName InName, EObjectFlags Flags, UObject* Context, FFeedbackContext* Warn) override;
@@ -35,11 +35,20 @@ public:
 	const UFusionPatchCreateOptions* CreateOptions = nullptr;
 private:
 
-	bool GetReplaceExistingSamplesResponse(const FString& InName);
-	
-	bool ApplyOptionsToAllImport = false;
+	static bool GetReplaceExistingSamplesResponse(const FString& InName);
+	static bool GetApplyOptionsToAllImportResponse();
 
-	EAppReturnType::Type ReplaceExistingSamplesResponse = EAppReturnType::No;
+	enum class EApplyAllOption : uint8
+	{
+		Unset,
+		No,
+		Yes
+	};
+
+	int32 ImportCounter = 0;
+	EApplyAllOption ApplyOptionsToAllImport = EApplyAllOption::Unset;
+	bool ReplaceExistingSamples = false;
+	TArray<UObject*> ImportedObjects;
 
 	void UpdateFusionPatchImportNotificationItem(TSharedPtr<SNotificationItem> InItem, bool bImportSuccessful, FName InName);
 };
\ No newline at end of file
diff --git a/Engine/Plugins/Runtime/Harmonix/Source/HarmonixDspEditor/Private/FusionPatchJsonImporter.cpp b/Engine/Plugins/Runtime/Harmonix/Source/HarmonixDspEditor/Private/FusionPatchJsonImporter.cpp
--- a/Engine/Plugins/Runtime/Harmonix/Source/HarmonixDspEditor/Private/FusionPatchJsonImporter.cpp
+++ b/Engine/Plugins/Runtime/Harmonix/Source/HarmonixDspEditor/Private/FusionPatchJsonImporter.cpp
@@ -31,7 +31,7 @@ TArray<UObject*> FFusionPatchJsonImporter::ImportAudioSamples(const TArray<FStri
 	return AssetToolsModule.Get().ImportAssetsAutomated(AutomatedAssetImportData);
 }
 
-bool FFusionPatchJsonImporter::TryParseJson(TSharedPtr<FJsonObject> JsonObj, UFusionPatch* FusionPatch, const FImportArgs& ImportArgs, TArray<FString>& OutErrors)
+bool FFusionPatchJsonImporter::TryParseJson(TSharedPtr<FJsonObject> JsonObj, UFusionPatch* FusionPatch, TArray<UObject*>& OutAdditionalImportedObjects, const FImportArgs& ImportArgs, TArray<FString>& OutErrors)
 {
 	if (!ensure(FusionPatch))
 		return false;
@@ -154,6 +154,7 @@ bool FFusionPatchJsonImporter::TryParseJson(TSharedPtr<FJsonObject> JsonObj, UFu
 	}
 
 	TArray<UObject*> ImportedAssets = ImportAudioSamples(AudioSampleFiles, ImportArgs.SamplesDestPath, ImportArgs.ReplaceExistingSamples);
+	OutAdditionalImportedObjects.Append(ImportedAssets);
 
 	// update the sound wave loading behavior and compression type with the selection set in the import options
 	for (UObject* Asset : ImportedAssets)
diff --git a/Engine/Plugins/Runtime/Harmonix/Source/HarmonixDspEditor/Private/FusionPatchJsonImporter.h b/Engine/Plugins/Runtime/Harmonix/Source/HarmonixDspEditor/Private/FusionPatchJsonImporter.h
--- a/Engine/Plugins/Runtime/Harmonix/Source/HarmonixDspEditor/Private/FusionPatchJsonImporter.h
+++ b/Engine/Plugins/Runtime/Harmonix/Source/HarmonixDspEditor/Private/FusionPatchJsonImporter.h
@@ -54,5 +54,5 @@ public:
 
 	static TArray<UObject*> ImportAudioSamples(const TArray<FString>& Files, const FString& DestinationPath, bool ReplaceExisting);
 
-	static bool TryParseJson(TSharedPtr<FJsonObject> JsonObj, UFusionPatch* FusionPatch, const FImportArgs& ImportArgs, TArray<FString>& OutErrors);
+	static bool TryParseJson(TSharedPtr<FJsonObject> JsonObj, UFusionPatch* FusionPatch, TArray<UObject*>& OutAdditionalImportedObjects, const FImportArgs& ImportArgs, TArray<FString>& OutErrors);
 };
\ No newline at end of file
