[Harmonix ] Live Update support for TransportWavePlayerControllerNode

diff --git a/Engine/Plugins/Runtime/Harmonix/Source/HarmonixMetasound/Private/HarmonixMetasound/Nodes/TransportWavePlayerControllerNode.cpp b/Engine/Plugins/Runtime/Harmonix/Source/HarmonixMetasound/Private/HarmonixMetasound/Nodes/TransportWavePlayerControllerNode.cpp
--- a/Engine/Plugins/Runtime/Harmonix/Source/HarmonixMetasound/Private/HarmonixMetasound/Nodes/TransportWavePlayerControllerNode.cpp
+++ b/Engine/Plugins/Runtime/Harmonix/Source/HarmonixMetasound/Private/HarmonixMetasound/Nodes/TransportWavePlayerControllerNode.cpp
@@ -26,16 +26,16 @@ namespace HarmonixMetasound
 	{
 	public:
 		FTransportWavePlayerControllerOperator(const FOperatorSettings& InSettings,
-									const FMusicTransportEventStreamReadRef& InTransport,
-									const FMidiClockReadRef& InMidiClock) :
-			FMusicTransportControllable(EMusicPlayerTransportState::Prepared),
-			TransportInPin(InTransport),
-			MidiClockInPin(InMidiClock),
-			PlayOutPin(FTriggerWriteRef::CreateNew(InSettings)),
-			StopOutPin(FTriggerWriteRef::CreateNew(InSettings)),
-			StartTimeOutPin(FTimeWriteRef::CreateNew()),
-			BlockSizeFrames(InSettings.GetNumFramesPerBlock()),
-			SampleRate(InSettings.GetSampleRate())
+											   const FMusicTransportEventStreamReadRef& InTransport,
+											   const FMidiClockReadRef& InMidiClock)
+			: FMusicTransportControllable(EMusicPlayerTransportState::Prepared)
+			, TransportInPin(InTransport)
+			, MidiClockInPin(InMidiClock)
+			, PlayOutPin(FTriggerWriteRef::CreateNew(InSettings))
+			, StopOutPin(FTriggerWriteRef::CreateNew(InSettings))
+			, StartTimeOutPin(FTimeWriteRef::CreateNew())
+			, BlockSizeFrames(InSettings.GetNumFramesPerBlock())
+			, SampleRate(InSettings.GetSampleRate())
 		{
 		}
 
@@ -87,6 +87,8 @@ namespace HarmonixMetasound
 
 			InVertexData.BindReadVertex(METASOUND_GET_PARAM_NAME(Inputs::Transport), TransportInPin);
 			InVertexData.BindReadVertex(METASOUND_GET_PARAM_NAME(Inputs::MidiClock), MidiClockInPin);
+
+			bNeedsTransportInit = true;
 		}
 
 		virtual void BindOutputs(FOutputVertexInterfaceData& InVertexData) override
@@ -97,8 +99,10 @@ namespace HarmonixMetasound
 			InVertexData.BindReadVertex(METASOUND_GET_PARAM_NAME(Outputs::TransportPlay), PlayOutPin);
 			InVertexData.BindReadVertex(METASOUND_GET_PARAM_NAME(Outputs::TransportStop), StopOutPin);
 			InVertexData.BindReadVertex(METASOUND_GET_PARAM_NAME(OutputStartTime), StartTimeOutPin);
+
+			bNeedsTransportInit = true;
 		}
-		
+
 		virtual FDataReferenceCollection GetInputs() const override
 		{
 			// This should never be called. Bind(...) is called instead. This method
@@ -122,7 +126,7 @@ namespace HarmonixMetasound
 			const FInputVertexInterfaceData& InputData = InParams.InputData;
 			FMusicTransportEventStreamReadRef InTransport = InputData.GetOrConstructDataReadReference<FMusicTransportEventStream>(METASOUND_GET_PARAM_NAME(Inputs::Transport), InParams.OperatorSettings);
 			FMidiClockReadRef InMidiClock = InputData.GetOrConstructDataReadReference<FMidiClock>(METASOUND_GET_PARAM_NAME(Inputs::MidiClock), InParams.OperatorSettings);
-			
+
 			return MakeUnique<FTransportWavePlayerControllerOperator>(InParams.OperatorSettings, InTransport, InMidiClock);
 		}
 
@@ -132,81 +136,83 @@ namespace HarmonixMetasound
 			PlayOutPin->AdvanceBlock();
 			StopOutPin->AdvanceBlock();
 
+			InitTransportIfNeeded();
+
 			TransportSpanProcessor TransportHandler = [this](int32 StartFrameIndex, int32 EndFrameIndex, EMusicPlayerTransportState CurrentState)
-			{
-				switch (CurrentState)
 				{
-				case EMusicPlayerTransportState::Invalid:
-				case EMusicPlayerTransportState::Preparing:
-					return EMusicPlayerTransportState::Prepared;
-
-				case EMusicPlayerTransportState::Prepared:
-					return EMusicPlayerTransportState::Prepared;
-
-				case EMusicPlayerTransportState::Starting:
-					if (!ReceivedSeekWhileStopped())
+					switch (CurrentState)
 					{
-						// Play from the beginning if we haven't received a seek call while we were stopped...
-						*StartTimeOutPin = FTime();
-					}
-					PlayOutPin->TriggerFrame(StartFrameIndex);
-					bPlaying = true;
-					return EMusicPlayerTransportState::Playing;
-
-				case EMusicPlayerTransportState::Playing:
-					return EMusicPlayerTransportState::Playing;
+					case EMusicPlayerTransportState::Invalid:
+					case EMusicPlayerTransportState::Preparing:
+						return EMusicPlayerTransportState::Prepared;
+
+					case EMusicPlayerTransportState::Prepared:
+						return EMusicPlayerTransportState::Prepared;
+
+					case EMusicPlayerTransportState::Starting:
+						if (!ReceivedSeekWhileStopped())
+						{
+							// Play from the beginning if we haven't received a seek call while we were stopped...
+							*StartTimeOutPin = FTime();
+						}
+						PlayOutPin->TriggerFrame(StartFrameIndex);
+						bPlaying = true;
+						return EMusicPlayerTransportState::Playing;
+
+					case EMusicPlayerTransportState::Playing:
+						return EMusicPlayerTransportState::Playing;
+
+					case EMusicPlayerTransportState::Seeking:
+						if (ReceivedSeekWhileStopped())
+						{
+							// Assumes the MidiClock is stopped for the remainder of the block.
+							*StartTimeOutPin = FTime(MidiClockInPin->GetCurrentHiResMs() * 0.001f);
+						}
+						else
+						{
+							StopOutPin->TriggerFrame(StartFrameIndex);
+							int32 PlayFrameIndex = FMath::Min(StartFrameIndex + 1, EndFrameIndex);
+
+							// Assumes the MidiClock is playing for the remainder of the block.
+							*StartTimeOutPin = FTime(MidiClockInPin->GetCurrentHiResMs() * 0.001f - (BlockSizeFrames - PlayFrameIndex) / SampleRate);
+							PlayOutPin->TriggerFrame(PlayFrameIndex);
+						}
+						// Here we will return that we want to be in the same state we were in before this request to 
+						// seek since we can seek "instantaneously"...
+						return GetTransportState();
+
+					case EMusicPlayerTransportState::Continuing:
+						// Assumes the StartTimeOutPin won't change for the remainder of the block.
+						PlayOutPin->TriggerFrame(StartFrameIndex);
+						bPlaying = true;
+						return EMusicPlayerTransportState::Playing;
+
+					case EMusicPlayerTransportState::Pausing:
+						bPlaying = false;
+						StopOutPin->TriggerFrame(StartFrameIndex);
 
-				case EMusicPlayerTransportState::Seeking:
-					if (ReceivedSeekWhileStopped())
-					{
-						// Assumes the MidiClock is stopped for the remainder of the block.
+						// Assumes the MidiClock is paused for the remainder of the block.
 						*StartTimeOutPin = FTime(MidiClockInPin->GetCurrentHiResMs() * 0.001f);
-					}
-					else
-					{
-						StopOutPin->TriggerFrame(StartFrameIndex);
-						int32 PlayFrameIndex = FMath::Min(StartFrameIndex + 1, EndFrameIndex);
+						return EMusicPlayerTransportState::Paused;
+
+					case EMusicPlayerTransportState::Paused:
+						return EMusicPlayerTransportState::Paused;
+
+					case EMusicPlayerTransportState::Stopping:
+					case EMusicPlayerTransportState::Killing:
+						if (bPlaying)
+						{
+							bPlaying = false;
+							StopOutPin->TriggerFrame(StartFrameIndex);
+						}
+						*StartTimeOutPin = FTime();
+						return EMusicPlayerTransportState::Prepared;
 
-						// Assumes the MidiClock is playing for the remainder of the block.
-						*StartTimeOutPin = FTime(MidiClockInPin->GetCurrentHiResMs() * 0.001f - (BlockSizeFrames - PlayFrameIndex) / SampleRate);
-						PlayOutPin->TriggerFrame(PlayFrameIndex);
-					}
-					// Here we will return that we want to be in the same state we were in before this request to 
-					// seek since we can seek "instantaneously"...
-					return GetTransportState();
-
-				case EMusicPlayerTransportState::Continuing:
-					// Assumes the StartTimeOutPin won't change for the remainder of the block.
-					PlayOutPin->TriggerFrame(StartFrameIndex);
-					bPlaying = true;
-					return EMusicPlayerTransportState::Playing;
-
-				case EMusicPlayerTransportState::Pausing:
-					bPlaying = false;
-					StopOutPin->TriggerFrame(StartFrameIndex);
-
-					// Assumes the MidiClock is paused for the remainder of the block.
-					*StartTimeOutPin = FTime(MidiClockInPin->GetCurrentHiResMs() * 0.001f);
-					return EMusicPlayerTransportState::Paused;
-
-				case EMusicPlayerTransportState::Paused:
-					return EMusicPlayerTransportState::Paused;
-
-				case EMusicPlayerTransportState::Stopping:
-				case EMusicPlayerTransportState::Killing:
-					if (bPlaying)
-					{
-						bPlaying = false;
-						StopOutPin->TriggerFrame(StartFrameIndex);
+					default:
+						checkNoEntry();
+						return EMusicPlayerTransportState::Invalid;
 					}
-					*StartTimeOutPin = FTime();
-					return EMusicPlayerTransportState::Prepared;
-
-				default:
-					checkNoEntry();
-					return EMusicPlayerTransportState::Invalid;
-				}
-			};
+				};
 			ExecuteTransportSpans(TransportInPin, BlockSizeFrames, TransportHandler);
 		}
 
@@ -220,6 +226,8 @@ namespace HarmonixMetasound
 			SampleRate = InParams.OperatorSettings.GetSampleRate();
 
 			bPlaying = false;
+
+			bNeedsTransportInit = true;
 		}
 
 	private:
@@ -237,6 +245,53 @@ namespace HarmonixMetasound
 		float SampleRate;
 
 		bool bPlaying = false;
+		bool bNeedsTransportInit = true;
+
+		void InitTransportIfNeeded()
+		{
+			if (bNeedsTransportInit)
+			{
+				// Get the node caught up to its transport input
+				FTransportInitFn InitFn = [this](EMusicPlayerTransportState CurrentState)
+					{
+						switch (CurrentState)
+						{
+						case EMusicPlayerTransportState::Invalid:
+						case EMusicPlayerTransportState::Preparing:
+						case EMusicPlayerTransportState::Prepared:
+						case EMusicPlayerTransportState::Stopping:
+						case EMusicPlayerTransportState::Killing:
+							StopOutPin->TriggerFrame(0);
+							return EMusicPlayerTransportState::Prepared;
+
+						case EMusicPlayerTransportState::Starting:
+						case EMusicPlayerTransportState::Playing:
+						case EMusicPlayerTransportState::Continuing:
+							// Catch up with our MidiClock
+							*StartTimeOutPin = FTime(MidiClockInPin->GetCurrentHiResMs() * 0.001f);
+							PlayOutPin->TriggerFrame(0);
+							bPlaying = true;
+							return EMusicPlayerTransportState::Playing;
+
+						case EMusicPlayerTransportState::Seeking: // seeking is omitted from init, shouldn't happen
+							checkNoEntry();
+							return EMusicPlayerTransportState::Invalid;
+
+						case EMusicPlayerTransportState::Pausing:
+						case EMusicPlayerTransportState::Paused:
+							StopOutPin->TriggerFrame(0);
+							return EMusicPlayerTransportState::Paused;
+
+						default:
+							checkNoEntry();
+							return EMusicPlayerTransportState::Invalid;
+						}
+					};
+				Init(*TransportInPin, MoveTemp(InitFn));
+
+				bNeedsTransportInit = false;
+			}
+		}
 	};
 
 	class FTransportWavePlayerControllerNode : public FNodeFacade
