Add hidden friend GetHashType hash functions for MetaSound types for Array Random Get node refactor

diff --git a/Engine/Plugins/Runtime/AudioModulation/Source/AudioModulation/Private/SoundModulationParameter.cpp b/Engine/Plugins/Runtime/AudioModulation/Source/AudioModulation/Private/SoundModulationParameter.cpp
--- a/Engine/Plugins/Runtime/AudioModulation/Source/AudioModulation/Private/SoundModulationParameter.cpp
+++ b/Engine/Plugins/Runtime/AudioModulation/Source/AudioModulation/Private/SoundModulationParameter.cpp
@@ -63,6 +63,7 @@ Audio::FModulationParameter USoundModulationParameter::CreateParameter() const
 	Parameter.DefaultValue = GetUnitDefault();
 	Parameter.MinValue = GetUnitMin();
 	Parameter.MaxValue = GetUnitMax();
+	Parameter.TypeHash = GetTypeHash(this);
 
 #if WITH_EDITORONLY_DATA
 	Parameter.UnitDisplayName = Settings.UnitDisplayName;
diff --git a/Engine/Plugins/Runtime/AudioModulation/Source/AudioModulation/Public/SoundModulatorAsset.h b/Engine/Plugins/Runtime/AudioModulation/Source/AudioModulation/Public/SoundModulatorAsset.h
--- a/Engine/Plugins/Runtime/AudioModulation/Source/AudioModulation/Public/SoundModulatorAsset.h
+++ b/Engine/Plugins/Runtime/AudioModulation/Source/AudioModulation/Public/SoundModulatorAsset.h
@@ -69,6 +69,11 @@ namespace AudioModulation
 		{
 			return Proxy.Get();
 		}
+				
+		friend FORCEINLINE uint32 GetTypeHash(const AudioModulation::FSoundModulatorAsset& InModulatorAsset)
+		{
+			return InModulatorAsset->GetModulatorId();
+		}
 	};
 
 	class AUDIOMODULATION_API FSoundModulationParameterAsset
@@ -103,6 +108,11 @@ namespace AudioModulation
 		{
 			return Proxy.Get();
 		}
+
+		friend FORCEINLINE uint32 GetTypeHash(const AudioModulation::FSoundModulationParameterAsset& InModulationParameterAsset)
+		{
+			return GetTypeHash(InModulationParameterAsset->GetParameter());
+		}
 	};
 } // namespace AudioModulation
 
diff --git a/Engine/Plugins/Runtime/Harmonix/Source/HarmonixDsp/Public/HarmonixDsp/AudioAnalysis/FFTAnalyzer.h b/Engine/Plugins/Runtime/Harmonix/Source/HarmonixDsp/Public/HarmonixDsp/AudioAnalysis/FFTAnalyzer.h
--- a/Engine/Plugins/Runtime/Harmonix/Source/HarmonixDsp/Public/HarmonixDsp/AudioAnalysis/FFTAnalyzer.h
+++ b/Engine/Plugins/Runtime/Harmonix/Source/HarmonixDsp/Public/HarmonixDsp/AudioAnalysis/FFTAnalyzer.h
@@ -44,6 +44,11 @@ struct HARMONIXDSP_API FHarmonixFFTAnalyzerResults
 	
 	UPROPERTY(BlueprintReadOnly, Category="Results")
 	TArray<float> Spectrum;
+	
+	friend FORCEINLINE uint32 GetTypeHash(const FHarmonixFFTAnalyzerResults& InAnalyzerResults)
+	{
+		return GetTypeHash(InAnalyzerResults.Spectrum);
+	}
 };
 
 namespace Harmonix::Dsp::AudioAnalysis
diff --git a/Engine/Plugins/Runtime/Harmonix/Source/HarmonixMetasound/Public/HarmonixMetasound/DataTypes/MidiEventInfo.h b/Engine/Plugins/Runtime/Harmonix/Source/HarmonixMetasound/Public/HarmonixMetasound/DataTypes/MidiEventInfo.h
--- a/Engine/Plugins/Runtime/Harmonix/Source/HarmonixMetasound/Public/HarmonixMetasound/DataTypes/MidiEventInfo.h
+++ b/Engine/Plugins/Runtime/Harmonix/Source/HarmonixMetasound/Public/HarmonixMetasound/DataTypes/MidiEventInfo.h
@@ -29,6 +29,13 @@ struct HARMONIXMETASOUND_API FMidiEventInfo
 	bool IsNoteOff() const;
 	uint8 GetNoteNumber() const;
 	uint8 GetVelocity() const;
+
+	friend FORCEINLINE uint32 GetTypeHash(const FMidiEventInfo& InMidiEventInfo)
+	{
+		return HashCombineFast(
+			HashCombineFast(GetTypeHash(InMidiEventInfo.Timestamp), GetTypeHash(InMidiEventInfo.TrackIndex)),
+			GetTypeHash(InMidiEventInfo.MidiMessage));
+	}
 };
 
 UCLASS()
diff --git a/Engine/Plugins/Runtime/Harmonix/Source/HarmonixMetasound/Public/HarmonixMetasound/DataTypes/MusicLoopConfiguration.h b/Engine/Plugins/Runtime/Harmonix/Source/HarmonixMetasound/Public/HarmonixMetasound/DataTypes/MusicLoopConfiguration.h
--- a/Engine/Plugins/Runtime/Harmonix/Source/HarmonixMetasound/Public/HarmonixMetasound/DataTypes/MusicLoopConfiguration.h
+++ b/Engine/Plugins/Runtime/Harmonix/Source/HarmonixMetasound/Public/HarmonixMetasound/DataTypes/MusicLoopConfiguration.h
@@ -35,6 +35,11 @@ public:
 		, RegionStart(InRegionStart)
 		, RegionEnd(InRegionEnd)
 	{}
+	
+	friend FORCEINLINE uint32 GetTypeHash(const FMusicLoopConfiguration& InConfiguration)
+	{
+		return HashCombineFast(HashCombineFast(GetTypeHash(InConfiguration.RegionStart), GetTypeHash(InConfiguration.RegionEnd)), InConfiguration.Enabled);
+	}
 };
 
 // NOTE: Since there is no corresponding cpp file, the corresponding REGISTER_METASOUND_DATATYPE is in MidiClock.cpp
diff --git a/Engine/Plugins/Runtime/Harmonix/Source/HarmonixMetasound/Public/HarmonixMetasound/DataTypes/MusicSeekRequest.h b/Engine/Plugins/Runtime/Harmonix/Source/HarmonixMetasound/Public/HarmonixMetasound/DataTypes/MusicSeekRequest.h
--- a/Engine/Plugins/Runtime/Harmonix/Source/HarmonixMetasound/Public/HarmonixMetasound/DataTypes/MusicSeekRequest.h
+++ b/Engine/Plugins/Runtime/Harmonix/Source/HarmonixMetasound/Public/HarmonixMetasound/DataTypes/MusicSeekRequest.h
@@ -30,6 +30,11 @@ public:
 
 	UPROPERTY(BlueprintReadWrite, Category = "MusicTransport", meta = (EditCondition = "Type == ESeekPointType::Millisecond", EditConditionHides))
 	float Ms = 0.0f;
+	
+	friend FORCEINLINE uint32 GetTypeHash(const FMusicSeekTarget& InSeekTarget)
+	{
+		return HashCombineFast(HashCombineFast(GetTypeHash(InSeekTarget.Ms), GetTypeHash(InSeekTarget.BarBeat)), GetTypeHash(InSeekTarget.Type));
+	}
 };
 
 USTRUCT(BlueprintType)
@@ -62,6 +67,11 @@ public:
 		EmmediateIfPastFromPoint = false;
 		FromPoint  = ToPoint  = FMusicTimestamp();
 	}
+	
+	friend FORCEINLINE uint32 GetTypeHash(const FMusicSeekRequest& InSeekRequest)
+	{
+		return HashCombineFast(HashCombineFast(GetTypeHash(InSeekRequest.FromPoint), GetTypeHash(InSeekRequest.ToPoint)), InSeekRequest.EmmediateIfPastFromPoint);
+	}
 };
 
 // NOTE: Since there is no corresponding cpp file, the corresponding REGISTER_METASOUND_DATATYPE is in MidiClock.cpp
diff --git a/Engine/Plugins/Runtime/Harmonix/Source/HarmonixMidi/Public/HarmonixMidi/BarMap.h b/Engine/Plugins/Runtime/Harmonix/Source/HarmonixMidi/Public/HarmonixMidi/BarMap.h
--- a/Engine/Plugins/Runtime/Harmonix/Source/HarmonixMidi/Public/HarmonixMidi/BarMap.h
+++ b/Engine/Plugins/Runtime/Harmonix/Source/HarmonixMidi/Public/HarmonixMidi/BarMap.h
@@ -58,6 +58,11 @@ public:
 		return Bar < Other.Bar ||
 			(Bar == Other.Bar && Beat < Other.Beat);
 	}
+	
+	friend FORCEINLINE uint32 GetTypeHash(const FMusicTimestamp& InTimestamp)
+	{
+		return HashCombineFast(GetTypeHash(InTimestamp.Bar), GetTypeHash(InTimestamp.Beat));
+	}
 };
 
 /** A simple container for a musical time signature (numerator and denominator.) */
@@ -90,6 +95,11 @@ public:
 	int16 Numerator;
 	UPROPERTY()
 	int16 Denominator;
+
+	friend FORCEINLINE uint32 GetTypeHash(const FTimeSignature& InTimeSignature)
+	{
+		return HashCombineFast(GetTypeHash(InTimeSignature.Numerator), GetTypeHash(InTimeSignature.Denominator));
+	}
 };
 
 /** A time signature associated with a specific tick in a midi file. */
diff --git a/Engine/Plugins/Runtime/Harmonix/Source/HarmonixMidi/Public/HarmonixMidi/MidiMsg.h b/Engine/Plugins/Runtime/Harmonix/Source/HarmonixMidi/Public/HarmonixMidi/MidiMsg.h
--- a/Engine/Plugins/Runtime/Harmonix/Source/HarmonixMidi/Public/HarmonixMidi/MidiMsg.h
+++ b/Engine/Plugins/Runtime/Harmonix/Source/HarmonixMidi/Public/HarmonixMidi/MidiMsg.h
@@ -123,6 +123,13 @@ public:
 		return Type == Other.Type && Status == Other.Status && Data1 == Other.Data1 && Data2 == Other.Data2;
 	}
 
+	friend FORCEINLINE uint32 GetTypeHash(const FMidiMsg& InMidiMsg)
+	{
+		return HashCombineFast(
+			HashCombineFast(GetTypeHash(InMidiMsg.Type), GetTypeHash(InMidiMsg.Status)),
+			HashCombineFast(GetTypeHash(InMidiMsg.Data1), GetTypeHash(InMidiMsg.Data2)));
+	}
+
 	EType Type;
 	// Use a union instead of subclassing MidiMsg so that MidiMsg is 
 	// constant size and therefore you can have a vector<MidiMsg>.
diff --git a/Engine/Plugins/Runtime/Metasound/Source/MetasoundEngine/Public/MetasoundAudioBus.h b/Engine/Plugins/Runtime/Metasound/Source/MetasoundEngine/Public/MetasoundAudioBus.h
--- a/Engine/Plugins/Runtime/Metasound/Source/MetasoundEngine/Public/MetasoundAudioBus.h
+++ b/Engine/Plugins/Runtime/Metasound/Source/MetasoundEngine/Public/MetasoundAudioBus.h
@@ -29,7 +29,11 @@ namespace Metasound
 		{
 			return AudioBusProxy;
 		}
-		
+
+		friend FORCEINLINE uint32 GetTypeHash(const Metasound::FAudioBusAsset& InAudioBusAsset)
+		{
+			return GetTypeHash(*InAudioBusAsset.GetAudioBusProxy());
+		}
 	};
 
 	DECLARE_METASOUND_DATA_REFERENCE_TYPES(FAudioBusAsset, METASOUNDENGINE_API, FAudioBusAssetTypeInfo, FAudioBusAssetReadRef, FAudioBusAssetWriteRef)
diff --git a/Engine/Plugins/Runtime/Metasound/Source/MetasoundEngine/Public/MetasoundWave.h b/Engine/Plugins/Runtime/Metasound/Source/MetasoundEngine/Public/MetasoundWave.h
--- a/Engine/Plugins/Runtime/Metasound/Source/MetasoundEngine/Public/MetasoundWave.h
+++ b/Engine/Plugins/Runtime/Metasound/Source/MetasoundEngine/Public/MetasoundWave.h
@@ -50,6 +50,11 @@ namespace Metasound
 		{
 			return SoundWaveProxy.Get();
 		}
+
+		friend FORCEINLINE uint32 GetTypeHash(const Metasound::FWaveAsset& InWaveAsset)
+		{
+			return GetTypeHash(*InWaveAsset.GetSoundWaveProxy());
+		}
 	};
 
 	DECLARE_METASOUND_DATA_REFERENCE_TYPES(FWaveAsset, METASOUNDENGINE_API, FWaveAssetTypeInfo, FWaveAssetReadRef, FWaveAssetWriteRef)
diff --git a/Engine/Plugins/Runtime/Metasound/Source/MetasoundEngine/Public/MetasoundWaveTable.h b/Engine/Plugins/Runtime/Metasound/Source/MetasoundEngine/Public/MetasoundWaveTable.h
--- a/Engine/Plugins/Runtime/Metasound/Source/MetasoundEngine/Public/MetasoundWaveTable.h
+++ b/Engine/Plugins/Runtime/Metasound/Source/MetasoundEngine/Public/MetasoundWaveTable.h
@@ -76,6 +76,11 @@ namespace Metasound
 		{
 			return Proxy.Get();
 		}
+		
+		friend FORCEINLINE uint32 GetTypeHash(const Metasound::FWaveTableBankAsset& InWaveTableBankAsset)
+		{
+			return GetTypeHash(*InWaveTableBankAsset.GetProxy());
+		}
 	};
 
 	// Disable arrays of WaveTables (as of addition, not supported by UX). Must be defined prior to DataType declaration macro below.
diff --git a/Engine/Plugins/Runtime/Metasound/Source/MetasoundFrontend/Public/MetasoundParameterPack.h b/Engine/Plugins/Runtime/Metasound/Source/MetasoundFrontend/Public/MetasoundParameterPack.h
--- a/Engine/Plugins/Runtime/Metasound/Source/MetasoundFrontend/Public/MetasoundParameterPack.h
+++ b/Engine/Plugins/Runtime/Metasound/Source/MetasoundFrontend/Public/MetasoundParameterPack.h
@@ -506,6 +506,11 @@ public:
 	{
 		return ParameterStoragePtr.Get();
 	}
+
+	friend FORCEINLINE uint32 GetTypeHash(const FMetasoundParameterStorageWrapper& InParameterStorageWrapper)
+	{
+		return GetTypeHash(InParameterStorageWrapper->Storage);
+	}
 };
 
 DECLARE_METASOUND_DATA_REFERENCE_TYPES(FMetasoundParameterStorageWrapper, METASOUNDFRONTEND_API, FMetasoundParameterStorageWrapperTypeInfo, FMetasoundParameterStorageWrapperReadRef, FMetasoundParameterStorageWrapperWriteRef)
diff --git a/Engine/Plugins/Runtime/Metasound/Source/MetasoundFrontend/Public/MetasoundTime.h b/Engine/Plugins/Runtime/Metasound/Source/MetasoundFrontend/Public/MetasoundTime.h
--- a/Engine/Plugins/Runtime/Metasound/Source/MetasoundFrontend/Public/MetasoundTime.h
+++ b/Engine/Plugins/Runtime/Metasound/Source/MetasoundFrontend/Public/MetasoundTime.h
@@ -358,7 +358,12 @@ namespace Metasound
 				static_assert(TIsArithmetic<ArithmeticType>::Value, "Must be arithmetic type.");
 				return FTime(static_cast<TimeType>(InMicroseconds * 1e-6));
 			}
-
+			
+			friend FORCEINLINE uint32 GetTypeHash(const FTime& InTime)
+			{
+				return ::GetTypeHash(InTime.Time);
+			}
+			
 			friend FSampleCounter;
 	};
 
diff --git a/Engine/Plugins/Runtime/WaveTable/Source/WaveTable/Public/WaveTableBank.h b/Engine/Plugins/Runtime/WaveTable/Source/WaveTable/Public/WaveTableBank.h
--- a/Engine/Plugins/Runtime/WaveTable/Source/WaveTable/Public/WaveTableBank.h
+++ b/Engine/Plugins/Runtime/WaveTable/Source/WaveTable/Public/WaveTableBank.h
@@ -127,6 +127,11 @@ public:
 		return ObjectId;
 	}
 
+	friend FORCEINLINE uint32 GetTypeHash(const FWaveTableBankAssetProxy& InProxy)
+	{
+		return InProxy.GetObjectId();
+	}
+
 protected:
 	bool bBipolar = false;
 	uint32 ObjectId = INDEX_NONE;
diff --git a/Engine/Source/Runtime/AudioExtensions/Public/IAudioModulation.h b/Engine/Source/Runtime/AudioExtensions/Public/IAudioModulation.h
--- a/Engine/Source/Runtime/AudioExtensions/Public/IAudioModulation.h
+++ b/Engine/Source/Runtime/AudioExtensions/Public/IAudioModulation.h
@@ -71,6 +71,8 @@ namespace Audio
 		// Whether or not unit conversion is required
 		bool bRequiresConversion = false;
 
+		uint32 TypeHash = INDEX_NONE;
+
 #if WITH_EDITORONLY_DATA
 		FText UnitDisplayName;
 
@@ -89,6 +91,11 @@ namespace Audio
 		static AUDIOEXTENSIONS_API const FModulationMixFunction& GetDefaultMixFunction();
 		static AUDIOEXTENSIONS_API const FModulationUnitConversionFunction& GetDefaultUnitConversionFunction();
 		static AUDIOEXTENSIONS_API const FModulationNormalizedConversionFunction& GetDefaultNormalizedConversionFunction();
+
+		friend FORCEINLINE uint32 GetTypeHash(const FModulationParameter& InModulationParameter)
+		{
+			return InModulationParameter.TypeHash;
+		}
 	};
 
 	AUDIOEXTENSIONS_API bool IsModulationParameterRegistered(FName InName);
diff --git a/Engine/Source/Runtime/Engine/Classes/Sound/AudioBus.h b/Engine/Source/Runtime/Engine/Classes/Sound/AudioBus.h
--- a/Engine/Source/Runtime/Engine/Classes/Sound/AudioBus.h
+++ b/Engine/Source/Runtime/Engine/Classes/Sound/AudioBus.h
@@ -56,6 +56,14 @@ public:
 
 	uint32 AudioBusId = INDEX_NONE;
 	int32 NumChannels = INDEX_NONE;
+
+	friend FORCEINLINE uint32 GetTypeHash(const FAudioBusProxy& InProxy)
+	{
+		return InProxy.TypeHash;
+	}
+
+private: 
+	uint32 TypeHash = INDEX_NONE;
 };
 
 // Function to retrieve an audio bus buffer given a handle
diff --git a/Engine/Source/Runtime/Engine/Classes/Sound/SoundWave.h b/Engine/Source/Runtime/Engine/Classes/Sound/SoundWave.h
--- a/Engine/Source/Runtime/Engine/Classes/Sound/SoundWave.h
+++ b/Engine/Source/Runtime/Engine/Classes/Sound/SoundWave.h
@@ -1751,8 +1751,14 @@ public:
 	ENGINE_API const uint8* GetResourceData() const;
 
 	ENGINE_API const FSoundWavePtr GetSoundWaveData();
-
+	
+	friend FORCEINLINE uint32 GetTypeHash(const FSoundWaveProxy& InProxy)
+	{
+		return InProxy.TypeHash;
+	}
+	
 private:
 	TSharedPtr<FSoundWaveData, ESPMode::ThreadSafe> SoundWaveDataPtr;
+	uint32 TypeHash = INDEX_NONE;
 };
 
diff --git a/Engine/Source/Runtime/Engine/Private/AudioBus.cpp b/Engine/Source/Runtime/Engine/Private/AudioBus.cpp
--- a/Engine/Source/Runtime/Engine/Private/AudioBus.cpp
+++ b/Engine/Source/Runtime/Engine/Private/AudioBus.cpp
@@ -14,6 +14,7 @@ FAudioBusProxy::FAudioBusProxy(UAudioBus* InAudioBus)
 	{
 		AudioBusId = InAudioBus->GetUniqueID();
 		NumChannels = InAudioBus->GetNumChannels();
+		TypeHash = ::GetTypeHash(InAudioBus);
 	}
 }
 
diff --git a/Engine/Source/Runtime/Engine/Private/SoundWave.cpp b/Engine/Source/Runtime/Engine/Private/SoundWave.cpp
--- a/Engine/Source/Runtime/Engine/Private/SoundWave.cpp
+++ b/Engine/Source/Runtime/Engine/Private/SoundWave.cpp
@@ -4668,6 +4668,7 @@ FSoundWaveProxy::FSoundWaveProxy(USoundWave* InWave)
 	{
 		UE_LOG(LogAudio, Warning, TEXT("FSoundWaveProxy doesn't have audio data! bIsStreaming = %d, bHasResourceData = %d, Sound = %s"), bIsStreaming, bHasResourceData, *InWave->GetName());
 	}
+	TypeHash = ::GetTypeHash(InWave);
 }
 
 FSoundWaveProxy::~FSoundWaveProxy()
diff --git a/Engine/Source/Runtime/MediaAssets/Private/Assets/MediaPlayer.cpp b/Engine/Source/Runtime/MediaAssets/Private/Assets/MediaPlayer.cpp
--- a/Engine/Source/Runtime/MediaAssets/Private/Assets/MediaPlayer.cpp
+++ b/Engine/Source/Runtime/MediaAssets/Private/Assets/MediaPlayer.cpp
@@ -1396,6 +1396,7 @@ void UMediaPlayer::OpenSourceLatent(const UObject* WorldContextObject, FLatentAc
 FMediaPlayerProxy::FMediaPlayerProxy(UMediaPlayer* Player)
 {
 	PlayerFacade = Player->GetPlayerFacade();
+	TypeHash = GetTypeHash(Player);
 }
 
 FMediaPlayerProxy::~FMediaPlayerProxy()
diff --git a/Engine/Source/Runtime/MediaAssets/Public/MediaPlayer.h b/Engine/Source/Runtime/MediaAssets/Public/MediaPlayer.h
--- a/Engine/Source/Runtime/MediaAssets/Public/MediaPlayer.h
+++ b/Engine/Source/Runtime/MediaAssets/Public/MediaPlayer.h
@@ -153,8 +153,14 @@ public:
 		return(PlayerFacade.Pin());
 	}
 
+	friend FORCEINLINE uint32 GetTypeHash(const FMediaPlayerProxy& InProxy)
+	{
+		return InProxy.TypeHash;
+	}
+
 private:
 	TWeakPtr<FMediaPlayerFacade, ESPMode::ThreadSafe> PlayerFacade;
+	uint32 TypeHash = INDEX_NONE;
 };
 using FMediaPlayerProxyPtr = TSharedPtr<FMediaPlayerProxy, ESPMode::ThreadSafe>;
 
