[Harmonix] Skip appending/replacing tempo change if before latest tempo change, remove assert

diff --git a/Engine/Plugins/Runtime/Harmonix/Source/HarmonixMetasound/Private/HarmonixMetasound/DataTypes/MidiClock.cpp b/Engine/Plugins/Runtime/Harmonix/Source/HarmonixMetasound/Private/HarmonixMetasound/DataTypes/MidiClock.cpp
--- a/Engine/Plugins/Runtime/Harmonix/Source/HarmonixMetasound/Private/HarmonixMetasound/DataTypes/MidiClock.cpp
+++ b/Engine/Plugins/Runtime/Harmonix/Source/HarmonixMetasound/Private/HarmonixMetasound/DataTypes/MidiClock.cpp
@@ -684,18 +684,21 @@ namespace HarmonixMetasound
 
 	void FMidiClock::FMidiClockEventCursor::OnTempo(int32 TrackIndex, int32 Tick, int32 Tempo, bool IsPreroll)
 	{
-		const int32 BlockFrameIndex = MyMidiClock->CurrentBlockFrameIndex;
-
-		check(BlockFrameIndex >= MyMidiClock->TempoChangesInBlock.Last().BlockSampleFrameIndex);
-		MyMidiClock->HasTempoChangeInBlock = true;
+		const int32 BlockFrameIndex = MyMidiClock->GetCurrentBlockFrameIndex();
 		const float Bpm = Harmonix::Midi::Constants::MidiTempoToBPM(Tempo);
-		if (MyMidiClock->TempoChangesInBlock.Last().BlockSampleFrameIndex == BlockFrameIndex)
+
+		if (BlockFrameIndex >= MyMidiClock->TempoChangesInBlock.Last().BlockSampleFrameIndex)
 		{
-			MyMidiClock->TempoChangesInBlock.Last().Tempo = Bpm;
-		}
-		else
-		{ 
-			MyMidiClock->TempoChangesInBlock.Add({BlockFrameIndex, 0.0f, Bpm});
+			MyMidiClock->HasTempoChangeInBlock = true;
+			
+			if (MyMidiClock->TempoChangesInBlock.Last().BlockSampleFrameIndex == BlockFrameIndex)
+			{
+				MyMidiClock->TempoChangesInBlock.Last().Tempo = Bpm;
+			}
+			else
+			{ 
+				MyMidiClock->TempoChangesInBlock.Add({BlockFrameIndex, 0.0f, Bpm});
+			}
 		}
 
 		// CurrentAdvanceStartTick is the tick ALREADY processed in a previous advance. So the first new tick 
@@ -708,11 +711,11 @@ namespace HarmonixMetasound
 		const int32 TickPriorToThisEvent = Tick - 1;
 		if (Tick > FirstNewTickThisAdvance)
 		{
-			AddEvent(FMidiClockEvent(MyMidiClock->GetCurrentBlockFrameIndex(), MidiClockMessageTypes::FAdvanceThru(CurrentAdvanceStartTick, TickPriorToThisEvent, IsPreroll)));
+			AddEvent(FMidiClockEvent(BlockFrameIndex, MidiClockMessageTypes::FAdvanceThru(CurrentAdvanceStartTick, TickPriorToThisEvent, IsPreroll)));
 			CurrentAdvanceStartTick = TickPriorToThisEvent;
 		}
 		
-		AddEvent(FMidiClockEvent(MyMidiClock->GetCurrentBlockFrameIndex(), MidiClockMessageTypes::FTempoChange(Tick, Bpm)));
+		AddEvent(FMidiClockEvent(BlockFrameIndex, MidiClockMessageTypes::FTempoChange(Tick, Bpm)));
 	}
 
 	void FMidiClock::FMidiClockEventCursor::OnTimeSig(int32 TrackIndex, int32 Tick, int32 Numerator, int32 Denominator, bool IsPreroll)
