[Harmonix] Delay: Use FMultichannelBuffer and FAlignedFloatBuffer instead of TAudioBuffer

#tests automated tests

diff --git a/Engine/Plugins/Runtime/Harmonix/Source/HarmonixDsp/Private/HarmonixDsp/Effects/Delay.cpp b/Engine/Plugins/Runtime/Harmonix/Source/HarmonixDsp/Private/HarmonixDsp/Effects/Delay.cpp
--- a/Engine/Plugins/Runtime/Harmonix/Source/HarmonixDsp/Private/HarmonixDsp/Effects/Delay.cpp
+++ b/Engine/Plugins/Runtime/Harmonix/Source/HarmonixDsp/Private/HarmonixDsp/Effects/Delay.cpp
@@ -1,6 +1,8 @@
 // Copyright Epic Games, Inc. All Rights Reserved.
 #include "HarmonixDsp/Effects/Delay.h"
 
+#include "DSP/FloatArrayMath.h"
+
 #define LOG10OFONEBIT -4.5154367f
 
 namespace Harmonix::Dsp::Effects
@@ -116,12 +118,11 @@ namespace Harmonix::Dsp::Effects
 		PosMask = Length - 1;
 
 		// allocate and clear the delay line array;
-		DelayLineInterleaved.Configure(MaxChannels, Length, EAudioBufferCleanupMode::Delete, SampleRate, true);
-		DelayLineInterleaved.ZeroData();
+		DelayLineInterleaved.Reset();
+		DelayLineInterleaved.SetNumZeroed(MaxChannels * Length);
 		// allocate the wet channel buffer
-		WetChannelInterleaved.Configure(MaxChannels, Length, EAudioBufferCleanupMode::Delete, SampleRate, true);
-		WetChannelInterleaved.ZeroData();
-
+		WetChannelInterleaved.Reset();
+		WetChannelInterleaved.SetNumZeroed(MaxChannels * Length);
 
 		Speed = 1.0f;
 	}
@@ -164,10 +165,10 @@ namespace Harmonix::Dsp::Effects
 		Lfo.UseSettings(&LfoSettings);
 	}
 
-	void FDelay::Process(TAudioBuffer<float>& InOutData)
+	void FDelay::Process(Audio::FMultichannelBufferView& InOutBuffer)
 	{
-		const int32 NumFrames = InOutData.GetNumValidFrames();
-		ActiveChannels = InOutData.GetNumValidChannels();
+		const int32 NumFrames = Audio::GetMultichannelBufferNumFrames(InOutBuffer);
+		ActiveChannels = InOutBuffer.Num();
 		check(MaxChannels >= ActiveChannels);
 
 		ApplyNewParams();
@@ -180,11 +181,11 @@ namespace Harmonix::Dsp::Effects
 		TDynamicStridePtr<float> InOutDataPointers[AbsoluteMaxChannels];
 		TDynamicStridePtr<float> DelayLinePointers[AbsoluteMaxChannels];
 		TDynamicStridePtr<float> WetChannelPointers[AbsoluteMaxChannels];
-		for (int32 i = 0; i < AbsoluteMaxChannels; ++i)
+		for (int32 i = 0; i < ActiveChannels; ++i)
 		{
-			InOutDataPointers[i] = InOutData.GetStridingChannelDataPointer(i);
-			DelayLinePointers[i] = DelayLineInterleaved.GetStridingChannelDataPointer(i);
-			WetChannelPointers[i] = WetChannelInterleaved.GetStridingChannelDataPointer(i);
+			InOutDataPointers[i] = TDynamicStridePtr(&InOutBuffer[i][0], 1);
+			DelayLinePointers[i] = TDynamicStridePtr(&DelayLineInterleaved[i], ActiveChannels);
+			WetChannelPointers[i] = TDynamicStridePtr(&WetChannelInterleaved[i], ActiveChannels);
 		}
 		
 		for (int32 s = 0; s < NumFrames; ++s)
@@ -249,7 +250,7 @@ namespace Harmonix::Dsp::Effects
 			// Filter the feedback
 			if (FeedbackFilters.GetSettings().IsEnabled)
 			{
-				float* DelayLineFramePointer = DelayLineInterleaved.GetRawChannelData(0) + DelayPos * MaxChannels;
+				float* DelayLineFramePointer = DelayLineInterleaved.GetData() + DelayPos * MaxChannels;
 				FeedbackFilters.ProcessInterleavedInPlace(DelayLineFramePointer, 1, MaxChannels, ActiveChannels);
 			}
 			
@@ -259,7 +260,7 @@ namespace Harmonix::Dsp::Effects
 		// Filter the wet signal
 		if (WetFilters.GetSettings().IsEnabled)
 		{
-			WetFilters.ProcessInterleavedInPlace(WetChannelInterleaved.GetRawChannelData(0), NumFrames, MaxChannels, ActiveChannels);
+			WetFilters.ProcessInterleavedInPlace(WetChannelInterleaved.GetData(), NumFrames, MaxChannels, ActiveChannels);
 		}
 
 		// Add the wet signal to the output
@@ -274,7 +275,10 @@ namespace Harmonix::Dsp::Effects
 		// Scale the output
 		if (OutputGain != 1.0f)
 		{
-			InOutData.Scale(OutputGain);
+			Algo::ForEach(InOutBuffer, [Gain = OutputGain](TArrayView<float>& Channel)
+			{
+				Audio::ArrayMultiplyByConstantInPlace(Channel, Gain);
+			});
 		}
 	}
 
@@ -631,7 +635,7 @@ namespace Harmonix::Dsp::Effects
 	{
 		DelayType = EDelayStereoType(InType);
 	}
-
+	
 	float FDelay::CalculateSecsToIdle()
 	{
 		float AbsGain = FMath::Abs(FeedbackGain);
diff --git a/Engine/Plugins/Runtime/Harmonix/Source/HarmonixDsp/Public/HarmonixDsp/Effects/Delay.h b/Engine/Plugins/Runtime/Harmonix/Source/HarmonixDsp/Public/HarmonixDsp/Effects/Delay.h
--- a/Engine/Plugins/Runtime/Harmonix/Source/HarmonixDsp/Public/HarmonixDsp/Effects/Delay.h
+++ b/Engine/Plugins/Runtime/Harmonix/Source/HarmonixDsp/Public/HarmonixDsp/Effects/Delay.h
@@ -1,7 +1,8 @@
 // Copyright Epic Games, Inc. All Rights Reserved.
 #pragma once
 
-#include "HarmonixDsp/AudioBuffer.h"
+#include "Algo/ForEach.h"
+#include "DSP/MultichannelBuffer.h"
 #include "HarmonixDsp/Ramper.h"
 #include "HarmonixDsp/TimeSyncOption.h"
 #include "HarmonixDsp/Effects/Settings/DelaySettings.h"
@@ -40,7 +41,12 @@ namespace Harmonix::Dsp::Effects
 
 		void Clear()
 		{
-			DelayLineInterleaved.ZeroData();
+			// zero the delay line
+			if (DelayLineInterleaved.Num() > 0)
+			{
+				FMemory::Memzero(DelayLineInterleaved.GetData(), DelayLineInterleaved.Num() * sizeof(float));
+			}
+			
 			DelayPos = 0;
 			CanSlamParams = true;
 		}
@@ -77,7 +83,7 @@ namespace Harmonix::Dsp::Effects
 
 		*/
 
-		void Process(TAudioBuffer<float>& InOutData);
+		void Process(Audio::FMultichannelBufferView& InOutBuffer);
 
 		float CalculateSecsToIdle();
 
@@ -171,8 +177,8 @@ namespace Harmonix::Dsp::Effects
 		uint32 MaxBlockSize = 0;
 
 		uint32 Length = 0;
-		TAudioBuffer<float> DelayLineInterleaved;
-		TAudioBuffer<float> WetChannelInterleaved;
+		Audio::FAlignedFloatBuffer DelayLineInterleaved;
+		Audio::FAlignedFloatBuffer WetChannelInterleaved;
 		uint32 DelayPos = 0;
 		uint32 PosMask;
 
diff --git a/Engine/Plugins/Runtime/Harmonix/Source/HarmonixMetasound/Private/HarmonixMetasound/Nodes/DelayNode.cpp b/Engine/Plugins/Runtime/Harmonix/Source/HarmonixMetasound/Private/HarmonixMetasound/Nodes/DelayNode.cpp
--- a/Engine/Plugins/Runtime/Harmonix/Source/HarmonixMetasound/Private/HarmonixMetasound/Nodes/DelayNode.cpp
+++ b/Engine/Plugins/Runtime/Harmonix/Source/HarmonixMetasound/Private/HarmonixMetasound/Nodes/DelayNode.cpp
@@ -325,8 +325,6 @@ namespace HarmonixMetasound::DelayNode
 			Delay.Prepare(Params.OperatorSettings.GetSampleRate(), Constants::NumChannels, Constants::MaxDelayTime);
 			LastTempo = -1;
 			LastSpeed = -1;
-			ScratchBuffer.Initialize();
-			ScratchBuffer.Configure(Constants::NumChannels, Params.OperatorSettings.GetNumFramesPerBlock(), EAudioBufferCleanupMode::DontDelete);
 			Outputs.AudioLeft->Zero();
 			Outputs.AudioRight->Zero();
 		}
@@ -416,26 +414,17 @@ namespace HarmonixMetasound::DelayNode
 				return;
 			}
 
-			// copy the input audio to the output, then alias and process in place
-			const float* InputAudio[Constants::NumChannels];
-			InputAudio[0] = Inputs.AudioLeft->GetData() + StartSample;
-			InputAudio[1] = Inputs.AudioRight->GetData() + StartSample;
-			float* OutputAudio[Constants::NumChannels];
-			OutputAudio[0] = Outputs.AudioLeft->GetData() + StartSample;
-			OutputAudio[1] = Outputs.AudioRight->GetData() + StartSample;
-			for (int32 i = 0; i < Constants::NumChannels; ++i)
-			{
-				FMemory::Memcpy(OutputAudio[i], InputAudio[i], NumSamples * sizeof(float));
-			}
-			const FAudioBufferConfig Config{ Constants::NumChannels, NumSamples, Delay.GetSampleRate() };
-			ScratchBuffer.AliasChannelDataPointers(Config, OutputAudio);
-
-			// process
-			Delay.Process(ScratchBuffer);
+			// copy the input audio to the output, then process in place
+			FMemory::Memcpy(Outputs.AudioLeft->GetData() + StartSample, Inputs.AudioLeft->GetData() + StartSample, NumSamples * sizeof(float));
+			FMemory::Memcpy(Outputs.AudioRight->GetData() + StartSample, Inputs.AudioRight->GetData() + StartSample, NumSamples * sizeof(float));
+			BufferView.Reset();
+			BufferView.Emplace(Outputs.AudioLeft->GetData() + StartSample, NumSamples);
+			BufferView.Emplace(Outputs.AudioRight->GetData() + StartSample, NumSamples);
+			Delay.Process(BufferView);
 		}
 		
 		Harmonix::Dsp::Effects::FDelay Delay;
-		TAudioBuffer<float> ScratchBuffer;
+		Audio::FMultichannelBufferView BufferView;
 		float LastTempo{ -1 };
 		float LastSpeed{ -1 };
 		FInputs Inputs;
diff --git a/Engine/Plugins/Runtime/Harmonix/Source/HarmonixMetasoundTests/Private/Nodes/DelayNodeTests.cpp b/Engine/Plugins/Runtime/Harmonix/Source/HarmonixMetasoundTests/Private/Nodes/DelayNodeTests.cpp
--- a/Engine/Plugins/Runtime/Harmonix/Source/HarmonixMetasoundTests/Private/Nodes/DelayNodeTests.cpp
+++ b/Engine/Plugins/Runtime/Harmonix/Source/HarmonixMetasoundTests/Private/Nodes/DelayNodeTests.cpp
@@ -23,17 +23,10 @@ namespace HarmonixMetasoundTests::DelayNode
 		FTestFixture(float InSampleRate, int32 NumSamplesPerBlock, FAutomationTestBase& InTest, bool WithClock)
 			: SampleRate(InSampleRate)
 			, Test(InTest)
-			, ComparisonBuffer(HarmonixMetasound::DelayNode::Constants::NumChannels, NumSamplesPerBlock, EAudioBufferCleanupMode::Delete)
 		{
-			{
-				const FAudioBufferConfig BufferConfig{
-					HarmonixMetasound::DelayNode::Constants::NumChannels,
-					NumSamplesPerBlock,
-					SampleRate,
-					true
-				};
-				GeneratorBufferInterleaved.Configure(BufferConfig, EAudioBufferCleanupMode::Delete);
-			}
+			constexpr int32 NumChannels = HarmonixMetasound::DelayNode::Constants::NumChannels;
+			Audio::SetMultichannelBufferSize(NumChannels, NumSamplesPerBlock, ComparisonBuffer);
+			GeneratorBufferInterleaved.SetNumZeroed(NumChannels * NumSamplesPerBlock);
 			
 			DelayForComparison.Prepare(
 				InSampleRate,
@@ -96,7 +89,11 @@ namespace HarmonixMetasoundTests::DelayNode
 			using namespace Metasound;
 
 			// zero the input buffers
-			ComparisonBuffer.ZeroData();
+			Algo::ForEach(ComparisonBuffer, [](Audio::FAlignedFloatBuffer& Channel)
+			{
+				check(Channel.Num() > 0);
+				FMemory::Memzero(Channel.GetData(), Channel.Num() * sizeof(float));
+			});
 			TOptional<TDataWriteReference<FAudioBuffer>> InputAudioLeft =
 				Generator->GetInputWriteReference<FAudioBuffer>(HarmonixMetasound::DelayNode::Inputs::AudioLeftName);
 			TOptional<TDataWriteReference<FAudioBuffer>> InputAudioRight =
@@ -120,23 +117,28 @@ namespace HarmonixMetasoundTests::DelayNode
 			}
 
 			// render
-			GeneratorBufferInterleaved.ZeroData();
-			Generator->OnGenerateAudio(GeneratorBufferInterleaved.GetRawChannelData(0), GeneratorBufferInterleaved.GetNumTotalValidSamples());
-			DelayForComparison.Process(ComparisonBuffer);
+			check(GeneratorBufferInterleaved.Num() > 0);
+			FMemory::Memzero(GeneratorBufferInterleaved.GetData(), GeneratorBufferInterleaved.Num() * sizeof(float));
+			Generator->OnGenerateAudio(GeneratorBufferInterleaved.GetData(), GeneratorBufferInterleaved.Num());
+			Audio::FMultichannelBufferView ComparisonBufferView = Audio::MakeMultichannelBufferView(ComparisonBuffer);
+			DelayForComparison.Process(ComparisonBufferView);
 			
 			// check that the output buffers are equal
-			for (int32 ChannelIdx = 0; ChannelIdx < HarmonixMetasound::DelayNode::Constants::NumChannels; ++ChannelIdx)
 			{
-				TDynamicStridePtr<float> GeneratorPtr = GeneratorBufferInterleaved.GetStridingChannelDataPointer(ChannelIdx);
-				TDynamicStridePtr<float> ComparisonPtr = ComparisonBuffer.GetStridingChannelDataPointer(ChannelIdx);
-				for (int32 SampleIdx = 0; SampleIdx < ComparisonBuffer.GetNumValidFrames(); ++SampleIdx)
+				constexpr int32 NumChannels = HarmonixMetasound::DelayNode::Constants::NumChannels;
+				const int32 NumFrames = Generator->OperatorSettings.GetNumFramesPerBlock();
+				
+				for (int32 ChannelIdx = 0; ChannelIdx < NumChannels; ++ChannelIdx)
 				{
-					if (!Test.TestEqual(
-						FString::Printf(TEXT("Channel %i samples match at idx %i"), ChannelIdx, SampleIdx),
-						GeneratorPtr[SampleIdx],
-						ComparisonPtr[SampleIdx]))
+					for (int32 FrameIdx = 0; FrameIdx < NumFrames; ++FrameIdx)
 					{
-						return false;
+						if (!Test.TestEqual(
+							FString::Printf(TEXT("Channel %i samples match at idx %i"), ChannelIdx, FrameIdx),
+							GeneratorBufferInterleaved[FrameIdx * NumChannels + ChannelIdx],
+							ComparisonBuffer[ChannelIdx][FrameIdx]))
+						{
+							return false;
+						}
 					}
 				}
 			}
@@ -278,7 +280,7 @@ namespace HarmonixMetasoundTests::DelayNode
 			{
 				return false;
 			}
-			const int32 NumSamples = ComparisonBuffer.GetNumValidFrames();
+			const int32 NumSamples = Audio::GetMultichannelBufferNumFrames(ComparisonBuffer);
 			SampleRemainder += NumSamples;
 			constexpr int32 MidiGranularity = 128;
 			while (SampleRemainder >= MidiGranularity)
@@ -295,9 +297,9 @@ namespace HarmonixMetasoundTests::DelayNode
 		float SampleRate;
 		FAutomationTestBase& Test;
 		Harmonix::Dsp::Effects::FDelay DelayForComparison;
-		TAudioBuffer<float> ComparisonBuffer;
+		Audio::FMultichannelBuffer ComparisonBuffer;
 		TUniquePtr<FMetasoundGenerator> Generator;
-		TAudioBuffer<float> GeneratorBufferInterleaved;
+		Audio::FAlignedFloatBuffer GeneratorBufferInterleaved;
 		Metasound::FSampleCount SampleCount = 0;
 		Metasound::FSampleCount SampleRemainder = 0;
 	};
