Change Threshold for PSNR tests.

60-96 is considered passing.
#tests HarmonixTests
[RN]

diff --git a/Engine/Plugins/Runtime/Harmonix/Source/HarmonixMetasoundTests/Private/HarmonixMetasoundFunctionalTest.cpp b/Engine/Plugins/Runtime/Harmonix/Source/HarmonixMetasoundTests/Private/HarmonixMetasoundFunctionalTest.cpp
--- a/Engine/Plugins/Runtime/Harmonix/Source/HarmonixMetasoundTests/Private/HarmonixMetasoundFunctionalTest.cpp
+++ b/Engine/Plugins/Runtime/Harmonix/Source/HarmonixMetasoundTests/Private/HarmonixMetasoundFunctionalTest.cpp
@@ -265,7 +265,9 @@ void AHarmonixMetasoundFunctionalTest::CompareResults()
 
 	int32 NumFramesToCompare = FMath::Min(AudioCaptureOutput.Num(), AudioData.Num());
 	float PSNR = Harmonix::Dsp::AudioAnalysis::CalculatePSNR(AudioCaptureOutput.GetData(), AudioData.GetData(), NumChannels, NumFramesToCompare);
-	AssertEqual_Float(PSNR, 96.0f, FString::Printf(TEXT("PSNR of Captured Audio and %s"), *WavFilename_Expected), 0.01f, this);
+
+	static constexpr float PSNRThreshold = 60.0f;
+	AssertTrue(PSNR >= PSNRThreshold, FString::Printf(TEXT("PSNR = %.2f where the acceptable range is (PSNR >= %.2f)"), PSNR, PSNRThreshold), this);
 }
 
 void AHarmonixMetasoundFunctionalTest::StartTest()
