Fixes to make project>Game build with modular target type

diff --git a/Engine/Plugins/Animation/LiveLinkHub/Source/LiveLinkHubMessaging/LiveLinkHubMessaging.build.cs b/Engine/Plugins/Animation/LiveLinkHub/Source/LiveLinkHubMessaging/LiveLinkHubMessaging.build.cs
--- a/Engine/Plugins/Animation/LiveLinkHub/Source/LiveLinkHubMessaging/LiveLinkHubMessaging.build.cs
+++ b/Engine/Plugins/Animation/LiveLinkHub/Source/LiveLinkHubMessaging/LiveLinkHubMessaging.build.cs
@@ -19,7 +19,8 @@ namespace UnrealBuildTool.Rules
 				"CoreUObject",
 				"Engine",
 				"LiveLink",
-				"LiveLinkInterface"
+				"LiveLinkInterface",
+				"SlateCore",
 			});
 
 			if (Target.bCompileAgainstEditor)
@@ -28,7 +29,6 @@ namespace UnrealBuildTool.Rules
 				new string[]
 				{
 					"Slate",
-					"SlateCore",
 					"UnrealEd"
 				});
 			}
diff --git a/Engine/Plugins/Runtime/Harmonix/Source/HarmonixMidi/HarmonixMidi.Build.cs b/Engine/Plugins/Runtime/Harmonix/Source/HarmonixMidi/HarmonixMidi.Build.cs
--- a/Engine/Plugins/Runtime/Harmonix/Source/HarmonixMidi/HarmonixMidi.Build.cs
+++ b/Engine/Plugins/Runtime/Harmonix/Source/HarmonixMidi/HarmonixMidi.Build.cs
@@ -20,21 +20,14 @@ public class HarmonixMidi : ModuleRules
 			new string[]
 			{
 				"AudioExtensions",
+				"Engine",
 				"Harmonix",
 			});
+
 		if (Target.bBuildEditor == true)
 		{
 			PrivateDependencyModuleNames.Add("AssetRegistry");
 			PrivateDependencyModuleNames.Add("UnrealEd");
-			PrivateDependencyModuleNames.Add("Engine");
-		}
-		else
-		{
-			// This next little gem is needed because when this module is built without
-			// an Engine dependency the MidiFile.gen.cpp file has an UNUSED forward declaration
-			// at the top of it that is prefixed by "ENGINE_API", so the compiler sees an undefined
-			// symbol. :-( - Buzz
-			PrivateDefinitions.Add("ENGINE_API=");
 		}
 	}
 }
