[HarmonixMetasound] Crash fix in MidiClockStateNode if BarMap is empty.

diff --git a/Engine/Plugins/Runtime/Harmonix/Source/HarmonixMetasound/Private/HarmonixMetasound/Nodes/MidiClockStateNode.cpp b/Engine/Plugins/Runtime/Harmonix/Source/HarmonixMetasound/Private/HarmonixMetasound/Nodes/MidiClockStateNode.cpp
--- a/Engine/Plugins/Runtime/Harmonix/Source/HarmonixMetasound/Private/HarmonixMetasound/Nodes/MidiClockStateNode.cpp
+++ b/Engine/Plugins/Runtime/Harmonix/Source/HarmonixMetasound/Private/HarmonixMetasound/Nodes/MidiClockStateNode.cpp
@@ -227,12 +227,15 @@ namespace HarmonixMetasound
 		}
 			
 		const FBarMap& BarMap = MidiClockInPin->GetBarMap();
-		const FTimeSignature& TimeSignature = BarMap.GetTimeSignatureAtTick(CurrentTick);
-		if (TimeSignature.Numerator != *CurrentTimeSigNumOutPin || TimeSignature.Denominator != *CurrentTimeSigDenomOutPin)
+		if (!BarMap.IsEmpty())
 		{
-			HasTempoSpeedOrTimeSigChange = true;
-			*CurrentTimeSigNumOutPin = TimeSignature.Numerator;
-			*CurrentTimeSigDenomOutPin = FMath::Max(1, TimeSignature.Denominator);
+			const FTimeSignature& TimeSignature = BarMap.GetTimeSignatureAtTick(CurrentTick);
+			if (TimeSignature.Numerator != *CurrentTimeSigNumOutPin || TimeSignature.Denominator != *CurrentTimeSigDenomOutPin)
+			{
+				HasTempoSpeedOrTimeSigChange = true;
+				*CurrentTimeSigNumOutPin = TimeSignature.Numerator;
+				*CurrentTimeSigDenomOutPin = FMath::Max(1, TimeSignature.Denominator);
+			}
 		}
 
 		if (CurrentTick != LastMidiTick)
