Remove references to Child Renderers in Fusion Sampler.

#tests HarmonixFunctionalTests
#rnx

diff --git a/Engine/Plugins/Runtime/Harmonix/Source/HarmonixDsp/Private/HarmonixDsp/FusionSampler/FusionSampler.cpp b/Engine/Plugins/Runtime/Harmonix/Source/HarmonixDsp/Private/HarmonixDsp/FusionSampler/FusionSampler.cpp
--- a/Engine/Plugins/Runtime/Harmonix/Source/HarmonixDsp/Private/HarmonixDsp/FusionSampler/FusionSampler.cpp
+++ b/Engine/Plugins/Runtime/Harmonix/Source/HarmonixDsp/Private/HarmonixDsp/FusionSampler/FusionSampler.cpp
@@ -78,7 +78,6 @@ FFusionSampler::FFusionSampler()
 
 FFusionSampler::~FFusionSampler()
 {
-	DumpAllChildren();
 	if (FusionPatchData)
 	{
 		FusionPatchData->DisconnectSampler(this);
@@ -773,7 +772,6 @@ void FFusionSampler::SetSampleRate(float InSampleRateHz)
 
 void FFusionSampler::ResetInstrumentStateImpl()
 {
-	DumpAllChildren();
 	// TODO: Move some reset stuff from ResetPatchRelatedState to here!!!!
 	SetMidiChannelGain(1.0, AudioRendering::kMicroFadeSec);
 	SetMidiExpressionGain(1.0, AudioRendering::kMicroFadeSec);
@@ -1217,28 +1215,7 @@ void FFusionSampler::Process(uint32 InSliceIdx, uint32 InSubSliceIdx, TAudioBuff
 			Output.Accumulate(VoiceWorkBufferChannels, NumChannels, NumSamplesToProcess);
 		}
 	}
-
-	// set up AudioBuffer alias for mixing children
-	TAudioBuffer<float> ChildMixBuffer;
-	ChildMixBuffer.Configure(
-		Output.GetChannelLayout(), 
-		NumSamplesToProcess,
-		EAudioBufferCleanupMode::DontDelete,
-		(float)GetSamplesPerSecond(), 
-		false);
-
-	for (uint32 Ch = 0; Ch < NumChannels; ++Ch)
-	{
-		ChildMixBuffer.SetChannelData(Ch, VoiceWorkBufferChannels[Ch]);
-	}
-
-	// mix in any pre-fx children now
-	for (int32 ChildIdx = 0; ChildIdx < FusionPreChildren.Num(); ++ChildIdx)
-	{
-		FusionPreChildren[ChildIdx]->Process(InSliceIdx, InSubSliceIdx, ChildMixBuffer);
-		Output.Accumulate(ChildMixBuffer);
-	}
-
+	
 	// de-activate portamento if we're in legato mode
 	// and there are no voices playing
 
@@ -1246,13 +1223,6 @@ void FFusionSampler::Process(uint32 InSliceIdx, uint32 InSubSliceIdx, TAudioBuff
 	{
 		IsPortamentoActive = false;
 	}
-
-	// mix in any post-fx children now
-	for (int32 ChildIdx = 0; ChildIdx < FusionPostChildren.Num(); ++ChildIdx)
-	{
-		FusionPostChildren[ChildIdx]->Process(InSliceIdx, InSubSliceIdx, ChildMixBuffer);
-		Output.Accumulate(ChildMixBuffer);
-	}
 }
 
 void FFusionSampler::GetController(Harmonix::Midi::Constants::EControllerID InController, int8& Msb, int8& Lsb, int8 InMidiChannel) const
@@ -1977,72 +1947,6 @@ void FFusionSampler::UpdateVoicesForEnvelopeOrderChange()
 	}
 }
 
-// TODO: Needs to be reimplemented without FHarmonixGeneratorHandle
-//bool FFusionSampler::AddChild(const FHarmonixGeneratorHandle& child, EInstrumentRenderMode InRenderMode)
-//{
-//	FInstrumentGeneratorLock ihl(child);
-//	if (!ihl)
-//		return false;
-//
-//	_AddChild(ihl.GetRaw(), InRenderMode);
-//
-//	FusionChildHandles.Add(child);
-//
-//	return true;
-//}
-//
-//bool FFusionSampler::RemoveChild(const FHarmonixGeneratorHandle& child)
-//{
-//	FInstrumentGeneratorLock ihl(child);
-//	if (!ihl)
-//		return false;
-//
-//	FScopeLock Lock(&GetBusLock());
-//
-//	bool found = _RemoveChild(ihl.GetRaw());
-//
-//	int32 ChildIdx = FusionChildHandles.Find(child);
-//	if (ChildIdx != INDEX_NONE)
-//	{
-//		FusionChildHandles.RemoveAt(ChildIdx);
-//		found = true;
-//	}
-//	return found;
-//}
-
-void FFusionSampler::AddChild(FVirtualInstrument* InChildInstrument, EInstrumentRenderMode InRenderMode)
-{
-	FScopeLock Lock(&GetBusLock());
-
-	if (InRenderMode == EInstrumentRenderMode::PreFxChild)
-	{
-		FusionPreChildren.Add(InChildInstrument);
-	}
-	else
-	{
-		FusionPostChildren.Add(InChildInstrument);
-	}
-}
-
-bool FFusionSampler::RemoveChild(FVirtualInstrument* InChildInstrument)
-{
-	FScopeLock Lock(&GetBusLock());
-
-	bool OutRemoved = false;
-
-	OutRemoved |= FusionPreChildren.Remove(InChildInstrument) > 0; 
-	OutRemoved |= FusionPostChildren.Remove(InChildInstrument) > 0;
-
-	return OutRemoved;
-}
-
-void FFusionSampler::DumpAllChildren()
-{
-	FScopeLock Lock(&BusLock);
-	FusionPreChildren.Reset();
-	FusionPostChildren.Reset();
-}
-
 FString FFusionSampler::GetPatchPath() const
 {
 	unimplemented();
diff --git a/Engine/Plugins/Runtime/Harmonix/Source/HarmonixDsp/Public/HarmonixDsp/FusionSampler/FusionSampler.h b/Engine/Plugins/Runtime/Harmonix/Source/HarmonixDsp/Public/HarmonixDsp/FusionSampler/FusionSampler.h
--- a/Engine/Plugins/Runtime/Harmonix/Source/HarmonixDsp/Public/HarmonixDsp/FusionSampler/FusionSampler.h
+++ b/Engine/Plugins/Runtime/Harmonix/Source/HarmonixDsp/Public/HarmonixDsp/FusionSampler/FusionSampler.h
@@ -26,13 +26,6 @@ struct FFusionPatchData;
 class FFusionVoicePool;
 using FSharedFusionVoicePoolPtr = TSharedPtr<FFusionVoicePool, ESPMode::ThreadSafe>;
 
-enum class EInstrumentRenderMode : uint8
-{
-	Parent = 0,
-	PreFxChild = 1,
-	PostFxChild = 2
-};
-
 // Make Fusion Sampler derive directly from virtual instrument
 class HARMONIXDSP_API FFusionSampler : public FVirtualInstrument
 {
@@ -102,14 +95,6 @@ public:
 		return SubstreamGain[Index];
 	}
 
-	// "child" will be processed manually by us, with its output
-	//  mixed into our signal chain.
-	// TODO: Commented out for now. Needs to be reimplemented without HarmonixGeneratorHandle
-	// virtual bool AddChild(const FHarmonixGeneratorHandle& child, EInstrumentRenderMode renderMode);
-	// virtual bool RemoveChild(const FHarmonixGeneratorHandle& child);
-
-	virtual bool IsChildRenderer() const { return false; }
-
 	virtual void SetSampleRate(float InSampleRateHz) override;
 
 	float GetSubstreamGain(int32 Index) const
@@ -158,10 +143,6 @@ protected:
 
 	FString GetPatchPath() const;
 
-	void AddChild(FVirtualInstrument* child, EInstrumentRenderMode renderMode);
-	bool RemoveChild(FVirtualInstrument* child);
-	void DumpAllChildren();
-
 	// apply all settings from the patch.
 	// call whenever the patch is set
 	void ApplyPatchSettings();
@@ -375,10 +356,6 @@ private:
 
 protected:
 
-	// TODO: Replace these with FVirtualInstrument
-	TArray<FVirtualInstrument*> FusionPreChildren;
-	TArray<FVirtualInstrument*> FusionPostChildren;
-
 	// we'll keep track of the 'current tempo' as it is
 	// needed by any 'beat sync' effects...
 	float CurrentTempoBPM = 120.0f;
