[HarmonixMetasound] Allow Single Step Length

diff --git a/Engine/Plugins/Runtime/Harmonix/Source/HarmonixMetasound/Private/HarmonixMetasound/Nodes/StepSequencePlayerNode.cpp b/Engine/Plugins/Runtime/Harmonix/Source/HarmonixMetasound/Private/HarmonixMetasound/Nodes/StepSequencePlayerNode.cpp
--- a/Engine/Plugins/Runtime/Harmonix/Source/HarmonixMetasound/Private/HarmonixMetasound/Nodes/StepSequencePlayerNode.cpp
+++ b/Engine/Plugins/Runtime/Harmonix/Source/HarmonixMetasound/Private/HarmonixMetasound/Nodes/StepSequencePlayerNode.cpp
@@ -819,7 +819,7 @@ namespace HarmonixMetasound::Nodes::StepSequencePlayer
 				}
 			}
 
-			if (CellInRow != CurrentCellIndex)
+			if (CellInRow != CurrentCellIndex || CurrentMaxColumns == 1)
 			{	
 				for (int32 i = 0; i < CurrentPage->Rows.Num(); ++i)
 				{
