MidiStream::Merge now valid if stream Clocks aren't same

diff --git a/Engine/Plugins/Runtime/Harmonix/Source/HarmonixMetasound/Private/HarmonixMetasound/DataTypes/MidiStream.cpp b/Engine/Plugins/Runtime/Harmonix/Source/HarmonixMetasound/Private/HarmonixMetasound/DataTypes/MidiStream.cpp
--- a/Engine/Plugins/Runtime/Harmonix/Source/HarmonixMetasound/Private/HarmonixMetasound/DataTypes/MidiStream.cpp
+++ b/Engine/Plugins/Runtime/Harmonix/Source/HarmonixMetasound/Private/HarmonixMetasound/DataTypes/MidiStream.cpp
@@ -133,22 +133,13 @@ namespace HarmonixMetasound
 
 	void FMidiStream::Merge(const FMidiStream& From, FMidiStream& To, const FEventFilter& Filter, const FEventTransformer& Transformer)
 	{
-		// We need to make sure that the clocks match, or that one of them doesn't have a clock.
-		// Otherwise a merge is invalid.
-		{
-			const auto FromClock = From.GetClock();
-			const auto ToClock = To.GetClock();
-			
-			if (FromClock.Get() != ToClock.Get() && FromClock.IsValid() && ToClock.IsValid())
-			{
-				return;
-			}
+		const auto FromClock = From.GetClock();
+		const auto ToClock = To.GetClock();
 
-			// If the "to" clock is null, and the "from" clock isn't, overwrite the "to" clock
-			if (FromClock.IsValid() && !ToClock.IsValid())
-			{
-				To.SetClock(*FromClock);
-			}
+		// If the "to" clock is null, and the "from" clock isn't, overwrite the "to" clock
+		if (FromClock.IsValid() && !ToClock.IsValid())
+		{
+			To.SetClock(*FromClock);
 		}
 
 		// Insert the events
