[Harmonix] Tech debt cleanup. Don't use TAudioBuffer for distortion (+node)

#rnx

diff --git a/Engine/Plugins/Runtime/Harmonix/Source/HarmonixDsp/Private/HarmonixDsp/Effects/DistortionV2.cpp b/Engine/Plugins/Runtime/Harmonix/Source/HarmonixDsp/Private/HarmonixDsp/Effects/DistortionV2.cpp
--- a/Engine/Plugins/Runtime/Harmonix/Source/HarmonixDsp/Private/HarmonixDsp/Effects/DistortionV2.cpp
+++ b/Engine/Plugins/Runtime/Harmonix/Source/HarmonixDsp/Private/HarmonixDsp/Effects/DistortionV2.cpp
@@ -1,7 +1,7 @@
 // Copyright Epic Games, Inc. All Rights Reserved.
+
 #include "HarmonixDsp/Effects/DistortionV2.h"
 
-#include "HarmonixDsp/StridePointer.h"
 #include "Misc/ScopeLock.h"
 
 namespace Harmonix::Dsp::Effects
@@ -75,7 +75,7 @@ namespace Harmonix::Dsp::Effects
 
 		// allow for 4x oversampling
 		int32 NumFrames = InMaxRenderBufferSize * 4;
-		UpsampleBuffer.Configure(1, NumFrames, EAudioBufferCleanupMode::Delete, (float)SampleRate);
+		UpsampleBuffer.SetNumUninitialized(NumFrames);
 	}
 
 	void FDistortionV2::SetInputGainDb(float InGainDb, bool Snap)
@@ -305,10 +305,7 @@ namespace Harmonix::Dsp::Effects
 			}
 
 			// make sure up-sample buffer can fit enough samples for 4x oversampling...
-			if (UpsampleBuffer.GetMaxConfig().GetNumFrames() < InRenderBufferSizeFrames)
-			{
-				UpsampleBuffer.Configure(1, InRenderBufferSizeFrames * 4, EAudioBufferCleanupMode::Delete, (float)SampleRate);
-			}
+			UpsampleBuffer.SetNumUninitialized(InRenderBufferSizeFrames * 4);
 		}
 		DoOversampling = InDoOversample;
 	}
@@ -322,7 +319,6 @@ namespace Harmonix::Dsp::Effects
 
 		FScopeLock Lock(&SettingsLock);
 		SampleRate = InSampleRate;
-		UpsampleBuffer.SetSampleRate((float)SampleRate);
 		for (int32 i = 0; i < FDistortionSettingsV2::kNumFilters; ++i)
 		{
 			FilterCoefs[i].SetTarget(FBiquadFilterCoefs(FilterSettings[i], (float)SampleRate));
@@ -361,7 +357,7 @@ namespace Harmonix::Dsp::Effects
 		}
 	}
 
-	void FDistortionV2::Process(TAudioBuffer<float>& InBuffer, TAudioBuffer<float>& OutBuffer)
+	void FDistortionV2::Process(const TArray<TArrayView<const float>>& InBuffer, const TArray<TArrayView<float>>& OutBuffer)
 	{
 		FScopeLock Lock(&SettingsLock);
 		TLinearRamper<float>  InputGainTemp = InputGain;
@@ -372,12 +368,15 @@ namespace Harmonix::Dsp::Effects
 		TLinearRamper<FBiquadFilterCoefs> FilterCoefsTemp[FDistortionSettingsV2::kNumFilters];
 		TLinearRamper<float>  FilterGainTemp[FDistortionSettingsV2::kNumFilters];
 
-		for (int32 ch = 0; ch < InBuffer.GetNumValidChannels(); ++ch)
-		{
-			TDynamicStridePtr<float> StrideInPtr = InBuffer.GetStridingChannelDataPointer(ch);
-			TDynamicStridePtr<float> StrideOutPtr = OutBuffer.GetStridingChannelDataPointer(ch);
-			int32 NumFrames = InBuffer.GetNumValidFrames();
+		const int32 NumChannels = InBuffer.Num();
+		check(NumChannels > 0);
+		check(OutBuffer.Num() == NumChannels);
 
+		const int32 NumFrames = InBuffer[0].Num();
+		check(OutBuffer[0].Num() == NumFrames);
+
+		for (int32 ch = 0; ch < NumChannels; ++ch)
+		{
 			// process "pre-clip" filters...
 			bool DidFilter = false;
 			for (int32 i = 0; i < FDistortionSettingsV2::kNumFilters; ++i)
@@ -386,37 +385,40 @@ namespace Harmonix::Dsp::Effects
 				FilterGainTemp[i] = FilterGain[i];
 				if (FilterSettings[i].IsEnabled && FilterPreClip[i])
 				{
-					TDynamicStridePtr<float> InAlias = StrideInPtr;
-					TDynamicStridePtr<float> OutAlias = StrideOutPtr;
+					const float* InPtr = InBuffer[ch].GetData();
+					float* OutPtr = OutBuffer[ch].GetData();
+					
 					for (int32 f = 0; f < NumFrames; f += kRampHops)
 					{
 						FilterCoefsTemp[i].Ramp();
 						FilterGainTemp[i].Ramp();
-						Filter[i][ch].ProcessInterleaved(InAlias, OutAlias, kRampHops, FilterCoefsTemp[i], (double)FilterGainTemp[i], FilterPasses[i]);
-						InAlias += kRampHops;
-						OutAlias += kRampHops;
+						Filter[i][ch].Process(InPtr, OutPtr, kRampHops, FilterCoefsTemp[i], (double)FilterGainTemp[i], FilterPasses[i]);
+						InPtr += kRampHops;
+						OutPtr += kRampHops;
 					}
 					DidFilter = true;
 				}
 			}
 
-			TDynamicStridePtr<float> ProcInPtr = (DidFilter) ? StrideOutPtr : StrideInPtr;
-			TDynamicStridePtr<float> ProcOutPtr = StrideOutPtr;
+			const float* ProcInPtr = (DidFilter) ? OutBuffer[ch].GetData() : InBuffer[ch].GetData();
+			float* ProcOutPtr = OutBuffer[ch].GetData();
 			int32 ProcFrames = NumFrames;
 
 
 			if (DoOversampling)
 			{
 				// up sample...
-				check(NumFrames * 4 <= UpsampleBuffer.GetLengthInFrames());
-				float* FilterOut = UpsampleBuffer.GetValidChannelData(0);
+				check(NumFrames * 4 <= UpsampleBuffer.Num());
+				float* FilterOut = UpsampleBuffer.GetData();
+				
 				for (int32 i = 0; i < NumFrames; ++i)
 				{
 					OversampleFilterUp[ch].Upsample4x(ProcInPtr[i] * 4.0f, FilterOut);
 					FilterOut += 4;
 				}
-				ProcInPtr = UpsampleBuffer.GetStridingChannelDataPointer(0);
-				ProcOutPtr = UpsampleBuffer.GetStridingChannelDataPointer(0);
+				
+				ProcInPtr = UpsampleBuffer.GetData();
+				ProcOutPtr = UpsampleBuffer.GetData();
 				ProcFrames = NumFrames * 4;
 			}
 
@@ -581,16 +583,18 @@ namespace Harmonix::Dsp::Effects
 				}
 			}
 			break;
+			default:
+				checkNoEntry();
 			}
 			if (DoOversampling)
 			{
 				// down sample...
-				float* FilterIn = UpsampleBuffer.GetValidChannelData(0);
+				float* FilterIn = UpsampleBuffer.GetData();
 				for (int32 i = 0; i < NumFrames; ++i)
 				{
 					OversampleFilterDown[ch].AddData(FilterIn, 4);
 					FilterIn += 4;
-					StrideOutPtr[i] = OversampleFilterDown[ch].GetSample();
+					OutBuffer[ch][i] = OversampleFilterDown[ch].GetSample();
 				}
 			}
 
@@ -599,15 +603,15 @@ namespace Harmonix::Dsp::Effects
 			{
 				if (FilterSettings[i].IsEnabled && !FilterPreClip[i])
 				{
-					TDynamicStridePtr<float> InAlias = StrideOutPtr;
-					TDynamicStridePtr<float> OutAlias = StrideOutPtr;
+					const float* InPtr = OutBuffer[ch].GetData();
+					float* OutPtr = OutBuffer[ch].GetData();
 					for (int32 f = 0; f < NumFrames; f += kRampHops)
 					{
 						FilterCoefsTemp[i].Ramp();
 						FilterGainTemp[i].Ramp();
-						Filter[i][ch].ProcessInterleaved(InAlias, OutAlias, kRampHops, FilterCoefsTemp[i], FilterGainTemp[i], FilterPasses[i]);
-						InAlias += kRampHops;
-						OutAlias += kRampHops;
+						Filter[i][ch].Process(InPtr, OutPtr, kRampHops, FilterCoefsTemp[i], FilterGainTemp[i], FilterPasses[i]);
+						InPtr += kRampHops;
+						OutPtr += kRampHops;
 					}
 				}
 			}
@@ -620,7 +624,7 @@ namespace Harmonix::Dsp::Effects
 					WetGainTemp.Ramp();
 					DryGainTemp.Ramp();
 				}
-				StrideOutPtr[i] = FMath::Clamp(StrideOutPtr[i] * WetGainTemp + StrideInPtr[i] * DryGainTemp, -1.0f, 1.0f);
+				OutBuffer[ch][i] = FMath::Clamp(OutBuffer[ch][i] * WetGainTemp + InBuffer[ch][i] * DryGainTemp, -1.0f, 1.0f);
 			}
 		}
 		InputGain = InputGainTemp;
diff --git a/Engine/Plugins/Runtime/Harmonix/Source/HarmonixDsp/Public/HarmonixDsp/Effects/DistortionV2.h b/Engine/Plugins/Runtime/Harmonix/Source/HarmonixDsp/Public/HarmonixDsp/Effects/DistortionV2.h
--- a/Engine/Plugins/Runtime/Harmonix/Source/HarmonixDsp/Public/HarmonixDsp/Effects/DistortionV2.h
+++ b/Engine/Plugins/Runtime/Harmonix/Source/HarmonixDsp/Public/HarmonixDsp/Effects/DistortionV2.h
@@ -1,8 +1,7 @@
 // Copyright Epic Games, Inc. All Rights Reserved.
 #pragma once
 
-#include "HarmonixDsp/AudioBuffer.h"
-#include "HarmonixDsp/AudioUtility.h"
+#include "DSP/AlignedBuffer.h"
 #include "HarmonixDsp/Ramper.h"
 #include "HarmonixDsp/Effects/Settings/DistortionSettings.h"
 #include "HarmonixDsp/Effects/BiquadFilter.h"
@@ -20,8 +19,7 @@ namespace Harmonix::Dsp::Effects
 		void Setup(uint32 SampleRate, uint32 MaxRenderBufferSize);
 		virtual ~FDistortionV2() {};
 
-		void Process(TAudioBuffer<float>& InOutBuffer) { Process(InOutBuffer, InOutBuffer); }
-		void Process(TAudioBuffer<float>& InBuffer, TAudioBuffer<float>& OutBuffer);
+		void Process(const TArray<TArrayView<const float>>& InBuffer, const TArray<TArrayView<float>>& OutBuffer);
 
 		void  SetInputGainDb(float InGainDb, bool Snap = false);
 		float GetInputGainDb() const;
@@ -88,7 +86,7 @@ namespace Harmonix::Dsp::Effects
 		FFirFilter32			OversampleFilterUp[kMaxChannels];
 		FFirFilter32			OversampleFilterDown[kMaxChannels];
 		bool					DoOversampling;
-		TAudioBuffer<float>		UpsampleBuffer;
+		Audio::FAlignedFloatBuffer UpsampleBuffer;
 		uint32					SampleRate;
 
 		static const uint32   kNumFilterTaps = 32;
diff --git a/Engine/Plugins/Runtime/Harmonix/Source/HarmonixMetasound/Private/HarmonixMetasound/Nodes/DistortionNode.cpp b/Engine/Plugins/Runtime/Harmonix/Source/HarmonixMetasound/Private/HarmonixMetasound/Nodes/DistortionNode.cpp
--- a/Engine/Plugins/Runtime/Harmonix/Source/HarmonixMetasound/Private/HarmonixMetasound/Nodes/DistortionNode.cpp
+++ b/Engine/Plugins/Runtime/Harmonix/Source/HarmonixMetasound/Private/HarmonixMetasound/Nodes/DistortionNode.cpp
@@ -146,7 +146,7 @@ namespace HarmonixMetasound
 		};
 
 		FDistortionOperator(
-			const FOperatorSettings& InSettings,
+			const FBuildOperatorParams& InParams,
 			const FAudioBufferReadRef& InAudio,
 			const FFloatReadRef& InInputGain,
 			const FFloatReadRef& InOutputGain,
@@ -159,10 +159,8 @@ namespace HarmonixMetasound
 			const FFilterSettings& InFilterSettings2,
 			const FFilterSettings& InFilterSettings3
 		)
-			: Distortion(InSettings.GetSampleRate(), InSettings.GetNumFramesPerBlock())
+			: Distortion(InParams.OperatorSettings.GetSampleRate(), InParams.OperatorSettings.GetNumFramesPerBlock())
 			, AudioIn(InAudio)
-			, InputBufferAlias(NumChannels, InSettings.GetNumFramesPerBlock(), EAudioBufferCleanupMode::DontDelete)
-			, OutputBufferAlias(NumChannels, InSettings.GetNumFramesPerBlock(), EAudioBufferCleanupMode::DontDelete)
 			, InputGain(InInputGain)
 			, OutputGain(InOutputGain)
 			, DryGain(InDryGain)
@@ -171,16 +169,18 @@ namespace HarmonixMetasound
 			, DistortionType(InDistortionType)
 			, Oversample(InOversample)
 			, FilterSettings{ InFilterSettings1, InFilterSettings2, InFilterSettings3 }
-			, AudioOut(FAudioBufferWriteRef::CreateNew(InSettings))
+			, AudioOut(FAudioBufferWriteRef::CreateNew(InParams.OperatorSettings))
 		{
-			check(AudioOut->Num() == InSettings.GetNumFramesPerBlock());
+			check(AudioOut->Num() == InParams.OperatorSettings.GetNumFramesPerBlock());
 
-			SampleRate = InSettings.GetSampleRate();
-			FramesPerBlock = InSettings.GetNumFramesPerBlock();
+			SampleRate = InParams.OperatorSettings.GetSampleRate();
+			FramesPerBlock = InParams.OperatorSettings.GetNumFramesPerBlock();
 			FDistortionSettingsV2 DistortionSettings;
 			GetCurrentSettings(DistortionSettings);
 			Distortion.Reset();
-			Distortion.Setup(DistortionSettings, InSettings.GetSampleRate(), InSettings.GetNumFramesPerBlock(), true);
+			Distortion.Setup(DistortionSettings, InParams.OperatorSettings.GetSampleRate(), InParams.OperatorSettings.GetNumFramesPerBlock(), true);
+
+			Reset(InParams);
 		}
 
 		static const FVertexInterface& GetVertexInterface()
@@ -282,7 +282,7 @@ namespace HarmonixMetasound
 			};
 
 			return MakeUnique<FDistortionOperator>(
-				InParams.OperatorSettings,
+				InParams,
 				InAudio,
 				InInputGain,
 				InOutputGain,
@@ -317,28 +317,19 @@ namespace HarmonixMetasound
 				InVertexData.BindReadVertex(METASOUND_GET_PARAM_NAME_WITH_INDEX(Inputs::PreClip, Index), FilterSettings[Index].PreClip);
 				InVertexData.BindReadVertex(METASOUND_GET_PARAM_NAME_WITH_INDEX(Inputs::Passes, Index), FilterSettings[Index].Passes);
 			}
+
+			InputBufferView.Reset();
+			InputBufferView.Emplace(*AudioIn);
 		}
 
 		virtual void BindOutputs(FOutputVertexInterfaceData& InVertexData) override
 		{
 			using namespace Effects::Distortion::PinNames;
 			InVertexData.BindReadVertex(METASOUND_GET_PARAM_NAME(Outputs::AudioOut), AudioOut);
-		}
 
-		virtual FDataReferenceCollection GetInputs() const override
-		{
-			// This should never be called. Bind(...) is called instead. This method
-			// exists as a stop-gap until the API can be deprecated and removed.
-			checkNoEntry();
-			return {};
-		}
-
-		virtual FDataReferenceCollection GetOutputs() const override
-		{
-			// This should never be called. Bind(...) is called instead. This method
-			// exists as a stop-gap until the API can be deprecated and removed.
-			checkNoEntry();
-			return {};
+			OutputBufferView.Reset();
+			OutputBufferView.Emplace(*AudioOut);
+			AudioOut->Zero();
 		}
 
 		void Reset(const FResetParams& Params)
@@ -350,6 +341,11 @@ namespace HarmonixMetasound
 			GetCurrentSettings(DistortionSettings);
 			Distortion.Reset();
 			Distortion.Setup(DistortionSettings, SampleRate, FramesPerBlock, true);
+
+			InputBufferView.Reset();
+			InputBufferView.Emplace(*AudioIn);
+			OutputBufferView.Reset();
+			OutputBufferView.Emplace(*AudioOut);
 			AudioOut->Zero();
 		}
 
@@ -380,11 +376,7 @@ namespace HarmonixMetasound
 				Distortion.SetupFilter(Index, Settings);
 			}
 
-			const int32 NumSamples = AudioOut->Num();
-			int32 SampleSliceSize = AudioRendering::kMicroSliceSize;
-			InputBufferAlias.Alias(AudioIn->GetData(), NumSamples, NumChannels);
-			OutputBufferAlias.Alias(AudioOut->GetData(), NumSamples, NumChannels);
-			Distortion.Process(InputBufferAlias, OutputBufferAlias);
+			Distortion.Process(InputBufferView, OutputBufferView);
 		}
 
 	private:
@@ -392,8 +384,6 @@ namespace HarmonixMetasound
 		// INPUT
 		FDistortionV2 Distortion;
 		FAudioBufferReadRef AudioIn;
-		TAudioBuffer<float> InputBufferAlias;
-		TAudioBuffer<float> OutputBufferAlias;
 		FFloatReadRef InputGain;
 		FFloatReadRef OutputGain;
 		FFloatReadRef DryGain;
@@ -407,6 +397,8 @@ namespace HarmonixMetasound
 		FAudioBufferWriteRef AudioOut;
 
 		// DATA
+		TArray<TArrayView<const float>> InputBufferView;
+		TArray<TArrayView<float>> OutputBufferView;
 		int32 FramesPerBlock = 0;
 		int32 SampleRate = 0;
 
