Massive Midi Clock Refactor

- Entirely new data structure for FMidiClock
- New FMidiCursor to replace old FMidiPlayCursorMgr & FMidiPlayCursor
- Heavily modified FMetasoundMidiClockDriverto handle new clock system
- New Midi clock vertex analyzer that saves a clock history

diff --git a/Engine/Plugins/Runtime/Harmonix/Source/Harmonix/Public/Harmonix/LocalMinimumMagnitudeTracker.h b/Engine/Plugins/Runtime/Harmonix/Source/Harmonix/Public/Harmonix/LocalMinimumMagnitudeTracker.h
--- a/Engine/Plugins/Runtime/Harmonix/Source/Harmonix/Public/Harmonix/LocalMinimumMagnitudeTracker.h
+++ b/Engine/Plugins/Runtime/Harmonix/Source/Harmonix/Public/Harmonix/LocalMinimumMagnitudeTracker.h
@@ -4,42 +4,64 @@
 #include "Math/UnrealMathUtility.h"
 #include <limits>
 
-template <int SIZE>
+template <typename T, int32 SIZE>
 class FLocalMinimumMagnitudeTracker
 {
 public:
-	FLocalMinimumMagnitudeTracker()	{ Reset(); }
+	FLocalMinimumMagnitudeTracker() { Reset(); }
 
-   void Reset()
-   {
-      ring[0] = std::numeric_limits<float>::max();
-      nextWrite = 1;
-      minPosition = 0;
-   }
+	void Reset()
+	{
+		Ring[0] = 0;
+		NextWrite = 0;
+		MinPosition = 0;
+		AccumulatedError = 0;
+		Wrapped = false;
+	}
 
-   void Push(float v)
-   {
-      ring[nextWrite] = v;
-      if (nextWrite == minPosition)
-      {
-         // search for new lowest...
-         minPosition = (nextWrite + 1) % SIZE;
-         for (int i = 1; i < SIZE; ++i)
-         {
-            int testIdx = (nextWrite + 1 + i) % SIZE;
-            if (FMath::Abs(ring[testIdx]) <= FMath::Abs(ring[minPosition]))
-               minPosition = testIdx;
-         }
-      }
-      else if (FMath::Abs(v) <= FMath::Abs(ring[minPosition]))
-         minPosition = nextWrite;
-      nextWrite = (nextWrite + 1) % SIZE;
-   }
+	void Push(T v)
+	{
+		if (Wrapped)
+		{
+			AccumulatedError -= Ring[NextWrite];
+		}
+		AccumulatedError += v;
+		Ring[NextWrite] = v;
+		if (NextWrite == MinPosition)
+		{
+			// search for new lowest...
+			MinPosition = (NextWrite + 1) % SIZE;
+			for (int32 i = 1; i < SIZE; ++i)
+			{
+				int32 testIdx = (NextWrite + 1 + i) % SIZE;
+				if (FMath::Abs(Ring[testIdx]) <= FMath::Abs(Ring[MinPosition]))
+					MinPosition = testIdx;
+			}
+		}
+		else if (FMath::Abs(v) <= FMath::Abs(Ring[MinPosition]))
+		{
+			MinPosition = NextWrite;
+		}
+		Wrapped = Wrapped || (NextWrite + 1) >= SIZE;
+		NextWrite = (NextWrite + 1) % SIZE;
+	}
 
-   float Min() const { return ring[minPosition]; }
+	T Min() const { return Ring[MinPosition]; }
+
+	T Average() const 
+	{
+		int32 Count = Wrapped ? (T)SIZE : (T)NextWrite;
+		if (Count == 0)
+		{
+			return 0;
+		}
+		return AccumulatedError / (T)Count;
+	}
 
 private:
-   float ring[SIZE];
-   int   nextWrite = 0;
-   int   minPosition = 0;
+	T Ring[SIZE];
+	int32 NextWrite = 0;
+	int32 MinPosition = 0;
+	T AccumulatedError = 0;
+	bool  Wrapped = false;
 };
diff --git a/Engine/Plugins/Runtime/Harmonix/Source/HarmonixDsp/Private/HarmonixDsp/FusionSampler/FusionSampler.cpp b/Engine/Plugins/Runtime/Harmonix/Source/HarmonixDsp/Private/HarmonixDsp/FusionSampler/FusionSampler.cpp
--- a/Engine/Plugins/Runtime/Harmonix/Source/HarmonixDsp/Private/HarmonixDsp/FusionSampler/FusionSampler.cpp
+++ b/Engine/Plugins/Runtime/Harmonix/Source/HarmonixDsp/Private/HarmonixDsp/FusionSampler/FusionSampler.cpp
@@ -258,10 +258,10 @@ void FFusionSampler::SetTempo(float InBPM)
 	UpdateVoiceLfos();
 }
 
-void FFusionSampler::SetBeat(float Beat)
+void FFusionSampler::SetQuarterNote(float QuarterNote)
 {
 	FScopeLock Lock(&GetBusLock());
-	CurrentBeat = Beat;
+	CurrentQuarterNote = QuarterNote;
 }
 
 void FFusionSampler::ResetNoteActions(bool ClearNotes)
@@ -778,7 +778,7 @@ void FFusionSampler::ResetInstrumentStateImpl()
 	SetMidiChannelMute(false);
 
 	Speed = 1.0f;
-	CurrentBeat = 0.0f;
+	CurrentQuarterNote = 0.0f;
 	CurrentTempoBPM = 120.0f;
 	RawPitchMultiplier = 1.0f;
 	Transposition = 0;
diff --git a/Engine/Plugins/Runtime/Harmonix/Source/HarmonixDsp/Private/HarmonixDsp/FusionSampler/FusionVoice.cpp b/Engine/Plugins/Runtime/Harmonix/Source/HarmonixDsp/Private/HarmonixDsp/FusionSampler/FusionVoice.cpp
--- a/Engine/Plugins/Runtime/Harmonix/Source/HarmonixDsp/Private/HarmonixDsp/FusionSampler/FusionVoice.cpp
+++ b/Engine/Plugins/Runtime/Harmonix/Source/HarmonixDsp/Private/HarmonixDsp/FusionSampler/FusionVoice.cpp
@@ -342,28 +342,28 @@ void FFusionVoice::AttackWithTargetNote(uint8 InMidiNoteNumber, float InGain, in
 	{
 		if (KeyZone->TimeStretchConfig.bMaintainTime)
 		{
-			float EventBeat = (float)InEventTick / MySampler->GetTicksPerQuarterNote();
-			float SamplerBeat = MySampler->GetBeat();
-			float ErrorBeats = SamplerBeat - EventBeat;
+			float EventQuarterNote = (float)InEventTick / MySampler->GetTicksPerQuarterNote();
+			float SamplerQuarterNote = MySampler->GetQuarterNote();
+			float ErrorBeats = SamplerQuarterNote - EventQuarterNote;
 			float ErrorMs = ErrorBeats * 60000.0f / KeyZone->TimeStretchConfig.OriginalTempo;
 			if (ErrorMs > 10.0f)
 			{
 				// gotta let the render code skip forward a bit in process otherwise this voice will be out of sync by a lot!
-				StartBeat = (float)InEventTick / Harmonix::Midi::Constants::GTicksPerQuarterNote;
+				StartQuarterNote = (float)InEventTick / Harmonix::Midi::Constants::GTicksPerQuarterNote;
 			}
 			else
 			{
-				StartBeat = SamplerBeat;
+				StartQuarterNote = SamplerQuarterNote;
 			}
 		}
 		else
 		{
-			StartBeat = MySampler->GetBeat();
+			StartQuarterNote = MySampler->GetQuarterNote();
 		}
 	}
 	else
 	{
-		StartBeat = (float)InEventTick / Harmonix::Midi::Constants::GTicksPerQuarterNote;
+		StartQuarterNote = (float)InEventTick / Harmonix::Midi::Constants::GTicksPerQuarterNote;
 	}
 
 	// multiply by 100.0?
@@ -740,22 +740,22 @@ uint32 FFusionVoice::Process(uint32 InSliceIndex, uint32 InSubsliceIndex, float*
 			MaintainPitchWhenSpeedChanges = true;
 
 			// check for resync...
-			float currentMidiBeat = MySampler->GetBeat();
+			float CurrentMidiQuarterNote = MySampler->GetQuarterNote();
 
-			if (currentMidiBeat > StartBeat && (SamplePos > LastVsoPos || !bHasRenderedAnySamples))
+			if (CurrentMidiQuarterNote > StartQuarterNote && (SamplePos > LastVsoPos || !bHasRenderedAnySamples))
 			{
 				LastVsoPos = SamplePos;
 
-				float ElapsedBeats = currentMidiBeat - StartBeat;
+				float ElapsedBeats = CurrentMidiQuarterNote - StartQuarterNote;
 				float ExpectedElapsedMs = ElapsedBeats * 60000.0f / KeyZone->TimeStretchConfig.OriginalTempo;
 
-				float currentSampleFrame = (PitchShifter->HasCurrentSampleFrame() && bHasRenderedAnySamples) ? (float)StartPos + (float)PitchShifter->GetCurrentSampleFrame() : (float)SamplePos;
+				float CurrentSampleFrame = (PitchShifter->HasCurrentSampleFrame() && bHasRenderedAnySamples) ? (float)StartPos + (float)PitchShifter->GetCurrentSampleFrame() : (float)SamplePos;
 				if (KeyZone->SampleStartOffset != -1)
 				{
-					currentSampleFrame -= KeyZone->SampleStartOffset;
+					CurrentSampleFrame -= KeyZone->SampleStartOffset;
 				}
 
-				float ActualElapsedMs = (currentSampleFrame * 1000.0f) / KeyZone->SoundWaveProxy->GetSampleRate();
+				float ActualElapsedMs = (CurrentSampleFrame * 1000.0f) / KeyZone->SoundWaveProxy->GetSampleRate();
 				float ErrorMs = ActualElapsedMs - ExpectedElapsedMs;
 
 				if (!bHasRenderedAnySamples)
diff --git a/Engine/Plugins/Runtime/Harmonix/Source/HarmonixDsp/Public/HarmonixDsp/FusionSampler/FusionSampler.h b/Engine/Plugins/Runtime/Harmonix/Source/HarmonixDsp/Public/HarmonixDsp/FusionSampler/FusionSampler.h
--- a/Engine/Plugins/Runtime/Harmonix/Source/HarmonixDsp/Public/HarmonixDsp/FusionSampler/FusionSampler.h
+++ b/Engine/Plugins/Runtime/Harmonix/Source/HarmonixDsp/Public/HarmonixDsp/FusionSampler/FusionSampler.h
@@ -82,8 +82,8 @@ public:
 	}
 
 	virtual void SetTempo(float BPM) override;
-	virtual void SetBeat(float Beat) override;
-	float GetBeat() const { return CurrentBeat; }
+	virtual void SetQuarterNote(float QuarterNote) override;
+	float GetQuarterNote() const { return CurrentQuarterNote; }
 
 	float GetSubtsreamGain(int32 Index)
 	{
@@ -227,7 +227,7 @@ private:
 
 	float Speed = 1.0f;
 	bool  MaintainPitchWhenSpeedChanges = false;
-	float CurrentBeat = 0.0f;
+	float CurrentQuarterNote = 0.0f;
 
 	//---------------------------------------------
 	// Gain Settings
diff --git a/Engine/Plugins/Runtime/Harmonix/Source/HarmonixDsp/Public/HarmonixDsp/FusionSampler/FusionVoice.h b/Engine/Plugins/Runtime/Harmonix/Source/HarmonixDsp/Public/HarmonixDsp/FusionSampler/FusionVoice.h
--- a/Engine/Plugins/Runtime/Harmonix/Source/HarmonixDsp/Public/HarmonixDsp/FusionSampler/FusionVoice.h
+++ b/Engine/Plugins/Runtime/Harmonix/Source/HarmonixDsp/Public/HarmonixDsp/FusionSampler/FusionVoice.h
@@ -209,7 +209,7 @@ private:
 	FFusionVoicePool* VoicePool = nullptr;
 
 	// needed for tempo sync'd keyzones...
-	float StartBeat = 0.0f;
+	float StartQuarterNote = 0.0f;
 	float CurrentVso = 0.0f;
 	double LastVsoPos = 0.0f;
 	double StartPos = 0.0f;
diff --git a/Engine/Plugins/Runtime/Harmonix/Source/HarmonixDsp/Public/HarmonixDsp/Instruments/VirtualInstrument.h b/Engine/Plugins/Runtime/Harmonix/Source/HarmonixDsp/Public/HarmonixDsp/Instruments/VirtualInstrument.h
--- a/Engine/Plugins/Runtime/Harmonix/Source/HarmonixDsp/Public/HarmonixDsp/Instruments/VirtualInstrument.h
+++ b/Engine/Plugins/Runtime/Harmonix/Source/HarmonixDsp/Public/HarmonixDsp/Instruments/VirtualInstrument.h
@@ -87,8 +87,8 @@ public:
 	virtual float GetSpeed(bool* MaintainPitch = nullptr) = 0;
 
 	// update the effect with the current beat/tempo
-	virtual void SetBeat(float beat) {}
-	virtual void SetTempo(float bpm) {}
+	virtual void SetQuarterNote(float QuarterNote) {}
+	virtual void SetTempo(float Bpm) {}
 
 	virtual void SetSampleRate(float InSampleRateHz) = 0;
 
diff --git a/Engine/Plugins/Runtime/Harmonix/Source/HarmonixMetasound/Private/HarmonixMetasound/Analysis/MidiClockVertexAnalyzer.cpp b/Engine/Plugins/Runtime/Harmonix/Source/HarmonixMetasound/Private/HarmonixMetasound/Analysis/MidiClockVertexAnalyzer.cpp
--- a/Engine/Plugins/Runtime/Harmonix/Source/HarmonixMetasound/Private/HarmonixMetasound/Analysis/MidiClockVertexAnalyzer.cpp
+++ b/Engine/Plugins/Runtime/Harmonix/Source/HarmonixMetasound/Private/HarmonixMetasound/Analysis/MidiClockVertexAnalyzer.cpp
@@ -57,9 +57,17 @@ namespace HarmonixMetasound::Analysis
 	void FMidiClockVertexAnalyzer::Execute()
 	{
 		const FMidiClock& Clock = GetVertexData<FMidiClock>();
-		*Timestamp = Clock.GetCurrentMusicTimestamp();
+		*Timestamp = Clock.GetMusicTimestampAtBlockEnd();
 		*Tempo = Clock.GetTempoAtEndOfBlock();
-		*TimeSignature = Clock.GetBarMap().GetTimeSignatureAtBar((*Timestamp).Bar);
+		const FTimeSignature* TimeSigPtr = Clock.GetSongMapEvaluator().GetTimeSignatureAtBar((*Timestamp).Bar);
+		if (TimeSigPtr)
+		{
+			*TimeSignature = *TimeSigPtr;
+		}
+		else
+		{
+			*TimeSignature = FTimeSignature(4, 4);
+		}
 		*Speed = Clock.GetSpeedAtEndOfBlock();
 		MarkOutputDirty();
 	}
diff --git a/Engine/Plugins/Runtime/Harmonix/Source/HarmonixMetasound/Private/HarmonixMetasound/Analysis/MidiStreamVertexAnalyzer.cpp b/Engine/Plugins/Runtime/Harmonix/Source/HarmonixMetasound/Private/HarmonixMetasound/Analysis/MidiStreamVertexAnalyzer.cpp
--- a/Engine/Plugins/Runtime/Harmonix/Source/HarmonixMetasound/Private/HarmonixMetasound/Analysis/MidiStreamVertexAnalyzer.cpp
+++ b/Engine/Plugins/Runtime/Harmonix/Source/HarmonixMetasound/Private/HarmonixMetasound/Analysis/MidiStreamVertexAnalyzer.cpp
@@ -45,7 +45,7 @@ namespace HarmonixMetasound::Analysis
 			FMidiEventInfo EventOut;
 			if (const TSharedPtr<const FMidiClock, ESPMode::NotThreadSafe> Clock = MidiStreamIn.GetClock())
 			{
-				EventOut.Timestamp = Clock->GetBarMap().TickToMusicTimestamp(Event.CurrentMidiTick);
+				EventOut.Timestamp = Clock->GetSongMapEvaluator().TickToMusicTimestamp(Event.CurrentMidiTick);
 			}
 			EventOut.TrackIndex = Event.TrackIndex;
 			EventOut.MidiMessage = Event.MidiMessage;
diff --git a/Engine/Plugins/Runtime/Harmonix/Source/HarmonixMetasound/Private/HarmonixMetasound/Components/MetasoundMusicClockDriver.cpp b/Engine/Plugins/Runtime/Harmonix/Source/HarmonixMetasound/Private/HarmonixMetasound/Components/MetasoundMusicClockDriver.cpp
--- a/Engine/Plugins/Runtime/Harmonix/Source/HarmonixMetasound/Private/HarmonixMetasound/Components/MetasoundMusicClockDriver.cpp
+++ b/Engine/Plugins/Runtime/Harmonix/Source/HarmonixMetasound/Private/HarmonixMetasound/Components/MetasoundMusicClockDriver.cpp
@@ -2,6 +2,7 @@
 #include "MetasoundMusicClockDriver.h"
 #include "MetasoundGeneratorHandle.h"
 #include "Components/AudioComponent.h"
+#include "HarmonixMetasound/Analysis/MidiSongPosVertexAnalyzer.h"
 #include "HarmonixMetasound/DataTypes/MidiClock.h"
 #include "MetasoundGeneratorHandle.h"
 #include "MetasoundGenerator.h"
@@ -9,29 +10,124 @@
 #include "Async/Async.h"
 #include "Harmonix.h"
 
+namespace MetasoundMusicClockDriver
+{
+	float Fudge = 1.00;
+	float kP = 0.001;
+	float HistoricSmoothedAudioRenderLagSeconds = 0.030; // this used to be baked-in/hardcoded into the smoothing of the audio render time. 
+	float SmoothedAudioRenderLagSeconds = 0.030f; // 30 ms
+	float MaxErrorSecondsBeforeJump = 0.060f; // 60ms
+	int32 HighWaterNumDataAvailable = 0;
+	double SlowestCorrectionSpeed = 0.98;
+	double FastestCorrectionSpeed = 1.02;
+
+
+	FAutoConsoleVariableRef CVarFudge(
+		TEXT("au.MusicClockComponent.TEST.Fudge"),
+		Fudge,
+		TEXT("Clock Fudge FOR TESTING."),
+		ECVF_Cheat);
+
+	FAutoConsoleVariableRef CVarkP(
+		TEXT("au.MusicClockComponent.kP"),
+		kP,
+		TEXT("Clock kP."),
+		ECVF_Default);
+
+	FAutoConsoleVariableRef CVarAudioRenderLag(
+		TEXT("au.MusicClockComponent.SmoothedAudioRenderLagSeconds"),
+		SmoothedAudioRenderLagSeconds,
+		TEXT("SmoothedAudioRenderLagSeconds."),
+		ECVF_Cheat);
+
+	FAutoConsoleVariableRef CVarMaxErrorSecondsBeforeJump(
+		TEXT("au.MusicClockComponent.MaxErrorSecondsBeforeJump"),
+		MaxErrorSecondsBeforeJump,
+		TEXT("MaxErrorSecondsBeforeJump."),
+		ECVF_Default);
+}
+
 bool FMetasoundMusicClockDriver::CalculateSongPosWithOffset(float MsOffset, ECalibratedMusicTimebase Timebase, FMidiSongPos& OutResult) const
 {
+	using namespace HarmonixMetasound::Analysis;
+
 	check(IsInGameThread());
 
 	// if we have an owner, ask them directly
-	if (CursorOwner)
+	if (!ClockHistory)
+	{
+		return false;
+	}
+
+	if (!CurrentMapChain || !CurrentMapChain->SongMaps)
+	{
+		return false;
+	}
+
+	// start with the offset...
+	float AbsMs = MsOffset;
+	
+	// now add the current time from the proper time base...
+	switch (Timebase)
+	{
+	case ECalibratedMusicTimebase::AudioRenderTime:
+		AbsMs += ClockComponent->CurrentSmoothedAudioRenderSongPos.SecondsIncludingCountIn * 1000.0f;
+		break;
+	case ECalibratedMusicTimebase::ExperiencedTime:
+		AbsMs += ClockComponent->CurrentPlayerExperiencedSongPos.SecondsIncludingCountIn * 1000.0f;
+		break;
+	case ECalibratedMusicTimebase::VideoRenderTime:
+	default:
+		AbsMs += ClockComponent->CurrentVideoRenderSongPos.SecondsIncludingCountIn * 1000.0f;
+		break;
+	}
+
+	// if looping map to range of loop...
+	if (CurrentMapChain->LoopLengthTicks > 0)
 	{
-		switch (Timebase)
+		float LoopStartMs = CurrentMapChain->SongMaps->TickToMs(CurrentMapChain->FirstTickInLoop);
+		float LoopEndMs = CurrentMapChain->SongMaps->TickToMs(CurrentMapChain->FirstTickInLoop + CurrentMapChain->LoopLengthTicks);
+		if (AbsMs > LoopEndMs)
 		{
-		case ECalibratedMusicTimebase::AudioRenderTime:
-			OutResult = CursorOwner->CalculateLowResSongPosRelativeToCurrentMs((Clock->CurrentSmoothedAudioRenderSongPos.SecondsIncludingCountIn * 1000.0f) + MsOffset);
-			break;
-		case ECalibratedMusicTimebase::ExperiencedTime:
-			OutResult = CursorOwner->CalculateLowResSongPosRelativeToCurrentMs((Clock->CurrentPlayerExperiencedSongPos.SecondsIncludingCountIn * 1000.0f) + MsOffset);
-			break;
-		case ECalibratedMusicTimebase::VideoRenderTime:
-		default:
-			OutResult = CursorOwner->CalculateLowResSongPosRelativeToCurrentMs((Clock->CurrentVideoRenderSongPos.SecondsIncludingCountIn * 1000.0f) + MsOffset);
-			break;
+			float LoopLengthMs = LoopEndMs - LoopStartMs;
+			AbsMs = LoopStartMs + FMath::Fmod(AbsMs - LoopStartMs, LoopLengthMs);
 		}
-		return true;
 	}
-	return false;
+
+	// and finally we can calculate a song position...
+	OutResult.SetByTime(AbsMs, *(CurrentMapChain->SongMaps));
+
+	return true;
+}
+
+FMidiSongPos FMetasoundMusicClockDriver::CalculateSongPosAtMs(float AbsoluteMs) const
+{
+	using namespace HarmonixMetasound::Analysis;
+	FMidiSongPos OutSongPos;
+	if (!ClockHistory)
+	{
+		return OutSongPos;
+	}
+
+	if (!CurrentMapChain || !CurrentMapChain->SongMaps)
+	{
+		return OutSongPos;
+	}
+
+	if (CurrentMapChain->LoopLengthTicks > 0)
+	{
+		float LoopStartMs = CurrentMapChain->SongMaps->TickToMs(CurrentMapChain->FirstTickInLoop);
+		float LoopEndMs = CurrentMapChain->SongMaps->TickToMs(CurrentMapChain->FirstTickInLoop + CurrentMapChain->LoopLengthTicks);
+		if (AbsoluteMs > LoopEndMs)
+		{
+			float LoopLengthMs = LoopEndMs - LoopStartMs;
+			AbsoluteMs = LoopStartMs + FMath::Fmod(AbsoluteMs - LoopStartMs, LoopLengthMs);
+		}
+	}
+
+	OutSongPos.SetByTime(AbsoluteMs, *(CurrentMapChain->SongMaps));
+
+	return OutSongPos;
 }
 
 bool FMetasoundMusicClockDriver::RefreshCurrentSongPos()
@@ -54,21 +150,25 @@ bool FMetasoundMusicClockDriver::RefreshCurrentSongPos()
 		}
 	}
 
-	if (CursorOwner)
+	if (bRunning)
 	{
-		// cursor is attached and has the current info
-		RefreshCurrentSongPosFromCursor();
-		return true;
-	}
-	else
-	{
-		// Cursor not attached so use wall clock
-		if (!WasEverConnected || Clock->RunPastMusicEnd)
+		if (ClockHistory)
 		{
-			RefreshCurrentSongPosFromWallClock();
+			// cursor is attached and has the current info
+			RefreshCurrentSongPosFromHistory();
 			return true;
 		}
+		else
+		{
+			// Cursor not attached so use wall clock
+			if (!WasEverConnected || ClockComponent->RunPastMusicEnd)
+			{
+				RefreshCurrentSongPosFromWallClock();
+				return true;
+			}
+		}
 	}
+
 	return false;
 }
 
@@ -77,37 +177,52 @@ void FMetasoundMusicClockDriver::OnStart()
 	check(IsInGameThread());
 
 	SongPosOffsetMs = 0.0f;
-	FreeRunStartTimeSecs = Clock ? Clock->GetWorld()->GetTimeSeconds() : 0.0;
+	RenderStartSampleCount = 0;
+	RenderStartWallClockTimeSeconds = 0.0;
+	FreeRunStartTimeSecs = ClockComponent ? ClockComponent->GetWorld()->GetTimeSeconds() : 0.0;
+	bRunning = true;
+}
+
+void FMetasoundMusicClockDriver::OnPause()
+{
+	check(IsInGameThread());
+	bRunning = false;
 }
 
 void FMetasoundMusicClockDriver::OnContinue()
 {
-	if (!CursorOwner)
+	check(IsInGameThread());
+	if (!ClockHistory)
 	{
 		RefreshCurrentSongPosFromWallClock();
 	}
+	bRunning = true;
+}
+
+void FMetasoundMusicClockDriver::OnStop()
+{
+	check(IsInGameThread());
+	bRunning = false;
 }
 
 void FMetasoundMusicClockDriver::Disconnect()
 {
-	if (CursorOwner)
-	{
-		CursorOwner->UnregisterPlayCursor(&Cursor);
-		CursorOwner.Reset();
-	}
+	check(IsInGameThread());
+	ClockHistory = nullptr;
 	DetachAllCallbacks();
 	AudioComponentToWatch.Reset();
 	CurrentGeneratorHandle.Reset();
 }
 
-const FSongMaps* FMetasoundMusicClockDriver::GetCurrentSongMaps() const
+const ISongMapEvaluator* FMetasoundMusicClockDriver::GetCurrentSongMapEvaluator() const
 {
+	using namespace HarmonixMetasound::Analysis;
 	check(IsInGameThread());
-	if (CursorOwner)
+	if (ClockHistory && CurrentMapChain && CurrentMapChain->SongMaps)
 	{
-		return &CursorOwner->GetSongMaps();
+		return CurrentMapChain->SongMaps.Get();
 	}
-	return &Clock->DefaultMaps;
+	return &ClockComponent->DefaultMaps;
 }
 
 bool FMetasoundMusicClockDriver::ConnectToAudioComponentsMetasound(UAudioComponent* InAudioComponent, FName MetasoundOuputPinName)
@@ -117,43 +232,10 @@ bool FMetasoundMusicClockDriver::ConnectToAudioComponentsMetasound(UAudioCompone
 	return AttemptToConnectToAudioComponentsMetasound();
 }
 
-void FMetasoundMusicClockDriver::ResetCursorOwner(TSharedPtr<FMidiPlayCursorMgr> MidiPlayCursorMgr)
-{
-	check(IsInGameThread());
-
-	// Verify that the cursor owner is changing.
-	if (CursorOwner != MidiPlayCursorMgr)
-	{
-		if (MidiPlayCursorMgr)
-		{
-			// Register with the new cursor owner and broadcast the clock connection event.
-			MidiPlayCursorMgr->RegisterLowResPlayCursor(&Cursor);
-			CursorOwner = MoveTemp(MidiPlayCursorMgr);
-			WasEverConnected = true;
-
-			check(Clock);
-			Clock->MusicClockConnectedEvent.Broadcast();
-		}
-		else
-		{
-			// Unregister the old cursor owner and broadcast the clock disconnection event.
-			CursorOwner->UnregisterPlayCursor(&Cursor);
-			if (Clock->GetState() != EMusicClockState::Stopped)
-			{
-				Clock->DefaultMaps.Copy(CursorOwner->GetSongMaps(), 0, Cursor.GetCurrentTick());
-				SongPosOffsetMs = Clock->CurrentSmoothedAudioRenderSongPos.SecondsIncludingCountIn * 1000.0f;
-				FreeRunStartTimeSecs = Clock->GetWorld()->GetTimeSeconds();
-			}
-			CursorOwner.Reset();
-
-			check(Clock);
-			Clock->MusicClockDisconnectedEvent.Broadcast();
-		}
-	}
-}
-
 bool FMetasoundMusicClockDriver::AttemptToConnectToAudioComponentsMetasound()
 {
+	using namespace HarmonixMetasound::Analysis;
+
 	check(IsInGameThread());
 	if (!AudioComponentToWatch.IsValid() || MetasoundOutputName.IsNone())
 	{
@@ -167,13 +249,26 @@ bool FMetasoundMusicClockDriver::AttemptToConnectToAudioComponentsMetasound()
 	{
 		return false;
 	}
+
+	bool WatchingOutput = CurrentGeneratorHandle->WatchOutput(MetasoundOutputName, 
+		FOnMetasoundOutputValueChangedNative::CreateLambda([](FName, const FMetaSoundOutput&) {}),
+		FMidiSongPosVertexAnalyzer::GetAnalyzerName(),
+		FMidiSongPosVertexAnalyzer::SongPosition.Name);
+	if (WatchingOutput)
+	{
+		ensure(CurrentGeneratorHandle->TryCreateAnalyzerAddress(MetasoundOutputName, 
+				FMidiSongPosVertexAnalyzer::GetAnalyzerName(), 
+				FMidiSongPosVertexAnalyzer::SongPosition.Name, 
+				MidiSongPosAnalyzerAddress));
+	}
+
 	GeneratorAttachedCallbackHandle = CurrentGeneratorHandle->OnGeneratorHandleAttached.AddLambda([this](){OnGeneratorAttached();});
 	GeneratorDetachedCallbackHandle = CurrentGeneratorHandle->OnGeneratorHandleDetached.AddLambda([this](){OnGeneratorDetached();});
 	GeneratorIOUpdatedCallbackHandle = CurrentGeneratorHandle->OnIOUpdatedWithChanges.AddLambda([this](const TArray<Metasound::FVertexInterfaceChange>& VertexInterfaceChanges){OnGeneratorIOUpdatedWithChanges(VertexInterfaceChanges);});
 	UMetasoundGeneratorHandle::FOnSetGraph::FDelegate OnSetGraph;
 	OnSetGraph.BindLambda([this](){OnGraphSet();});
 	GraphChangedCallbackHandle = CurrentGeneratorHandle->AddGraphSetCallback(MoveTemp(OnSetGraph));
-	OnGeneratorAttached();
+	// OnGeneratorAttached();
 	return true;
 }
 
@@ -185,151 +280,301 @@ void FMetasoundMusicClockDriver::DetachAllCallbacks()
 		GeneratorAttachedCallbackHandle.Reset();
 		CurrentGeneratorHandle->OnGeneratorHandleDetached.Remove(GeneratorDetachedCallbackHandle);
 		GeneratorDetachedCallbackHandle.Reset();
+
 		CurrentGeneratorHandle->OnIOUpdatedWithChanges.Remove(GeneratorIOUpdatedCallbackHandle);
 		GeneratorIOUpdatedCallbackHandle.Reset();
 		CurrentGeneratorHandle->RemoveGraphSetCallback(GraphChangedCallbackHandle);
 		GraphChangedCallbackHandle.Reset();
+
 	}
+	ClockHistory = nullptr;
 }
 
 void FMetasoundMusicClockDriver::OnGeneratorAttached()
 {
-	TryToRegisterPlayCursor();
+	WasEverConnected = true;
+	ClockHistory = UMidiClockUpdateSubsystem::GetOrCreateClockHistory(MidiSongPosAnalyzerAddress);
+	SmoothedAudioRenderClockHistoryCursor = ClockHistory->CreateReadCursor();
+	SmoothedPlayerExperienceClockHistoryCursor = ClockHistory->CreateReadCursor();
+	SmoothedVideoRenderClockHistoryCursor = ClockHistory->CreateReadCursor();
+	ClockComponent->MusicClockConnectedEvent.Broadcast();
 }
 
 void FMetasoundMusicClockDriver::OnGraphSet()
 {
-	TryToRegisterPlayCursor();
+	ClockHistory = UMidiClockUpdateSubsystem::GetOrCreateClockHistory(MidiSongPosAnalyzerAddress);
+	SmoothedAudioRenderClockHistoryCursor = ClockHistory->CreateReadCursor();
+	SmoothedPlayerExperienceClockHistoryCursor = ClockHistory->CreateReadCursor();
+	SmoothedVideoRenderClockHistoryCursor = ClockHistory->CreateReadCursor();
 }
 
-void FMetasoundMusicClockDriver::OnGeneratorIOUpdated()
-{
-	// Replaced by OnGeneratorIOUpdatedWithChanges, no-op but still here as it's protected, and child classes may call to it
-	UE_LOG(LogMusicClock, Warning, TEXT("Called to OnGeneratorIOUpdated, which is deprecated and no-op. Use OnGeneratorIOUpdatedWithChanges."), *MetasoundOutputName.ToString());
-}
 
 void FMetasoundMusicClockDriver::OnGeneratorIOUpdatedWithChanges(const TArray<Metasound::FVertexInterfaceChange>& VertexInterfaceChanges)
 {
-	// An output vertex update may have destroyed our Clock, reattach to the new one if it's there
 	if (!MetasoundOutputName.IsNone())
 	{
 		for (auto& VertexInterfaceChange : VertexInterfaceChanges)
 		{
-			if (VertexInterfaceChange.VertexName.IsEqual(MetasoundOutputName))
-			{
-				TryToRegisterPlayCursor();
-				break;
-			}
+			ClockHistory = UMidiClockUpdateSubsystem::GetOrCreateClockHistory(MidiSongPosAnalyzerAddress);
+			SmoothedAudioRenderClockHistoryCursor = ClockHistory->CreateReadCursor();
+			SmoothedPlayerExperienceClockHistoryCursor = ClockHistory->CreateReadCursor();
+			SmoothedVideoRenderClockHistoryCursor = ClockHistory->CreateReadCursor();
+			break;
 		}
 	}
+	
 }
 
+
 void FMetasoundMusicClockDriver::OnGeneratorDetached()
 {
-	ResetCursorOwner();
+	using namespace HarmonixMetasound::Analysis;
+	check(ClockComponent);
+	if (ClockComponent->GetState() != EMusicClockState::Stopped)
+	{
+		if (ClockHistory && CurrentMapChain && CurrentMapChain->SongMaps)
+		{
+			ClockComponent->DefaultMaps.Copy(*(CurrentMapChain->SongMaps), 0, LastTickSeen);
+		}
+		SongPosOffsetMs = ClockComponent->CurrentSmoothedAudioRenderSongPos.SecondsIncludingCountIn * 1000.0f;
+		FreeRunStartTimeSecs = ClockComponent->GetWorld()->GetTimeSeconds();
+	}
+	ClockHistory = nullptr;
+	SmoothedAudioRenderClockHistoryCursor = HarmonixMetasound::Analysis::FMidiClockSongPositionHistory::FReadCursor();
+	SmoothedPlayerExperienceClockHistoryCursor = HarmonixMetasound::Analysis::FMidiClockSongPositionHistory::FReadCursor();
+	SmoothedVideoRenderClockHistoryCursor = HarmonixMetasound::Analysis::FMidiClockSongPositionHistory::FReadCursor();
+
+	ClockComponent->MusicClockDisconnectedEvent.Broadcast();
 }
 
-void FMetasoundMusicClockDriver::TryToRegisterPlayCursor()
+void FMetasoundMusicClockDriver::RefreshCurrentSongPosFromWallClock()
 {
-	check(IsInGameThread());
-	check(Clock);
+	check(ClockComponent);
+
+	bool TempoChanged = ClockComponent->CurrentSmoothedAudioRenderSongPos.Tempo != ClockComponent->Tempo;
 
-	if (!CurrentGeneratorHandle || MetasoundOutputName.IsNone())
+	double FreeRunTime = (ClockComponent->GetWorld()->GetTimeSeconds() - FreeRunStartTimeSecs) * ClockComponent->CurrentClockAdvanceRate;
+
+	ClockComponent->CurrentSmoothedAudioRenderSongPos.SetByTime(((float)FreeRunTime * 1000.0) + SongPosOffsetMs, ClockComponent->DefaultMaps);
+	ClockComponent->CurrentPlayerExperiencedSongPos.SetByTime(ClockComponent->CurrentSmoothedAudioRenderSongPos.SecondsIncludingCountIn * 1000.0f - FHarmonixModule::GetMeasuredUserExperienceAndReactionToAudioRenderOffsetMs(), ClockComponent->DefaultMaps);
+	ClockComponent->CurrentVideoRenderSongPos.SetByTime(ClockComponent->CurrentSmoothedAudioRenderSongPos.SecondsIncludingCountIn * 1000.0f - FHarmonixModule::GetMeasuredVideoToAudioRenderOffsetMs(), ClockComponent->DefaultMaps);
+	if (ClockComponent->CurrentSmoothedAudioRenderSongPos.SecondsIncludingCountIn > ClockComponent->RawUnsmoothedAudioRenderPos.SecondsIncludingCountIn)
 	{
-		return;
+		ClockComponent->RawUnsmoothedAudioRenderPos.SetByTime(ClockComponent->CurrentSmoothedAudioRenderSongPos.SecondsIncludingCountIn * 1000.0f, ClockComponent->DefaultMaps);
 	}
 
-	if (TSharedPtr<Metasound::FMetasoundGenerator> LowLevelGenerator = CurrentGeneratorHandle->GetGenerator())
+	if (TempoChanged)
 	{
-		// Send a command to the OnGenerateAudio thread, where it is safe to interact with the low level generator's output read references.
-		LowLevelGenerator->OnNextBuffer([MetasoundOutputName = MetasoundOutputName, ClockWeakPtr = TWeakObjectPtr<UMusicClockComponent>(Clock), ClockDriverWeakPtr = AsWeak()](Metasound::FMetasoundGenerator& LowLevelGenerator) mutable
-			{
-				// Try to get a cursor manager from the named midi clock output.
-				TSharedPtr<FMidiPlayCursorMgr> MidiPlayCursorMgr;
-				const TOptional<Metasound::TDataReadReference<HarmonixMetasound::FMidiClock>> MidiClockRef = LowLevelGenerator.GetOutputReadReference<HarmonixMetasound::FMidiClock>(MetasoundOutputName);
-				if (const Metasound::TDataReadReference<HarmonixMetasound::FMidiClock>* MidiClock = MidiClockRef.GetPtrOrNull())
-				{
-					MidiPlayCursorMgr = (*MidiClock)->GetDrivingMidiPlayCursorMgr();
-				}
-				else
-				{
-					UE_LOG(LogMusicClock, Verbose, TEXT("Didn't find MIDI Clock output named \"%s\" in the Metasound!"), *MetasoundOutputName.ToString());
-				}
-
-				// Send a command to the game thread, where it is safe to modify clock component state.
-				AsyncTask(ENamedThreads::GameThread, [ClockWeakPtr = MoveTemp(ClockWeakPtr), ClockDriverWeakPtr = MoveTemp(ClockDriverWeakPtr), MidiPlayCursorMgr = MoveTemp(MidiPlayCursorMgr)]() mutable
-					{
-						if (UMusicClockComponent* Clock = ClockWeakPtr.Get())
-						{
-							if (TSharedPtr<FMetasoundMusicClockDriver> ClockDriver = StaticCastSharedPtr<FMetasoundMusicClockDriver>(ClockDriverWeakPtr.Pin()))
-							{
-								// Verify that the music clock component still references the clock driver.
-								if (ClockDriver == Clock->ClockDriver)
-								{
-									ClockDriver->ResetCursorOwner(MoveTemp(MidiPlayCursorMgr));
-								}
-							}
-						}
-					});
-			});
+		ClockComponent->Tempo = ClockComponent->CurrentSmoothedAudioRenderSongPos.Tempo;
+		ClockComponent->CurrentBeatDurationSec = (60.0f / ClockComponent->Tempo) / ClockComponent->CurrentClockAdvanceRate;
+		ClockComponent->CurrentBarDurationSec = ((ClockComponent->TimeSignatureNum * ClockComponent->CurrentBeatDurationSec) / (ClockComponent->TimeSignatureDenom / 4.0f)) / ClockComponent->CurrentClockAdvanceRate;
 	}
 }
 
-void FMetasoundMusicClockDriver::RefreshCurrentSongPosFromWallClock()
+FString FMetasoundMusicClockDriver::HistoryFailureTypeToString(FMetasoundMusicClockDriver::EHistoryFailureType Error)
 {
-	check(Clock);
+	switch (Error)
+	{
+	case FMetasoundMusicClockDriver::EHistoryFailureType::None: return "None";
+	case FMetasoundMusicClockDriver::EHistoryFailureType::NotEnoughDataInTheHistoryRing: return "NotEnoughDataInTheHistoryRing";
+	case FMetasoundMusicClockDriver::EHistoryFailureType::NotEnoughHistory: return "NotEnoughHistory";
+	case FMetasoundMusicClockDriver::EHistoryFailureType::LookingForTimeInTheFutureOfWhatHasEvenRendered: return "LookingForTimeInTheFutureOfWhatHasEvenRendered";
+	case FMetasoundMusicClockDriver::EHistoryFailureType::CaughtUpToRenderPosition: return "CaughtUpToRenderPosition";
+	}
+	return "Unrecognized";
+}
 
-	bool TempoChanged = Clock->CurrentSmoothedAudioRenderSongPos.Tempo != Clock->Tempo;
+void FMetasoundMusicClockDriver::RefreshCurrentSongPosFromHistory()
+{
+	using namespace HarmonixMetasound::Analysis;
 
-	double FreeRunTime = (Clock->GetWorld()->GetTimeSeconds() - FreeRunStartTimeSecs) * Clock->CurrentClockAdvanceRate;
+	check(IsInGameThread());
 
-	Clock->CurrentSmoothedAudioRenderSongPos.SetByTime(((float)FreeRunTime * 1000.0) + SongPosOffsetMs, Clock->DefaultMaps);
-	Clock->CurrentPlayerExperiencedSongPos.SetByTime(Clock->CurrentSmoothedAudioRenderSongPos.SecondsIncludingCountIn * 1000.0f - FHarmonixModule::GetMeasuredUserExperienceAndReactionToAudioRenderOffsetMs(), Clock->DefaultMaps);
-	Clock->CurrentVideoRenderSongPos.SetByTime(Clock->CurrentSmoothedAudioRenderSongPos.SecondsIncludingCountIn * 1000.0f - FHarmonixModule::GetMeasuredVideoToAudioRenderOffsetMs(), Clock->DefaultMaps);
-	if (Clock->CurrentSmoothedAudioRenderSongPos.SecondsIncludingCountIn > Clock->RawUnsmoothedAudioRenderPos.SecondsIncludingCountIn)
+	if (!bRunning || !ClockComponent || !SmoothedAudioRenderClockHistoryCursor.DataAvailable())
 	{
-		Clock->RawUnsmoothedAudioRenderPos.SetByTime(Clock->CurrentSmoothedAudioRenderSongPos.SecondsIncludingCountIn * 1000.0f, Clock->DefaultMaps);
+		return;
+	}
+	
+	if (!SmoothedAudioRenderClockHistoryCursor.Queue)
+	{
+		return;
 	}
 
-	if (TempoChanged)
+	if (!CurrentMapChain || !CurrentMapChain->SongMaps || CurrentMapChain->NewSongMaps)
+	{
+		CurrentMapChain = ClockHistory->GetLatestMapsForConsumer();
+		if (!CurrentMapChain || !CurrentMapChain->SongMaps)
+		{
+			return;
+		}
+	}
+
+	auto Entry = ClockHistory->Positions.GetEntry(ClockHistory->Positions.GetLastWriteIndex());
+	ClockComponent->RawUnsmoothedAudioRenderPos.SetByTick(Entry->Item.UpToTick, *(CurrentMapChain->SongMaps));
+	Metasound::FSampleCount LastRenderPosSampleCount = Entry->Item.SampleCount;
+	float SpeedAtRawRenderTime = Entry->Item.CurrentSpeed;
+	LastTickSeen = Entry->Item.UpToTick;
+	bool bClockIsStopped = Entry->Item.CurrentTransportState != HarmonixMetasound::EMusicPlayerTransportState::Playing;
+
+	if (RenderStartWallClockTimeSeconds == 0.0)
+	{
+		// We are just starting up. We need to create the initial "sync point"...
+		// Wall clock <-> Render Samples...
+		RenderStartSampleCount = LastRenderPosSampleCount;
+		RenderStartWallClockTimeSeconds = ClockComponent->GetWorld()->GetTimeSeconds() - (double)RenderStartSampleCount / (double)ClockHistory->SampleRate;
+		RenderSmoothingLagSeconds = MetasoundMusicClockDriver::SmoothedAudioRenderLagSeconds;
+		ErrorTracker.Reset();
+	}
+
+	double CurrentWallClockSeconds = ClockComponent->GetWorld()->GetTimeSeconds();
+		
+	double ExpectedRenderedSeconds = (CurrentWallClockSeconds - RenderStartWallClockTimeSeconds) * SyncSpeed * MetasoundMusicClockDriver::Fudge;
+	double RenderedSeconds = (double)LastRenderPosSampleCount / (double)ClockHistory->SampleRate;
+	double Error = RenderedSeconds - ExpectedRenderedSeconds;
+
+	if (!bClockIsStopped)
+	{
+		ErrorTracker.Push(Error);
+
+		if (FMath::Abs(ErrorTracker.Min()) > MetasoundMusicClockDriver::MaxErrorSecondsBeforeJump)
+		{
+			UE_LOG(LogMusicClock, Verbose, TEXT("======== MASSIVE ERROR (%f) - SEEKING ==========="), (float)Error);
+			RenderStartSampleCount = LastRenderPosSampleCount;
+			RenderStartWallClockTimeSeconds = ClockComponent->GetWorld()->GetTimeSeconds() - (double)RenderStartSampleCount / (double)ClockHistory->SampleRate;
+			ExpectedRenderedSeconds = RenderedSeconds;
+			RenderSmoothingLagSeconds = MetasoundMusicClockDriver::SmoothedAudioRenderLagSeconds;
+			ErrorTracker.Reset();
+			Error = 0;
+			SyncSpeed = 1.0;
+		}
+
+		// Use proportional part of error to adjust speed ever so slightly...
+		SyncSpeed += MetasoundMusicClockDriver::kP * ErrorTracker.Min();
+		SyncSpeed = FMath::Clamp(SyncSpeed, MetasoundMusicClockDriver::SlowestCorrectionSpeed, MetasoundMusicClockDriver::FastestCorrectionSpeed);
+	}
+
+	//UE_LOG(LogMusicClock, Log, TEXT("Local Minimum Error: %f Average: %f Speed: %f"), (float)ErrorTracker.Min(), (float)ErrorTracker.Average(), (float)SyncSpeed);
+
+	Metasound::FSampleCount ExpectedRenderPosSampleCount = (Metasound::FSampleCount)(ExpectedRenderedSeconds * (double)ClockHistory->SampleRate);
+
+	// first smoothed render time. This is closest to the actual render time. If we catch up to the render time
+	// it means we are rendering in such large blocks that we need to push up our "look behind" number for smoothing.
+	float SpeedAtSmoothRenderedTime = 1.0f;
+	float SmoothedTick = ClockComponent->GetTicksIncludingCountIn(ECalibratedMusicTimebase::AudioRenderTime);
+	EHistoryFailureType Result = CalculateSmoothedTick(ExpectedRenderPosSampleCount, LastRenderPosSampleCount,
+		SmoothedTick, SpeedAtSmoothRenderedTime, SmoothedAudioRenderClockHistoryCursor, RenderSmoothingLagSeconds);
+	if (Result != EHistoryFailureType::None)
 	{
-		Clock->Tempo = Clock->CurrentSmoothedAudioRenderSongPos.Tempo;
-		Clock->CurrentBeatDurationSec = (60.0f / Clock->Tempo) / Clock->CurrentClockAdvanceRate;
-		Clock->CurrentBarDurationSec = ((Clock->TimeSignatureNum * Clock->CurrentBeatDurationSec) / (Clock->TimeSignatureDenom / 4.0f)) / Clock->CurrentClockAdvanceRate;
+		if (LastRenderPosSampleCount >  (RenderSmoothingLagSeconds * ClockHistory->SampleRate * 2))
+		{
+			if (RenderSmoothingLagSeconds < 0.250f)
+			{
+				RenderSmoothingLagSeconds += 0.005f;
+				UE_LOG(LogMusicClock, Verbose, TEXT("(%d) Smoothed Render Time too close to actual render time. Bumping up smoothing lag! (%f)"), LastRenderPosSampleCount, RenderSmoothingLagSeconds);
+			}
+			else
+			{
+				UE_LOG(LogMusicClock, Verbose, TEXT("(%d) Smoothed Render Time too close to actual render time. ALREADY MAX SMOOTHING LAG! (%f)"), LastRenderPosSampleCount, RenderSmoothingLagSeconds);
+			}
+		}
+		else
+		{
+			UE_LOG(LogMusicClock, Verbose, TEXT("(%d) Smoothed Render Time too close to actual render time. WAITING!"), LastRenderPosSampleCount);
+		}
+		return;
+	}
+
+	// Now... We are behind the actual render time because of the lag we introduce to have enough
+	// history... SO... Push forward to get a time that is approx. 30 ms behind the
+	// render time. Why 30ms? Historical reasons. 
+	float SmoothedSongMs = CurrentMapChain->SongMaps->TickToMs(SmoothedTick);
+	SmoothedSongMs += RenderSmoothingLagSeconds * 1000.0f;
+	ClockComponent->CurrentSmoothedAudioRenderSongPos = CalculateSongPosAtMs(SmoothedSongMs);
+
+	// Now the time the user should actually be "experiencing"... ie "hearing" can be calculated as an offset
+	// from the smooth audio rendering time...
+	float LookBehindLagMs = FHarmonixModule::Get().GetMeasuredUserExperienceAndReactionToAudioRenderOffsetMs();
+	float Ms = ClockComponent->CurrentSmoothedAudioRenderSongPos.SecondsIncludingCountIn * 1000.0f;
+	Ms -= LookBehindLagMs; // - (RenderSmoothingLagSeconds * 1000.0f);
+	ClockComponent->CurrentPlayerExperiencedSongPos = CalculateSongPosAtMs(Ms);
+
+	// Now the time we should be rendering graphics for...
+	LookBehindLagMs = FHarmonixModule::Get().GetMeasuredVideoToAudioRenderOffsetMs();
+	Ms = ClockComponent->CurrentSmoothedAudioRenderSongPos.SecondsIncludingCountIn * 1000.0f;
+	Ms -= LookBehindLagMs; // - (RenderSmoothingLagSeconds * 1000.0f);
+	ClockComponent->CurrentVideoRenderSongPos = CalculateSongPosAtMs(Ms);
+
+	ClockComponent->TimeSignatureNum = ClockComponent->CurrentSmoothedAudioRenderSongPos.TimeSigNumerator;
+	ClockComponent->TimeSignatureDenom = ClockComponent->CurrentSmoothedAudioRenderSongPos.TimeSigDenominator;
+
+	if (ClockComponent->Tempo != ClockComponent->CurrentSmoothedAudioRenderSongPos.Tempo || ClockComponent->CurrentClockAdvanceRate != SpeedAtRawRenderTime)
+	{
+		ClockComponent->CurrentClockAdvanceRate = SpeedAtRawRenderTime;
+		ClockComponent->Tempo = ClockComponent->CurrentSmoothedAudioRenderSongPos.Tempo;
+		ClockComponent->CurrentBeatDurationSec = (60.0f / ClockComponent->Tempo) / ClockComponent->CurrentClockAdvanceRate;
+		ClockComponent->CurrentBarDurationSec = ((ClockComponent->TimeSignatureNum * ClockComponent->CurrentBeatDurationSec) / (ClockComponent->TimeSignatureDenom / 4.0f)) / ClockComponent->CurrentClockAdvanceRate;
 	}
 }
 
-void FMetasoundMusicClockDriver::RefreshCurrentSongPosFromCursor()
+FMetasoundMusicClockDriver::EHistoryFailureType FMetasoundMusicClockDriver::CalculateSmoothedTick(Metasound::FSampleCount ExpectedRenderPosSampleCount, Metasound::FSampleCount LastRenderPosSampleCount,
+	float& SmoothedTick, float& CurrentSpeed, HarmonixMetasound::Analysis::FMidiClockSongPositionHistory::FReadCursor& ReadCursor, float LookBehindSeconds)
 {
-	check(IsInGameThread());
-	check(Clock);
-	Clock->CurrentSmoothedAudioRenderSongPos = Cursor.GetCurrentSongPos();
+	using namespace HarmonixMetasound::Analysis;
 
-	float PrevClockAdvanceRate = Clock->CurrentClockAdvanceRate;
-	if (CursorOwner)
+	// A little book keeping for tracking...
+	if (MetasoundMusicClockDriver::HighWaterNumDataAvailable < ReadCursor.NumDataAvailable())
 	{
-		Clock->CurrentClockAdvanceRate = CursorOwner->GetLowResAdvanceRate();
-		Clock->CurrentPlayerExperiencedSongPos = CursorOwner->CalculateLowResSongPosRelativeToCurrentMs(Clock->CurrentSmoothedAudioRenderSongPos.SecondsIncludingCountIn * 1000.0f - FHarmonixModule::GetMeasuredUserExperienceAndReactionToAudioRenderOffsetMs());
-		Clock->CurrentVideoRenderSongPos = CursorOwner->CalculateLowResSongPosRelativeToCurrentMs(Clock->CurrentSmoothedAudioRenderSongPos.SecondsIncludingCountIn * 1000.0f - FHarmonixModule::GetMeasuredVideoToAudioRenderOffsetMs());
-		Clock->RawUnsmoothedAudioRenderPos = CursorOwner->CalculateLowResSongPosWithOffsetMs(0.0f);
+		MetasoundMusicClockDriver::HighWaterNumDataAvailable = ReadCursor.NumDataAvailable();
 	}
-	else
+
+	Metasound::FSampleCount LookingForSampleFrame = ExpectedRenderPosSampleCount - (LookBehindSeconds * ClockHistory->SampleRate);
+
+	int32 NumHistoryAvailable = ReadCursor.NumDataAvailable();
+	if (LookingForSampleFrame >= LastRenderPosSampleCount && NumHistoryAvailable > 1)
 	{
-		Clock->CurrentPlayerExperiencedSongPos.SetByTime(Clock->CurrentSmoothedAudioRenderSongPos.SecondsIncludingCountIn * 1000.0f - FHarmonixModule::Get().GetMeasuredUserExperienceAndReactionToAudioRenderOffsetMs(), Clock->DefaultMaps);
-		Clock->CurrentVideoRenderSongPos.SetByTime(Clock->CurrentSmoothedAudioRenderSongPos.SecondsIncludingCountIn * 1000.0f - FHarmonixModule::Get().GetMeasuredVideoToAudioRenderOffsetMs(), Clock->DefaultMaps);
-		if (Clock->CurrentSmoothedAudioRenderSongPos.SecondsIncludingCountIn > Clock->RawUnsmoothedAudioRenderPos.SecondsIncludingCountIn)
+		while (ReadCursor.NumDataAvailable() > 1)
 		{
-			Clock->RawUnsmoothedAudioRenderPos.SetByTime(Clock->CurrentSmoothedAudioRenderSongPos.SecondsIncludingCountIn * 1000.0f, Clock->DefaultMaps);
+			ReadCursor.ConsumeNext();
 		}
+		NumHistoryAvailable = ReadCursor.NumDataAvailable();
 	}
 
-	Clock->TimeSignatureNum = Clock->CurrentSmoothedAudioRenderSongPos.TimeSigNumerator;
-	Clock->TimeSignatureDenom = Clock->CurrentSmoothedAudioRenderSongPos.TimeSigDenominator;
+	if (NumHistoryAvailable == 0)
+	{
+		return EHistoryFailureType::NotEnoughDataInTheHistoryRing;
+	}
 
-	if (Clock->Tempo != Clock->CurrentSmoothedAudioRenderSongPos.Tempo || PrevClockAdvanceRate != Clock->CurrentClockAdvanceRate)
+	FMidiClockSongPositionHistory::FScopedItemPeekRef PeekNextRef = ReadCursor.PeekNext();
+
+	if (NumHistoryAvailable == 1 || PeekNextRef->SampleCount > LookingForSampleFrame)
 	{
-		Clock->Tempo = Clock->CurrentSmoothedAudioRenderSongPos.Tempo;
-		Clock->CurrentBeatDurationSec = (60.0f / Clock->Tempo) / Clock->CurrentClockAdvanceRate;
-		Clock->CurrentBarDurationSec = ((Clock->TimeSignatureNum * Clock->CurrentBeatDurationSec) / (Clock->TimeSignatureDenom / 4.0f)) / Clock->CurrentClockAdvanceRate;
+		SmoothedTick = (float)PeekNextRef->UpToTick;
+		CurrentSpeed = PeekNextRef->CurrentSpeed;
+		return EHistoryFailureType::None;
 	}
+
+	// OK... or same SHOULD be in the history...
+	FMidiClockSongPositionHistory::FScopedItemPeekRef PeekOneAheadRef(ReadCursor.PeekAhead(1));
+	while (PeekOneAheadRef && PeekOneAheadRef->SampleCount <= LookingForSampleFrame)
+	{
+		ReadCursor.PeekAhead(2, PeekOneAheadRef);
+		ReadCursor.PeekAhead(1, PeekNextRef);
+		ReadCursor.ConsumeNext();
+	}
+
+	// Maybe the sample BEFORE our sample is in the history, but the sample AFTER is not, so we can't lerp?
+	if (!PeekOneAheadRef)
+	{
+		return EHistoryFailureType::CaughtUpToRenderPosition;
+	}
+
+	check(LookingForSampleFrame >= PeekNextRef->SampleCount && LookingForSampleFrame < PeekOneAheadRef->SampleCount);
+
+	// We now have enough to lerp!
+	float LerpAlpha = (float)(LookingForSampleFrame - PeekNextRef->SampleCount) / (float)(PeekOneAheadRef->SampleCount - PeekNextRef->SampleCount);
+	SmoothedTick = FMath::Lerp<float>((float)PeekNextRef->UpToTick, (float)PeekOneAheadRef->UpToTick, LerpAlpha);
+
+	CurrentSpeed = PeekNextRef->CurrentSpeed;
+
+	return EHistoryFailureType::None;
 }
diff --git a/Engine/Plugins/Runtime/Harmonix/Source/HarmonixMetasound/Private/HarmonixMetasound/Components/MetasoundMusicClockDriver.h b/Engine/Plugins/Runtime/Harmonix/Source/HarmonixMetasound/Private/HarmonixMetasound/Components/MetasoundMusicClockDriver.h
--- a/Engine/Plugins/Runtime/Harmonix/Source/HarmonixMetasound/Private/HarmonixMetasound/Components/MetasoundMusicClockDriver.h
+++ b/Engine/Plugins/Runtime/Harmonix/Source/HarmonixMetasound/Private/HarmonixMetasound/Components/MetasoundMusicClockDriver.h
@@ -1,9 +1,16 @@
 // Copyright Epic Games, Inc. All Rights Reserved.
 #pragma once
+#include "Harmonix/LocalMinimumMagnitudeTracker.h"
 #include "HarmonixMetasound/Components/MusicClockComponent.h"
+#include "HarmonixMetasound/Subsystems/MidiClockUpdateSubsystem.h"
 #include "MetasoundGeneratorHandle.h"
 #include "UObject/StrongObjectPtr.h"
 
+namespace HarmonixMetasound::Analysis
+{
+	struct FSongMapChain;
+}
+
 struct FMetasoundMusicClockDriver : public FMusicClockDriverBase
 {
 public:
@@ -16,11 +23,10 @@ public:
 	virtual void Disconnect() override;
 	virtual bool RefreshCurrentSongPos() override;
 	virtual void OnStart() override;
-	virtual void OnPause() override {}
+	virtual void OnPause() override;
 	virtual void OnContinue() override;
-	virtual void OnStop() override {}
-	virtual const FSongMaps* GetCurrentSongMaps() const override;
-
+	virtual void OnStop() override;
+	virtual const ISongMapEvaluator* GetCurrentSongMapEvaluator() const override;
 
 	bool ConnectToAudioComponentsMetasound(UAudioComponent* InAudioComponent, FName MetasoundOuputPinName = "MIDI Clock");
 
@@ -28,8 +34,6 @@ protected:
 	void OnGeneratorAttached();
 	void OnGeneratorDetached();
 	void OnGraphSet();
-	UE_DEPRECATED(5.5, "Use OnGeneratorIOUpdatedWithChanges.")
-	void OnGeneratorIOUpdated();
 	void OnGeneratorIOUpdatedWithChanges(const TArray<Metasound::FVertexInterfaceChange>& VertexInterfaceChanges);
 
 private:
@@ -40,22 +44,54 @@ private:
 	// We need a strong object ptr to this next thing since we will be the only one holding a reference to it...
 	TStrongObjectPtr<UMetasoundGeneratorHandle> CurrentGeneratorHandle;
 
-	FSmoothedMidiPlayCursor Cursor;
-	TSharedPtr<FMidiPlayCursorMgr> CursorOwner;
+	Metasound::Frontend::FAnalyzerAddress MidiSongPosAnalyzerAddress;
+
+	UMidiClockUpdateSubsystem::FClockHistoryPtr ClockHistory;
+	HarmonixMetasound::Analysis::FMidiClockSongPositionHistory::FReadCursor SmoothedAudioRenderClockHistoryCursor;
+	HarmonixMetasound::Analysis::FMidiClockSongPositionHistory::FReadCursor SmoothedPlayerExperienceClockHistoryCursor;
+	HarmonixMetasound::Analysis::FMidiClockSongPositionHistory::FReadCursor SmoothedVideoRenderClockHistoryCursor;
+
+	TSharedPtr<const HarmonixMetasound::Analysis::FSongMapChain> CurrentMapChain;
+	
+	bool bRunning = false;
 	double FreeRunStartTimeSecs = 0.0;
 	bool WasEverConnected = false;
 	float SongPosOffsetMs = 0.0f;
+	int32 LastTickSeen = 0;
+	float RenderSmoothingLagSeconds = 0.030f;
+
+	double RenderStartWallClockTimeSeconds = 0.0;
+	Metasound::FSampleCount RenderStartSampleCount {0};
+
+	static const int kFramesOfErrorHistory = 10;
+	FLocalMinimumMagnitudeTracker<double, kFramesOfErrorHistory> ErrorTracker;
+	double SyncSpeed = 1.0;
 
 	FDelegateHandle GeneratorAttachedCallbackHandle;
 	FDelegateHandle GeneratorDetachedCallbackHandle;
 	FDelegateHandle GeneratorIOUpdatedCallbackHandle;
 	FDelegateHandle GraphChangedCallbackHandle;
 
-	void ResetCursorOwner(TSharedPtr<FMidiPlayCursorMgr> MidiPlayCursorMgr = nullptr);
+	FMidiSongPos CalculateSongPosAtMs(float AbsoluteMs) const;
 
 	bool AttemptToConnectToAudioComponentsMetasound();
 	void DetachAllCallbacks();
-	void TryToRegisterPlayCursor();
 	void RefreshCurrentSongPosFromWallClock();
-	void RefreshCurrentSongPosFromCursor();
+	void RefreshCurrentSongPosFromHistory();
+
+	enum class EHistoryFailureType : uint8
+	{
+		None,
+		NotEnoughDataInTheHistoryRing,
+		NotEnoughHistory,
+		LookingForTimeInTheFutureOfWhatHasEvenRendered,
+		CaughtUpToRenderPosition
+	};
+
+	EHistoryFailureType CalculateSmoothedTick(Metasound::FSampleCount ExpectedRenderPosSampleCount, Metasound::FSampleCount LastRenderPosSampleCount,
+		float& SmoothedTick, float& CurrentSpeed, HarmonixMetasound::Analysis::FMidiClockSongPositionHistory::FReadCursor& ReadCursor,
+		float LookBehindSeconds);
+
+	static FString HistoryFailureTypeToString(FMetasoundMusicClockDriver::EHistoryFailureType Error);
+
 };
diff --git a/Engine/Plugins/Runtime/Harmonix/Source/HarmonixMetasound/Private/HarmonixMetasound/Components/MusicClockComponent.cpp b/Engine/Plugins/Runtime/Harmonix/Source/HarmonixMetasound/Private/HarmonixMetasound/Components/MusicClockComponent.cpp
--- a/Engine/Plugins/Runtime/Harmonix/Source/HarmonixMetasound/Private/HarmonixMetasound/Components/MusicClockComponent.cpp
+++ b/Engine/Plugins/Runtime/Harmonix/Source/HarmonixMetasound/Private/HarmonixMetasound/Components/MusicClockComponent.cpp
@@ -160,6 +160,10 @@ bool UMusicClockComponent::ConnectToMetasound()
 	TSharedPtr<FMetasoundMusicClockDriver> MetasoundClockDriver = MakeShared<FMetasoundMusicClockDriver>(this);
 	bool Connected = MetasoundClockDriver->ConnectToAudioComponentsMetasound(MetasoundsAudioComponent, MetasoundOutputName);
 	ClockDriver = MoveTemp(MetasoundClockDriver);
+	if (State == EMusicClockState::Running)
+	{
+		ClockDriver->OnStart();
+	}
 	return Connected;
 }
 
@@ -346,6 +350,21 @@ float UMusicClockComponent::GetBeatsIncludingCountIn(ECalibratedMusicTimebase Ti
 	}
 }
 
+float UMusicClockComponent::GetTicksFromBarOne(ECalibratedMusicTimebase Timebase) const
+{
+	EnsureClockIsValidForGameFrame();
+	float Seconds = GetSecondsFromBarOne(Timebase);
+	return GetSongMaps().MsToTick(Seconds * 1000.0f);
+}
+
+float UMusicClockComponent::GetTicksIncludingCountIn(ECalibratedMusicTimebase Timebase) const
+{
+	EnsureClockIsValidForGameFrame();
+	float Seconds = GetSecondsIncludingCountIn(Timebase);
+	return GetSongMaps().MsToTick(Seconds * 1000.0f);
+}
+
+
 FMusicTimestamp UMusicClockComponent::GetCurrentTimestamp(ECalibratedMusicTimebase Timebase) const
 {
 	EnsureClockIsValidForGameFrame();
@@ -367,7 +386,7 @@ FString UMusicClockComponent::GetCurrentSectionName(ECalibratedMusicTimebase Tim
 int32 UMusicClockComponent::GetCurrentSectionIndex(ECalibratedMusicTimebase Timebase) const
 {
 	const FMidiSongPos& SongPos = GetSongPos(Timebase);
-	return GetSongMaps().GetSectionMap().TickToSectionIndex(SongPos.CurrentSongSection.StartTick);
+	return GetSongMaps().GetSectionIndexAtTick(SongPos.CurrentSongSection.StartTick);
 }
 
 float UMusicClockComponent::GetCurrentSectionStartMs(ECalibratedMusicTimebase Timebase) const
@@ -438,7 +457,7 @@ float UMusicClockComponent::GetDeltaBeat(ECalibratedMusicTimebase Timebase) cons
 
 const TArray<FSongSection>& UMusicClockComponent::GetSongSections() const
 {
-	return GetSongMaps().GetSectionMap().GetSections();
+	return GetSongMaps().GetSections();
 }
 
 float UMusicClockComponent::GetCountInSeconds() const
@@ -508,7 +527,7 @@ float UMusicClockComponent::GetSectionEndMsAtMs(float Ms) const
 
 int32 UMusicClockComponent::GetNumSections() const
 {
-	return GetSongMaps().GetSectionMap().GetNumSections();
+	return GetSongMaps().GetNumSections();
 }
 
 float UMusicClockComponent::GetSongLengthMs() const
@@ -532,9 +551,9 @@ float UMusicClockComponent::GetSongRemainingMs(ECalibratedMusicTimebase Timebase
 	return SongLengthMs <= 0.f ? 0.f : SongLengthMs - (GetSongPos(Timebase).SecondsIncludingCountIn * 1000.0f);
 }
 
-const FSongMaps& UMusicClockComponent::GetSongMaps() const
+const ISongMapEvaluator& UMusicClockComponent::GetSongMaps() const
 {
-	const FSongMaps* SongMaps = ClockDriver ? ClockDriver->GetCurrentSongMaps() : nullptr;
+	const ISongMapEvaluator* SongMaps = ClockDriver ? ClockDriver->GetCurrentSongMapEvaluator() : nullptr;
 	return SongMaps ? *SongMaps : DefaultMaps;
 }
 
@@ -573,13 +592,24 @@ float UMusicClockComponent::MeasureSpanProgress(const FMusicalTimeSpan& Span, EC
 {
 	EnsureClockIsValidForGameFrame();
 
-	const FSongMaps* Maps = ClockDriver ? ClockDriver->GetCurrentSongMaps() : &DefaultMaps;
+	const ISongMapEvaluator* Maps = ClockDriver ? ClockDriver->GetCurrentSongMapEvaluator() : &DefaultMaps;
 	if (!Maps)
 	{
 		return 0.0f;
 	}
 
-	return Span.CalcPositionInSpan(CurrentSmoothedAudioRenderSongPos, *Maps);
+	const FMidiSongPos* Basis = &CurrentVideoRenderSongPos;
+	switch (Timebase)
+	{
+	case ECalibratedMusicTimebase::AudioRenderTime:
+		Basis = &CurrentSmoothedAudioRenderSongPos;
+		break;
+	case ECalibratedMusicTimebase::ExperiencedTime:
+		Basis = &CurrentPlayerExperiencedSongPos;
+		break;
+	}
+
+	return Span.CalcPositionInSpan(*Basis, *Maps);
 }
 
 void UMusicClockComponent::BroadcastSongPosChanges()
@@ -628,18 +658,18 @@ void FMusicClockDriverBase::EnsureClockIsValidForGameFrame()
 	// to make sure its current state is appropriate to the current musical time. See 
 	// UMusicClockComponent::EnsureClockIsValidForGameFrame for more details as to why 
 	// this is so.
-	Clock->PrevAudioRenderSongPos = Clock->CurrentSmoothedAudioRenderSongPos;
-	Clock->PrevPlayerExperiencedSongPos = Clock->CurrentPlayerExperiencedSongPos;
-	Clock->PrevVideoRenderSongPos = Clock->CurrentVideoRenderSongPos;
+	ClockComponent->PrevAudioRenderSongPos = ClockComponent->CurrentSmoothedAudioRenderSongPos;
+	ClockComponent->PrevPlayerExperiencedSongPos = ClockComponent->CurrentPlayerExperiencedSongPos;
+	ClockComponent->PrevVideoRenderSongPos = ClockComponent->CurrentVideoRenderSongPos;
 
 	if (RefreshCurrentSongPos())
 	{
-		Clock->AudioRenderDeltaBarF = Clock->CurrentSmoothedAudioRenderSongPos.BarsIncludingCountIn - Clock->PrevAudioRenderSongPos.BarsIncludingCountIn;
-		Clock->AudioRenderDeltaBeatF = Clock->CurrentSmoothedAudioRenderSongPos.BeatsIncludingCountIn - Clock->PrevAudioRenderSongPos.BeatsIncludingCountIn;
-		Clock->PlayerExperienceDeltaBarF = Clock->CurrentPlayerExperiencedSongPos.BarsIncludingCountIn - Clock->PrevPlayerExperiencedSongPos.BarsIncludingCountIn;
-		Clock->PlayerExperienceDeltaBeatF = Clock->CurrentPlayerExperiencedSongPos.BeatsIncludingCountIn - Clock->PrevPlayerExperiencedSongPos.BeatsIncludingCountIn;
-		Clock->VideoRenderDeltaBarF = Clock->CurrentVideoRenderSongPos.BarsIncludingCountIn - Clock->PrevVideoRenderSongPos.BarsIncludingCountIn;
-		Clock->VideoRenderDeltaBeatF = Clock->CurrentVideoRenderSongPos.BeatsIncludingCountIn - Clock->PrevVideoRenderSongPos.BeatsIncludingCountIn;
-		Clock->LastUpdateFrame = GFrameCounter;
+		ClockComponent->AudioRenderDeltaBarF = ClockComponent->CurrentSmoothedAudioRenderSongPos.BarsIncludingCountIn - ClockComponent->PrevAudioRenderSongPos.BarsIncludingCountIn;
+		ClockComponent->AudioRenderDeltaBeatF = ClockComponent->CurrentSmoothedAudioRenderSongPos.BeatsIncludingCountIn - ClockComponent->PrevAudioRenderSongPos.BeatsIncludingCountIn;
+		ClockComponent->PlayerExperienceDeltaBarF = ClockComponent->CurrentPlayerExperiencedSongPos.BarsIncludingCountIn - ClockComponent->PrevPlayerExperiencedSongPos.BarsIncludingCountIn;
+		ClockComponent->PlayerExperienceDeltaBeatF = ClockComponent->CurrentPlayerExperiencedSongPos.BeatsIncludingCountIn - ClockComponent->PrevPlayerExperiencedSongPos.BeatsIncludingCountIn;
+		ClockComponent->VideoRenderDeltaBarF = ClockComponent->CurrentVideoRenderSongPos.BarsIncludingCountIn - ClockComponent->PrevVideoRenderSongPos.BarsIncludingCountIn;
+		ClockComponent->VideoRenderDeltaBeatF = ClockComponent->CurrentVideoRenderSongPos.BeatsIncludingCountIn - ClockComponent->PrevVideoRenderSongPos.BeatsIncludingCountIn;
+		ClockComponent->LastUpdateFrame = GFrameCounter;
 	}
 }
diff --git a/Engine/Plugins/Runtime/Harmonix/Source/HarmonixMetasound/Private/HarmonixMetasound/Components/WallClockMusicClockDriver.cpp b/Engine/Plugins/Runtime/Harmonix/Source/HarmonixMetasound/Private/HarmonixMetasound/Components/WallClockMusicClockDriver.cpp
--- a/Engine/Plugins/Runtime/Harmonix/Source/HarmonixMetasound/Private/HarmonixMetasound/Components/WallClockMusicClockDriver.cpp
+++ b/Engine/Plugins/Runtime/Harmonix/Source/HarmonixMetasound/Private/HarmonixMetasound/Components/WallClockMusicClockDriver.cpp
@@ -6,23 +6,23 @@
 bool FWallClockMusicClockDriver::CalculateSongPosWithOffset(float MsOffset, ECalibratedMusicTimebase Timebase, FMidiSongPos& OutResult) const
 {
 	check(IsInGameThread());
-	if (!TempoMap.IsValid())
+	if (!TempoMapMidi.IsValid())
 	{
 		return false;
 	}
 
-	const FSongMaps* Maps = TempoMap->GetSongMaps();
+	const FSongMaps* Maps = TempoMapMidi->GetSongMaps();
 	switch (Timebase)
 	{
 	case ECalibratedMusicTimebase::AudioRenderTime:
-		OutResult.SetByTime((Clock->CurrentSmoothedAudioRenderSongPos.SecondsIncludingCountIn * 1000.0f) + MsOffset, *Maps);
+		OutResult.SetByTime((ClockComponent->CurrentSmoothedAudioRenderSongPos.SecondsIncludingCountIn * 1000.0f) + MsOffset, *Maps);
 		break;
 	case ECalibratedMusicTimebase::ExperiencedTime:
-		OutResult.SetByTime((Clock->CurrentPlayerExperiencedSongPos.SecondsIncludingCountIn * 1000.0f) + MsOffset, *Maps);
+		OutResult.SetByTime((ClockComponent->CurrentPlayerExperiencedSongPos.SecondsIncludingCountIn * 1000.0f) + MsOffset, *Maps);
 		break;
 	case ECalibratedMusicTimebase::VideoRenderTime:
 	default:
-		OutResult.SetByTime((Clock->CurrentVideoRenderSongPos.SecondsIncludingCountIn * 1000.0f) + MsOffset, *Maps);
+		OutResult.SetByTime((ClockComponent->CurrentVideoRenderSongPos.SecondsIncludingCountIn * 1000.0f) + MsOffset, *Maps);
 		break;
 	}
 
@@ -31,31 +31,31 @@ bool FWallClockMusicClockDriver::CalculateSongPosWithOffset(float MsOffset, ECal
 
 void FWallClockMusicClockDriver::Disconnect()
 {
-	TempoMap = nullptr;
+	TempoMapMidi = nullptr;
 }
 
 bool FWallClockMusicClockDriver::RefreshCurrentSongPos()
 {
 	check(IsInGameThread());
-	check(Clock);
-	check(Clock->GetWorld());
+	check(ClockComponent);
+	check(ClockComponent->GetWorld());
 
-	bool TempoChanged = Clock->CurrentSmoothedAudioRenderSongPos.Tempo != Clock->Tempo;
+	bool TempoChanged = ClockComponent->CurrentSmoothedAudioRenderSongPos.Tempo != ClockComponent->Tempo;
 
-	double RunTime = Clock->GetWorld()->GetTimeSeconds() - StartTimeSecs;
+	double RunTime = ClockComponent->GetWorld()->GetTimeSeconds() - StartTimeSecs;
 
-	const FSongMaps* Maps = GetCurrentSongMaps();
+	const ISongMapEvaluator* Maps = GetCurrentSongMapEvaluator();
 	check(Maps);
 
-	Clock->CurrentSmoothedAudioRenderSongPos.SetByTime((float)(RunTime * 1000.0), *Maps);
-	Clock->CurrentPlayerExperiencedSongPos.SetByTime(Clock->CurrentSmoothedAudioRenderSongPos.SecondsIncludingCountIn * 1000.0f - FHarmonixModule::GetMeasuredUserExperienceAndReactionToAudioRenderOffsetMs(), *Maps);
-	Clock->CurrentVideoRenderSongPos.SetByTime(Clock->CurrentSmoothedAudioRenderSongPos.SecondsIncludingCountIn * 1000.0f - FHarmonixModule::GetMeasuredVideoToAudioRenderOffsetMs(), *Maps);
+	ClockComponent->CurrentSmoothedAudioRenderSongPos.SetByTime((float)(RunTime * 1000.0), *Maps);
+	ClockComponent->CurrentPlayerExperiencedSongPos.SetByTime(ClockComponent->CurrentSmoothedAudioRenderSongPos.SecondsIncludingCountIn * 1000.0f - FHarmonixModule::GetMeasuredUserExperienceAndReactionToAudioRenderOffsetMs(), *Maps);
+	ClockComponent->CurrentVideoRenderSongPos.SetByTime(ClockComponent->CurrentSmoothedAudioRenderSongPos.SecondsIncludingCountIn * 1000.0f - FHarmonixModule::GetMeasuredVideoToAudioRenderOffsetMs(), *Maps);
 
 	if (TempoChanged)
 	{
-		Clock->Tempo = Clock->CurrentSmoothedAudioRenderSongPos.Tempo;
-		Clock->CurrentBeatDurationSec = (60.0f / Clock->Tempo) / Clock->CurrentClockAdvanceRate;
-		Clock->CurrentBarDurationSec = ((Clock->TimeSignatureNum * Clock->CurrentBeatDurationSec) / (Clock->TimeSignatureDenom / 4.0f)) / Clock->CurrentClockAdvanceRate;
+		ClockComponent->Tempo = ClockComponent->CurrentSmoothedAudioRenderSongPos.Tempo;
+		ClockComponent->CurrentBeatDurationSec = (60.0f / ClockComponent->Tempo) / ClockComponent->CurrentClockAdvanceRate;
+		ClockComponent->CurrentBarDurationSec = ((ClockComponent->TimeSignatureNum * ClockComponent->CurrentBeatDurationSec) / (ClockComponent->TimeSignatureDenom / 4.0f)) / ClockComponent->CurrentClockAdvanceRate;
 	}
 
 	return true;
@@ -65,35 +65,35 @@ bool FWallClockMusicClockDriver::RefreshCurrentSongPos()
 void FWallClockMusicClockDriver::OnStart()
 {
 	check(IsInGameThread());
-	check(Clock);
-	StartTimeSecs = Clock->GetWorld()->GetTimeSeconds();
+	check(ClockComponent);
+	StartTimeSecs = ClockComponent->GetWorld()->GetTimeSeconds();
 	PauseTimeSecs = 0.0;
 }
 
 void FWallClockMusicClockDriver::OnPause()
 {
 	check(IsInGameThread());
-	check(Clock);
-	PauseTimeSecs = Clock->GetWorld()->GetTimeSeconds();
+	check(ClockComponent);
+	PauseTimeSecs = ClockComponent->GetWorld()->GetTimeSeconds();
 }
 
 void FWallClockMusicClockDriver::OnContinue()
 {
 	check(IsInGameThread());
-	check(Clock);
-	double CurrentTime = Clock->GetWorld()->GetTimeSeconds();
+	check(ClockComponent);
+	double CurrentTime = ClockComponent->GetWorld()->GetTimeSeconds();
 	StartTimeSecs += (CurrentTime - PauseTimeSecs);
 	PauseTimeSecs = 0.0;
 	RefreshCurrentSongPos();
 }
 
-const FSongMaps* FWallClockMusicClockDriver::GetCurrentSongMaps() const
+const ISongMapEvaluator* FWallClockMusicClockDriver::GetCurrentSongMapEvaluator() const
 {
-	if (TempoMap.IsValid())
+	check(IsInGameThread());
+	if (TempoMapMidi.IsValid())
 	{
-		return TempoMap->GetSongMaps();
+		return TempoMapMidi->GetSongMaps();
 	}
-	check(IsInGameThread());
-	check(Clock);
-	return &Clock->DefaultMaps;
+	check(ClockComponent);
+	return &ClockComponent->DefaultMaps;
 }
diff --git a/Engine/Plugins/Runtime/Harmonix/Source/HarmonixMetasound/Private/HarmonixMetasound/Components/WallClockMusicClockDriver.h b/Engine/Plugins/Runtime/Harmonix/Source/HarmonixMetasound/Private/HarmonixMetasound/Components/WallClockMusicClockDriver.h
--- a/Engine/Plugins/Runtime/Harmonix/Source/HarmonixMetasound/Private/HarmonixMetasound/Components/WallClockMusicClockDriver.h
+++ b/Engine/Plugins/Runtime/Harmonix/Source/HarmonixMetasound/Private/HarmonixMetasound/Components/WallClockMusicClockDriver.h
@@ -9,7 +9,7 @@ struct FWallClockMusicClockDriver : public FMusicClockDriverBase
 public:
 	FWallClockMusicClockDriver(UMusicClockComponent* InClock, UMidiFile* InTempoMap)
 		: FMusicClockDriverBase(InClock)
-		, TempoMap(InTempoMap)
+		, TempoMapMidi(InTempoMap)
 	{}
 
 	virtual bool CalculateSongPosWithOffset(float MsOffset, ECalibratedMusicTimebase Timebase, FMidiSongPos& OutResult) const override;
@@ -20,10 +20,10 @@ public:
 	virtual void OnPause() override;
 	virtual void OnContinue() override;
 	virtual void OnStop() override {}
-	virtual const FSongMaps* GetCurrentSongMaps() const override;
+	virtual const ISongMapEvaluator* GetCurrentSongMapEvaluator() const override;
 
 private:
-	TWeakObjectPtr<UMidiFile> TempoMap;
+	TWeakObjectPtr<UMidiFile> TempoMapMidi;
 
 	double StartTimeSecs = 0.0;
 	double PauseTimeSecs = 0.0f;
