Massive Midi Clock Refactor

diff --git a/Engine/Plugins/Runtime/Harmonix/Source/HarmonixMetasound/Private/HarmonixMetasound/DataTypes/MidiClock.cpp b/Engine/Plugins/Runtime/Harmonix/Source/HarmonixMetasound/Private/HarmonixMetasound/DataTypes/MidiClock.cpp
--- a/Engine/Plugins/Runtime/Harmonix/Source/HarmonixMetasound/Private/HarmonixMetasound/DataTypes/MidiClock.cpp
+++ b/Engine/Plugins/Runtime/Harmonix/Source/HarmonixMetasound/Private/HarmonixMetasound/DataTypes/MidiClock.cpp
@@ -38,574 +38,536 @@ REGISTER_METASOUND_DATATYPE(HarmonixMetasound::FMidiClock, "MIDIClock")
 namespace HarmonixMetasound
 {
 	using namespace Metasound;
-	
+
 	FMidiClock::FMidiClock(const FOperatorSettings& InSettings)
-		: MidiClockEventCursor(this)
+		: SongMapEvaluator(MakeShared<FSongMapsWithAlternateTempoSource>(MakeClockConductorMidiData(120.0f, 4, 4)))
+		, CurrentTempoInfoPointIndex(0)
+		, CurrentTimeSignaturePointIndex(0)
+		, TickResidualWhenDriven(0.0f)
 		, BlockSize(InSettings.GetNumFramesPerBlock())
 		, CurrentBlockFrameIndex(0)
+		, LastProcessedMidiTick(-1)
+		, NextMidiTickToProcess(0)
 		, SampleRate(InSettings.GetSampleRate())
-		, HasSpeedChangeInBlock(false)
-		, HasTempoChangeInBlock(false)
-		, DrivingMidiPlayCursorMgr(MakeShared<FMidiPlayCursorMgr>())
+		, SampleCount(0)
+		, FramesUntilNextProcess(0)
+		, TransportAtBlockStart(EMusicPlayerTransportState::Invalid)
+		, TransportAtBlockEnd(EMusicPlayerTransportState::Invalid)
+		, SpeedAtBlockStart(0.0f)
+		, SpeedAtBlockEnd(0.0f)
+		, CurrentLocalSpeed(-1.0f)
+		, TempoAtBlockStart(0.0f)
+		, TempoAtBlockEnd(0.0f)
+		, TimeSignatureAtBlockStart(FTimeSignature(0,1))
+		, TimeSignatureAtBlockEnd(FTimeSignature(0,1))
+		, NumTransportChangeInBlock(0)
+		, NumSpeedChangeInBlock(0)
+		, NumTempoChangeInBlock(0)
+		, NumTimeSignatureChangeInBlock(0)
+		, NextTempoChangeTick(std::numeric_limits<int32>::max())
+		, NextTimeSigChangeTick(std::numeric_limits<int32>::max())
+		, NextTempoOrTimeSigChangeTick(std::numeric_limits<int32>::max())
+		, FirstTickInLoop(-1)
+		, LoopLengthTicks(0)
+		, MidiDataChangedInBlock(false)
 	{
-		SpeedChangesInBlock.Add({0, 0.0f, 1.0f});
-		TempoChangesInBlock.Add({0, 0.0f, 120.0f});
-		
-		DrivingMidiPlayCursorMgr->RegisterHiResPlayCursor(&MidiClockEventCursor);
-
-		RegisterForGameThreadUpdates();
 	}
 
 	FMidiClock::FMidiClock(const FMidiClock& Other)
-		: MidiClockEventCursor(this)
+		: TSharedFromThis<FMidiClock, ESPMode::NotThreadSafe>(Other)
+		, SongMapEvaluator(MakeShared<FSongMapsWithAlternateTempoSource>(Other.SongMapEvaluator))
+		, CurrentTempoInfoPointIndex(Other.CurrentTempoInfoPointIndex)
+		, CurrentTimeSignaturePointIndex(Other.CurrentTimeSignaturePointIndex)
+		, ExternalClockDriver(Other.ExternalClockDriver)
+		, TickResidualWhenDriven(Other.TickResidualWhenDriven)
 		, BlockSize(Other.BlockSize)
 		, CurrentBlockFrameIndex(Other.CurrentBlockFrameIndex)
+		, LastProcessedMidiTick(Other.LastProcessedMidiTick)
+		, NextMidiTickToProcess(Other.NextMidiTickToProcess)
 		, SampleRate(Other.SampleRate)
 		, SampleCount(Other.SampleCount)
 		, FramesUntilNextProcess(Other.FramesUntilNextProcess)
-		, CurrentTransportState(Other.CurrentTransportState)
-		, TransportChangesInBlock(Other.TransportChangesInBlock)
-		, HasSpeedChangeInBlock(Other.HasSpeedChangeInBlock)
-		, SpeedChangesInBlock(Other.SpeedChangesInBlock)
-		, HasTempoChangeInBlock(Other.HasTempoChangeInBlock)
-		, TempoChangesInBlock(Other.TempoChangesInBlock)
+		, TransportAtBlockStart(Other.TransportAtBlockStart)
+		, TransportAtBlockEnd(Other.TransportAtBlockEnd)
+		, SpeedAtBlockStart(Other.SpeedAtBlockStart)
+		, SpeedAtBlockEnd(Other.SpeedAtBlockEnd)
+		, CurrentLocalSpeed(Other.CurrentLocalSpeed)
+		, TempoAtBlockStart(Other.TempoAtBlockStart)
+		, TempoAtBlockEnd(Other.TempoAtBlockEnd)
+		, TimeSignatureAtBlockStart(Other.TimeSignatureAtBlockStart)
+		, TimeSignatureAtBlockEnd(Other.TimeSignatureAtBlockEnd)
+		, NumTransportChangeInBlock(Other.NumTransportChangeInBlock)
+		, NumSpeedChangeInBlock(Other.NumSpeedChangeInBlock)
+		, NumTempoChangeInBlock(Other.NumTempoChangeInBlock)
+		, NumTimeSignatureChangeInBlock(Other.NumTimeSignatureChangeInBlock)
+		, NextTempoChangeTick(Other.NextTempoChangeTick)
+		, NextTimeSigChangeTick(Other.NextTimeSigChangeTick)
+		, NextTempoOrTimeSigChangeTick(Other.NextTempoOrTimeSigChangeTick)
+		, FirstTickInLoop(Other.FirstTickInLoop)
+		, LoopLengthTicks(Other.LoopLengthTicks)
+		, MidiDataChangedInBlock(Other.MidiDataChangedInBlock)
 		, MidiClockEventsInBlock(Other.MidiClockEventsInBlock)
-		, SmoothingEnabled(Other.SmoothingEnabled)
-		, DrivingMidiPlayCursorMgr(Other.DrivingMidiPlayCursorMgr)
 	{
-		DrivingMidiPlayCursorMgr->RegisterHiResPlayCursor(&MidiClockEventCursor);
+	}
 
-		RegisterForGameThreadUpdates();
+	FMidiClock::~FMidiClock()
+	{
 	}
 
 	FMidiClock& FMidiClock::operator=(const FMidiClock& Other)
 	{
 		if (this != &Other)
 		{
-			UnregisterForGameThreadUpdates();
-
-			if (DrivingMidiPlayCursorMgr != Other.DrivingMidiPlayCursorMgr)
-			{
-				DrivingMidiPlayCursorMgr->UnregisterPlayCursor(&MidiClockEventCursor);
-			}
-
+			SongMapEvaluator = Other.SongMapEvaluator;
+			CurrentTempoInfoPointIndex = Other.CurrentTempoInfoPointIndex;
+			CurrentTimeSignaturePointIndex = Other.CurrentTimeSignaturePointIndex;
+			ExternalClockDriver = Other.ExternalClockDriver;
+			TickResidualWhenDriven = Other.TickResidualWhenDriven;
 			BlockSize = Other.BlockSize;
 			CurrentBlockFrameIndex = Other.CurrentBlockFrameIndex;
+			LastProcessedMidiTick = Other.LastProcessedMidiTick;
+			NextMidiTickToProcess = Other.NextMidiTickToProcess;
 			SampleRate = Other.SampleRate;
 			SampleCount = Other.SampleCount;
 			FramesUntilNextProcess = Other.FramesUntilNextProcess;
-			CurrentTransportState = Other.CurrentTransportState;
-			TransportChangesInBlock = Other.TransportChangesInBlock;
-			HasSpeedChangeInBlock = Other.HasSpeedChangeInBlock;
-			SpeedChangesInBlock = Other.SpeedChangesInBlock;
-			HasTempoChangeInBlock = Other.HasTempoChangeInBlock;
-			TempoChangesInBlock = Other.TempoChangesInBlock;
+			TransportAtBlockStart = Other.TransportAtBlockStart;
+			TransportAtBlockEnd = Other.TransportAtBlockEnd;
+			SpeedAtBlockStart = Other.SpeedAtBlockStart;
+			SpeedAtBlockEnd = Other.SpeedAtBlockEnd;
+			CurrentLocalSpeed = Other.CurrentLocalSpeed;
+			TempoAtBlockStart = Other.TempoAtBlockStart;
+			TempoAtBlockEnd = Other.TempoAtBlockEnd;
+			TimeSignatureAtBlockStart = Other.TimeSignatureAtBlockStart;
+			TimeSignatureAtBlockEnd = Other.TimeSignatureAtBlockEnd;
+			NumTransportChangeInBlock = Other.NumTransportChangeInBlock;
+			NumSpeedChangeInBlock = Other.NumSpeedChangeInBlock;
+			NumTempoChangeInBlock = Other.NumTempoChangeInBlock;
+			NextTempoChangeTick = Other.NextTempoChangeTick;
+			NextTimeSigChangeTick = Other.NextTimeSigChangeTick;
+			NextTempoOrTimeSigChangeTick = Other.NextTempoOrTimeSigChangeTick;
+			FirstTickInLoop = Other.FirstTickInLoop;
+			LoopLengthTicks = Other.LoopLengthTicks;
+			MidiDataChangedInBlock = Other.MidiDataChangedInBlock;
 			MidiClockEventsInBlock = Other.MidiClockEventsInBlock;
-			SmoothingEnabled = Other.SmoothingEnabled;
-
-			if (DrivingMidiPlayCursorMgr != Other.DrivingMidiPlayCursorMgr)
-			{
-				DrivingMidiPlayCursorMgr = Other.DrivingMidiPlayCursorMgr;
-				DrivingMidiPlayCursorMgr->RegisterHiResPlayCursor(&MidiClockEventCursor);
-			}
-
-			RegisterForGameThreadUpdates();
 		}
 
 		return *this;
 	}
 
-	FMidiClock::~FMidiClock()
-	{
-		UnregisterForGameThreadUpdates();
-	
-		DrivingMidiPlayCursorMgr->UnregisterPlayCursor(&MidiClockEventCursor, false);
-	}
-
-	void FMidiClock::RegisterForGameThreadUpdates()
-	{
-		UMidiClockUpdateSubsystem::TrackMidiClock(this);
-	}
-
-	void FMidiClock::UnregisterForGameThreadUpdates()
-	{
-		UMidiClockUpdateSubsystem::StopTrackingMidiClock(this);
-	}
-
-	void FMidiClock::ResetAndStart(int32 FrameIndex, bool SeekToStart)
+	void FMidiClock::AttachToMidiFile(TSharedPtr<FMidiFileData> MidiData, bool ResetToStart)
 	{
-		if (SeekToStart)
+		MidiDataChangedInBlock = true;
+		if (!MidiData)
 		{
-			SampleCount = 0;
-			DrivingMidiPlayCursorMgr->SeekTo(0, 0, true, false);
+			MidiData = MakeClockConductorMidiData(TempoAtBlockEnd, TimeSignatureAtBlockEnd.Numerator, TimeSignatureAtBlockEnd.Denominator);
 		}
-	
-		AddTransportStateChangeToBlock({FrameIndex,0.0f, EMusicPlayerTransportState::Playing});
-		
-		HasSpeedChangeInBlock = false;
-		SpeedChangesInBlock.SetNum(1);
-		SpeedChangesInBlock[0] = {0, 0.0f, 1.0f};
-
-		HasTempoChangeInBlock = false;
-		TempoChangesInBlock.SetNum(1);
-		TempoChangesInBlock[0] = { 0, 0.0f, 120.0f };
-
-		CurrentBlockFrameIndex = FrameIndex;
-		FramesUntilNextProcess = 0;
-	}
 
-	void FMidiClock::PrepareBlock()
-	{
-		TransportChangesInBlock.Empty(4);
-		if (SpeedChangesInBlock.Num() > 1)
+		if (ExternalClockDriver)
 		{
-			SpeedChangesInBlock[0].Speed = SpeedChangesInBlock.Last().Speed;
-			SpeedChangesInBlock.SetNum(1, EAllowShrinking::No);
+			// use the tempo from the external clock...
+			RebuildSongMapEvaluator(ExternalClockDriver->SongMapEvaluator->GetMidiFileWithTempoMap(), MidiData);
 		}
-		HasSpeedChangeInBlock = false;
-		if (TempoChangesInBlock.Num() > 1)
+		else
 		{
-			TempoChangesInBlock[0].Tempo = TempoChangesInBlock.Last().Tempo;
-			TempoChangesInBlock.SetNum(1, EAllowShrinking::No);
+			// All maps come from the same source...
+			RebuildSongMapEvaluator(MidiData, MidiData);
 		}
-		HasTempoChangeInBlock = false;
-		CurrentBlockFrameIndex = 0;
 
-		MidiClockEventsInBlock.Reset();
+		if (ResetToStart)
+		{
+			SeekTo(CurrentBlockFrameIndex, 0);
+		}
+		PostTempoOrTimeSignatureEventsIfNeeded();
 	}
 
-	bool FMidiClock::HasLowResCursors() const
+	void FMidiClock::MidiChanged()
 	{
-		return DrivingMidiPlayCursorMgr->HasLowResCursors();
+		MidiDataChangedInBlock = true;
+		PostTempoOrTimeSignatureEventsIfNeeded();
 	}
 
-	void FMidiClock::UpdateLowResCursors()
+	void FMidiClock::DetachFromMidiFile()
 	{
-		DrivingMidiPlayCursorMgr->AdvanceLowResCursors();
+		AttachToMidiFile(nullptr, false);
 	}
 
-	void FMidiClock::AddTransportStateChangeToBlock(const FMidiTimestampTransportState& NewTransportState)
+	void FMidiClock::SetDrivingClock(FConstSharedMidiClockPtr NewExternalClockDriver)
 	{
-		if (TransportChangesInBlock.IsEmpty() || TransportChangesInBlock.Last().BlockSampleFrameIndex <= NewTransportState.BlockSampleFrameIndex)
-		{
-			CurrentTransportState = NewTransportState;
-			TransportChangesInBlock.Add(NewTransportState);
-		}
+		MidiDataChangedInBlock = true;
+		ExternalClockDriver = NewExternalClockDriver;
+		RebuildSongMapEvaluator(ExternalClockDriver ? 
+									ExternalClockDriver->SongMapEvaluator->GetMidiFileWithTempoMap() : 
+									SongMapEvaluator->GetMidiFileWithOtherMaps(),
+								SongMapEvaluator->GetMidiFileWithOtherMaps());
 	}
 
-	void FMidiClock::AddSpeedChangeToBlock(const FMidiTimestampSpeed& NewSpeed)
+	void FMidiClock::PrepareBlock()
 	{
-		check(SpeedChangesInBlock.IsEmpty() || SpeedChangesInBlock.Last().BlockSampleFrameIndex <= NewSpeed.BlockSampleFrameIndex);
-		if (SpeedChangesInBlock.Last().BlockSampleFrameIndex == NewSpeed.BlockSampleFrameIndex)
-		{
-			SpeedChangesInBlock.Last().Speed = NewSpeed.Speed;
-		}
-		else
-		{
-			SpeedChangesInBlock.Add(NewSpeed);
-		}
-		HasSpeedChangeInBlock = true;
-		DrivingMidiPlayCursorMgr->InformOfHiResAdvanceRate(NewSpeed.Speed);
-	}
+		NumTransportChangeInBlock = 0;
+		TransportAtBlockStart = TransportAtBlockEnd;
+		
+		NumSpeedChangeInBlock = 0;
+		SpeedAtBlockStart = SpeedAtBlockEnd;
+		
+		NumTempoChangeInBlock = 0;
+		TempoAtBlockStart = TempoAtBlockEnd;
 
-	const TArray<FMidiClockEvent>& FMidiClock::GetMidiClockEventsInBlock() const
-	{
-		return MidiClockEventsInBlock;
-	}
+		NumTimeSignatureChangeInBlock = 0;
+		TimeSignatureAtBlockStart = TimeSignatureAtBlockEnd;
+		
+		CurrentBlockFrameIndex = 0;
 
-	const FMidiClockEvent* FMidiClock::FindLastMidiClockEventAtBlockSampleFrame(int32 FrameIndex) const
-	{
-		int32 Index = Algo::LowerBoundBy(MidiClockEventsInBlock, FrameIndex, &FMidiClockEvent::BlockFrameIndex);
-		if (MidiClockEventsInBlock.IsValidIndex(Index))
-		{
-			while (MidiClockEventsInBlock.IsValidIndex(Index + 1) && MidiClockEventsInBlock[Index + 1].BlockFrameIndex == FrameIndex)
-			{
-				++Index;
-			}
-			return &MidiClockEventsInBlock[Index];
-		}
+		MidiDataChangedInBlock = false;
+
+		MidiClockEventsInBlock.Reset();
 
-		if (!MidiClockEventsInBlock.IsEmpty())
+		if (ExternalClockDriver && ExternalClockDriver->MidiDataChangedInBlock)
 		{
-			return &MidiClockEventsInBlock.Last();
+			RebuildSongMapEvaluator(ExternalClockDriver->SongMapEvaluator->GetMidiFileWithTempoMap(), SongMapEvaluator->GetMidiFileWithOtherMaps());
+			PostTempoOrTimeSignatureEventsIfNeeded();
 		}
-		return nullptr;
 	}
 
-	EMusicPlayerTransportState FMidiClock::GetTransportStateAtBlockSampleFrame(int32 FrameIndex) const
+	void FMidiClock::SetTransportState(int32 BlockFrameIndex, EMusicPlayerTransportState TransportState)
 	{
-		return GetTransportTimestampForBlockSampleFrame(FrameIndex).TransportState;
+		AddTransportStateChangeToBlock(BlockFrameIndex, TransportState);
 	}
 
-	EMusicPlayerTransportState FMidiClock::GetTransportStateAtEndOfBlock() const
+	void FMidiClock::SetSpeed(int32 BlockFrameIndex, float Speed)
 	{
-		return CurrentTransportState.TransportState;
+		CurrentLocalSpeed = Speed;
+		AddSpeedChangeToBlock(BlockFrameIndex, Speed, true);
 	}
 
-	const FMidiTimestampTransportState& FMidiClock::GetTransportTimestampForBlockSampleFrame(int32 FrameIndex) const
+	void FMidiClock::SetTempo(int32 BlockFrameIndex, int32 Tick, float Bpm)
 	{
-		if (TransportChangesInBlock.IsEmpty())
-		{
-			return CurrentTransportState;
-		}
-		int32 Index = Algo::UpperBoundBy(TransportChangesInBlock, FrameIndex, [](const FMidiTimestampTransportState& t) { return t.BlockSampleFrameIndex; }) - 1;
-		if (Index < 0)
-		{
-			return CurrentTransportState;
-		}
-		return TransportChangesInBlock[Index];
+		AddTempoChangeToBlock(BlockFrameIndex, Tick, Bpm);
 	}
 
-	float FMidiClock::GetSpeedAtBlockSampleFrame(int32 FrameIndex) const
+	void HarmonixMetasound::FMidiClock::SetTimeSignature(int32 BlockFrameIndex, int32 Tick, const FTimeSignature& TimeSignature)
 	{
-		return GetSpeedTimestampForBlockSampleFrame(FrameIndex).Speed;
+		AddTimeSignatureChangeToBlock(BlockFrameIndex, Tick, TimeSignature);
 	}
 
-	float FMidiClock::GetSpeedAtEndOfBlock() const
+	void FMidiClock::SeekTo(int32 BlockFrameIndex, const FMusicSeekTarget& InTarget)
 	{
-		return SpeedChangesInBlock.Last().Speed;
-	}
+		CurrentBlockFrameIndex = BlockFrameIndex;
 
-	const FMidiTimestampSpeed& FMidiClock::GetSpeedTimestampForBlockSampleFrame(int32 FrameIndex) const
-	{
-		int32 Index = Algo::UpperBoundBy(SpeedChangesInBlock, FrameIndex, [](const FMidiTimestampSpeed& t) { return t.BlockSampleFrameIndex; }) - 1;
-		if (Index < 0)
+		int32 Tick = 0;
+		switch (InTarget.Type)
 		{
-			Index = 0;
+		case ESeekPointType::BarBeat:
+			Tick = SongMapEvaluator->MusicTimestampToTick(InTarget.BarBeat);
+			break;
+		default:
+		case ESeekPointType::Millisecond:
+			Tick = SongMapEvaluator->MsToTick(InTarget.Ms);
+			break;
 		}
-		return SpeedChangesInBlock[Index];
+
+		SeekTo(BlockFrameIndex, Tick);
 	}
 
-	float FMidiClock::GetTempoAtBlockSampleFrame(int32 FrameIndex) const
+	void FMidiClock::SeekTo(int32 BlockFrameIndex, int32 Tick)
 	{
-		int32 Index = Algo::UpperBoundBy(TempoChangesInBlock, FrameIndex, [](const FMidiTimestampTempo& t) { return t.BlockSampleFrameIndex; }) - 1;
-		if (Index < 0)
+		if (NextMidiTickToProcess != Tick)
 		{
-			Index = 0;
+			AddSeekToBlock(BlockFrameIndex, Tick);
+			SampleCount = FMath::Max<FSampleCount>(FSampleCount(SongMapEvaluator->TickToMs(NextMidiTickToProcess) / 1000.0f * SampleRate), 0);
 		}
-		return TempoChangesInBlock[Index].Tempo;
-	}
-
-	float FMidiClock::GetTempoAtEndOfBlock() const
-	{
-		return TempoChangesInBlock.Last().Tempo;
 	}
 
-	int32 FMidiClock::GetNumTempoChangesInBlock() const
+	void FMidiClock::AddTransientLoop(int32 BlockFrameIndex, int32 NewFirstTickInLoop, int32 NewLoopLengthTicks)
 	{
-		return TempoChangesInBlock.Num();
+		AddLoopToBlock(BlockFrameIndex, NewFirstTickInLoop, NewLoopLengthTicks);
 	}
 
-	FMidiTimestampTempo FMidiClock::GetTempoChangeByIndex(int32 Index) const
+	void FMidiClock::SetupPersistentLoop(int32 NewFirstTickInLoop, int32 NewLoopLengthTicks)
 	{
-		if (Index >= 0 && Index < TempoChangesInBlock.Num())
+		if (ensureAlwaysMsgf(NewLoopLengthTicks > (Harmonix::Midi::Constants::GTicksPerQuarterNoteInt / 4), TEXT("For performance reasons, Midi Clock loops must be at least a 1/16th note long!")))
 		{
-			return TempoChangesInBlock[Index]; // intentional copy
+			FirstTickInLoop = NewFirstTickInLoop;
+			LoopLengthTicks = NewLoopLengthTicks;
 		}
-
-		return InvalidMidiTimestampTempo;
-	}
-
-	int32 FMidiClock::GetCurrentMidiTick() const
-	{
-		return DrivingMidiPlayCursorMgr->GetCurrentHiResTick();
 	}
 
-	int32 FMidiClock::GetCurrentBlockFrameIndex() const
+	void FMidiClock::ClearPersistentLoop()
 	{
-		return CurrentBlockFrameIndex;
+		FirstTickInLoop = -1;
+		LoopLengthTicks = 0;
 	}
 
-	void FMidiClock::AdvanceHiResToMs(int32 BlockFrameIndex, float Ms, bool Broadcast)
+	bool FMidiClock::HasPersistentLoop() const
 	{
-		CurrentBlockFrameIndex = BlockFrameIndex;
-		DrivingMidiPlayCursorMgr->AdvanceHiResToMs(Ms, Broadcast);
+		return FirstTickInLoop != -1 && LoopLengthTicks > 0;
 	}
 
-	void FMidiClock::AttachToTimeAuthority(const FMidiClock& MidiClockRef)
+	float FMidiClock::GetLoopStartMs() const
 	{
-		DrivingMidiPlayCursorMgr->AttachToTimeAuthority(MidiClockRef.DrivingMidiPlayCursorMgr);
-		bSeekToAuthorityOnNextProcess = true;
-	}
+		if (!HasPersistentLoop() || FirstTickInLoop <= 0)
+		{
+			return 0.0f;
+		}
 
-	void FMidiClock::DetachFromTimeAuthority()
-	{
-		DrivingMidiPlayCursorMgr->DetachFromTimeAuthority();
-	}
-	
-	float FMidiClock::GetQuarterNoteIncludingCountIn() const
-	{
-		int32 Tick = DrivingMidiPlayCursorMgr->GetCurrentHiResTick();
-		return (float)Tick / (float)DrivingMidiPlayCursorMgr->GetSongMaps().GetTicksPerQuarterNote();
+		return SongMapEvaluator->TickToMs(FirstTickInLoop);
 	}
 
-	FMusicTimestamp FMidiClock::GetCurrentMusicTimestamp() const
+	float FMidiClock::GetLoopEndMs() const
 	{
-		return DrivingMidiPlayCursorMgr->GetMusicTimestampAtMs(GetCurrentHiResMs());
-	}
+		if (!HasPersistentLoop())
+		{
+			return 0.0f;
+		}
 
-	FMusicTimestamp FMidiClock::GetMusicTimestampAtBlockOffset(const int32 Offset) const
-	{
-		const float MsAtOffset = GetMsAtBlockOffset(Offset);
-		return DrivingMidiPlayCursorMgr->GetMusicTimestampAtMs(MsAtOffset);
+		return SongMapEvaluator->TickToMs(FirstTickInLoop + LoopLengthTicks);
 	}
 
-	float FMidiClock::GetMsAtBlockOffset(int32 Offset) const
+	float FMidiClock::GetLoopLengthMs() const
 	{
-		const float EndMs = GetCurrentHiResMs();
-		const float MsPerFrame = 1000 / SampleRate;
-		const int32 InvOffset = BlockSize - Offset - 1;
-		return EndMs - InvOffset * MsPerFrame;
+		return GetLoopEndMs() - GetLoopStartMs();
 	}
 
-	void FMidiClock::LockForMidiDataChanges()
+	void FMidiClock::Advance(const FMidiClock& DrivingClock, int32 StartFrame, int32 NumFrames)
 	{
+		int32 EndFrame = StartFrame + NumFrames;
+		const TArray<FMidiClockEvent>& ClockEvents = DrivingClock.GetMidiClockEventsInBlock();
+		int32 Index = Algo::LowerBoundBy(ClockEvents, StartFrame, &FMidiClockEvent::BlockFrameIndex);
+		while (ClockEvents.IsValidIndex(Index))
+		{
+			const FMidiClockEvent& Event = ClockEvents[Index];
+			if (Event.BlockFrameIndex >= EndFrame)
+			{
+				return;
+			}
 
-		DrivingMidiPlayCursorMgr->LockForMidiDataChanges();
-	}
-
-	void FMidiClock::MidiDataChangesComplete(FMidiPlayCursorMgr::EMidiChangePositionCorrectMode Mode /*= FMidiPlayCursorMgr::EMidiChangePositionCorrectMode::MaintainTick*/)
-	{
-		DrivingMidiPlayCursorMgr->MidiDataChangeComplete(Mode);
+			HandleClockEvent(DrivingClock, Event);
+			++Index;
+		}
 	}
 
-	void FMidiClock::SeekTo(const FMusicSeekTarget& Target, int32 PreRollBars)
+	void FMidiClock::Advance(int32 StartFrame, int32 NumFrames)
 	{
-		int32 Tick = 0;
-		switch (Target.Type)
+		if (TransportAtBlockEnd != EMusicPlayerTransportState::Playing && TransportAtBlockEnd != EMusicPlayerTransportState::Continuing)
 		{
-		case ESeekPointType::BarBeat:
-			Tick = DrivingMidiPlayCursorMgr->GetBarMap().MusicTimestampToTick(Target.BarBeat);
-			break;
-		default:
-		case ESeekPointType::Millisecond:
-			Tick = DrivingMidiPlayCursorMgr->GetSongMaps().MsToTick(Target.Ms);
-			break;
+			UE_LOG(LogMIDI, Error, TEXT("FMidiClock: Attempt to advance clock that is not playing!"));
+			return;
 		}
-		DrivingMidiPlayCursorMgr->SeekTo(Tick, PreRollBars,true,false);
-	}
 
-	void FMidiClock::SeekTo(int32 Tick, int32 PreRollBars)
-	{
-		DrivingMidiPlayCursorMgr->SeekTo(Tick, PreRollBars, true, false);
-	}
-
-	void FMidiClock::InformOfCurrentAdvanceRate(float AdvanceRate)
-	{
-		DrivingMidiPlayCursorMgr->InformOfHiResAdvanceRate(AdvanceRate);
+		int32 BlockFrameIndex = StartFrame;
+		while (NumFrames > FramesUntilNextProcess)
+		{
+			BlockFrameIndex += FramesUntilNextProcess;
+			NumFrames -= FramesUntilNextProcess;
+			FramesUntilNextProcess = kMidiGranularity;
+			FSampleCount TargetFrame = SampleCount + (FSampleCount)((float)kMidiGranularity * SpeedAtBlockEnd);
+			float TargetMs = ((float)TargetFrame * 1000.0f) / SampleRate;
+			bool bDidLoop = AdvanceToMs(BlockFrameIndex, TargetMs);
+			// AdvanceToMs will have updated SampleCount, but it will have been 
+			// quantized to the resulting target midi tick. That is appropriate in 
+			// some cases, but in this advance type (by sample frames) we need/want
+			// the un-quantized value UNLESS the advance caused a loop...
+			if (!bDidLoop)
+			{
+				// blow away the "quantized" value with our true, un-quantized value here...
+				SampleCount = TargetFrame;
+			}
+		}
+		FramesUntilNextProcess -= NumFrames;
 	}
 
-	void FMidiClock::CopySpeedAndTempoChanges(const FMidiClock* InClock, float InSpeedMult)
+	bool FMidiClock::AdvanceToTick(int32 BlockFrameIndex, int32 UpToTick)
 	{
-		HasTempoChangeInBlock = InClock->HasTempoChangeInBlock;
-		TempoChangesInBlock = InClock->TempoChangesInBlock;
+		using namespace MidiClockMessageTypes;
 
-		HasSpeedChangeInBlock = InClock->HasSpeedChangeInBlock;
-		SpeedChangesInBlock = InClock->SpeedChangesInBlock;
+		bool bDidLoop = false;
 
-		for (FMidiTimestampSpeed& SpeedChange : SpeedChangesInBlock)
+		int32 TickAfterLoop = FirstTickInLoop + LoopLengthTicks;
+		if (HasPersistentLoop() && UpToTick > TickAfterLoop)
 		{
-			SpeedChange.Speed *= InSpeedMult;
-		}
-	}
+			// first we might need advance to the loop end...
+			if (NextMidiTickToProcess > FirstTickInLoop)
+			{
+				AddAdvanceToBlock(BlockFrameIndex, NextMidiTickToProcess, TickAfterLoop - NextMidiTickToProcess);
+				AddLoopToBlock(BlockFrameIndex, FirstTickInLoop, LoopLengthTicks);
+				check(NextMidiTickToProcess == FirstTickInLoop);
+				bDidLoop = true;
+			}
 
-	void FMidiClock::HandleTransportChange(int32 BlockFrameIndex, EMusicPlayerTransportState TransportState)
-	{
-		switch (TransportState)
-		{
-		case EMusicPlayerTransportState::Prepared:
-		case EMusicPlayerTransportState::Paused:
-			AddTransportStateChangeToBlock({ BlockFrameIndex, 0.0f, TransportState });
-			break;
-		case EMusicPlayerTransportState::Playing:
-			AddTransportStateChangeToBlock({ BlockFrameIndex, 0.0f, TransportState });
-			break;
+			int NumTicksLeftToProcess = UpToTick - TickAfterLoop;
+			while (NumTicksLeftToProcess)
+			{
+				int32 TicksThisPass = FMath::Min(LoopLengthTicks, NumTicksLeftToProcess);
+				NumTicksLeftToProcess -= TicksThisPass;
+				AddAdvanceToBlock(BlockFrameIndex, NextMidiTickToProcess, TicksThisPass);
+				if (NumTicksLeftToProcess > 0)
+				{
+					AddLoopToBlock(BlockFrameIndex, FirstTickInLoop, LoopLengthTicks);
+					check(NextMidiTickToProcess == FirstTickInLoop);
+					bDidLoop = true;
+				}
+			}
 		}
-	}
-
-	void FMidiClock::Process(int32 StartFrame, int32 NumFrames, int32 PrerollBars, float Speed)
-	{
-		int32 EndFrame = StartFrame + NumFrames;
-		switch (CurrentTransportState.TransportState)
+		else
 		{
-		case EMusicPlayerTransportState::Playing:
-		case EMusicPlayerTransportState::Continuing:
-			WriteAdvance(StartFrame, EndFrame, Speed);
-			break;
+			AddAdvanceToBlock(BlockFrameIndex, NextMidiTickToProcess, UpToTick - NextMidiTickToProcess);
+			if (HasPersistentLoop() && NextMidiTickToProcess >= (FirstTickInLoop + LoopLengthTicks))
+			{
+				AddLoopToBlock(BlockFrameIndex, FirstTickInLoop, LoopLengthTicks);
+				check(NextMidiTickToProcess == FirstTickInLoop);
+				bDidLoop = true;
+			}
 		}
+		SampleCount = FMath::Max<FSampleCount>(FSampleCount(SongMapEvaluator->TickToMs(NextMidiTickToProcess) / 1000.0f * SampleRate), 0);
+		return bDidLoop;
 	}
 
-	void FMidiClock::Process(const FMidiClock& DrivingClock, int32 StartFrame, int32 NumFrames, int32 PrerollBars, float Speed)
+	bool FMidiClock::AdvanceToMs(int32 BlockFrameIndex, float Ms)
 	{
-		int32 EndFrame = StartFrame + NumFrames;
+		using namespace MidiClockMessageTypes;
 
-		// When we are driven by an external clock, we want to be sure that upon hookup we are synced to the same tick as our Driver
-		// Otherwise, we move in lockstep with our Driver but potentially at an offset beat/measure/etc
-		if (bSeekToAuthorityOnNextProcess)
+		bool bDidLoop = false;
+
+		int32 ToFutureTick = FMath::RoundToInt32(SongMapEvaluator->MsToTick(Ms));
+
+		// NOTE: We CAN'T just pass this calculated future tick to the AdvanceToTick function
+		// and let it wrap around the loop!... Because tempo changes during the loop might result in 
+		// a different wrapped tick. So we have to do this here where we will wrap in "ms space"
+		// instead of "tick space" each time around the loop...
+
+		// do we need to wrap around the loop?
+		int32 TickAfterLoop = FirstTickInLoop + LoopLengthTicks;
+		if (HasPersistentLoop() && ToFutureTick > TickAfterLoop)
 		{
-			bSeekToAuthorityOnNextProcess = false;
+			// first we might need advance to the loop end...
+			if (NextMidiTickToProcess < TickAfterLoop)
+			{
+				AddAdvanceToBlock(BlockFrameIndex, NextMidiTickToProcess, TickAfterLoop - NextMidiTickToProcess);
+			}
+
+			AddLoopToBlock(BlockFrameIndex, FirstTickInLoop, LoopLengthTicks);
+			check(NextMidiTickToProcess == FirstTickInLoop);
+			bDidLoop = true;
 
-			if (const FMidiClockEvent* Event = DrivingClock.FindLastMidiClockEventAtBlockSampleFrame(StartFrame))
+			float MsRemainingAfterProcessingToLoopEnd = Ms - GetLoopEndMs();
+
+			while (MsRemainingAfterProcessingToLoopEnd > 0.0f)
 			{
-				int32 SeekTick = 0;
-				if (Event->Msg.IsType<MidiClockMessageTypes::FSeekTo>())
-				{
-					SeekTick = Event->Msg.Get<MidiClockMessageTypes::FSeekTo>().ToTick;
-				}
-				else if (Event->Msg.IsType<MidiClockMessageTypes::FReset>())
-				{
-					SeekTick = Event->Msg.Get<MidiClockMessageTypes::FReset>().ToTick;
-				}
-				else if (Event->Msg.IsType<MidiClockMessageTypes::FSeekThru>())
-				{
-					SeekTick = Event->Msg.Get<MidiClockMessageTypes::FSeekThru>().ThruTick + 1;
-				}
-				else if (Event->Msg.IsType<MidiClockMessageTypes::FAdvanceThru>())
-				{
-					SeekTick = Event->Msg.Get<MidiClockMessageTypes::FAdvanceThru>().ThruTick + 1;
-				}
-				// These lower three branches shouldn't be hit in practice, but putting them here for completeness
-				else if (Event->Msg.IsType<MidiClockMessageTypes::FLoop>())
+				ToFutureTick = SongMapEvaluator->MsToTick(GetLoopStartMs() + MsRemainingAfterProcessingToLoopEnd);
+				if (ToFutureTick > TickAfterLoop)
 				{
-					SeekTick = Event->Msg.Get<MidiClockMessageTypes::FLoop>().LoopStartTick;
+					AddAdvanceToBlock(BlockFrameIndex, NextMidiTickToProcess, TickAfterLoop - NextMidiTickToProcess);
+					AddLoopToBlock(BlockFrameIndex, FirstTickInLoop, LoopLengthTicks);
+					check(NextMidiTickToProcess == FirstTickInLoop);
+					MsRemainingAfterProcessingToLoopEnd -= GetLoopLengthMs();
 				}
-				else if (Event->Msg.IsType<MidiClockMessageTypes::FTempoChange>())
+				else
 				{
-					SeekTick = Event->Msg.Get<MidiClockMessageTypes::FTempoChange>().Tick;
+					AddAdvanceToBlock(BlockFrameIndex, NextMidiTickToProcess, ToFutureTick - NextMidiTickToProcess);
+					MsRemainingAfterProcessingToLoopEnd = 0.0f;
 				}
-				else if (Event->Msg.IsType<MidiClockMessageTypes::FTimeSignatureChange>())
-				{
-					SeekTick = Event->Msg.Get<MidiClockMessageTypes::FTimeSignatureChange>().Tick;
-				}
-
-				SeekTo(StartFrame, CalculateMappedTick(SeekTick), PrerollBars);
-				StartFrame = Event->BlockFrameIndex + 1;
 			}
-		}
 
-		const TArray<FMidiClockEvent>& ClockEvents = DrivingClock.GetMidiClockEventsInBlock();
-		int32 Index = Algo::LowerBoundBy(ClockEvents, StartFrame, &FMidiClockEvent::BlockFrameIndex);
-		while (ClockEvents.IsValidIndex(Index))
+			// Because we did some looping, we need to update our SampleCount so it is "wrapped" appropriately...
+			SampleCount = FMath::Max<FSampleCount>(FSampleCount(SongMapEvaluator->TickToMs(NextMidiTickToProcess) / 1000.0f * SampleRate), 0);
+		}
+		else
 		{
-			const FMidiClockEvent& Event = ClockEvents[Index];
-			if (Event.BlockFrameIndex >= EndFrame)
+			if (ToFutureTick > NextMidiTickToProcess)
 			{
-				return;
+				AddAdvanceToBlock(BlockFrameIndex, NextMidiTickToProcess, ToFutureTick - NextMidiTickToProcess);
+				if (HasPersistentLoop() && LastProcessedMidiTick == TickAfterLoop)
+				{
+					AddLoopToBlock(BlockFrameIndex, FirstTickInLoop, LoopLengthTicks);
+					check(NextMidiTickToProcess == FirstTickInLoop);
+					bDidLoop = true;
+				}
+				SampleCount = FMath::Max<FSampleCount>(FSampleCount(Ms / 1000.0f * SampleRate), 0);
 			}
-			
-			HandleClockEvent(DrivingClock, Event, PrerollBars, Speed);
-			++Index;
 		}
+		return bDidLoop;
 	}
-	
-	void FMidiClock::HandleClockEvent(const FMidiClock& DrivingClock, const FMidiClockEvent& Event, int32 PrerollBars, float Speed)
+
+	float FMidiClock::GetSpeedAtBlockSampleFrame(int32 FrameIndex) const
 	{
-		if (Event.Msg.IsType<MidiClockMessageTypes::FReset>())
-		{
-			const int32 Tick = CalculateMappedTick(Event.Msg.Get<MidiClockMessageTypes::FReset>().ToTick);
-			SeekTo(Event.BlockFrameIndex, Tick + 1, PrerollBars);
-		}
-		else if (Event.Msg.IsType<MidiClockMessageTypes::FLoop>())
+		if (NumSpeedChangeInBlock == 0 || FrameIndex == 0)
 		{
-			// ignore loops since the driving clock will loop us by seeking
+			return SpeedAtBlockStart;
 		}
-		else if (Event.Msg.IsType<MidiClockMessageTypes::FSeekTo>())
+
+		float SpeedCandidate = SpeedAtBlockStart;
+		for (const FMidiClockEvent& ClockEvent : MidiClockEventsInBlock)
 		{
-			const int32 Tick = CalculateMappedTick(Event.Msg.Get<MidiClockMessageTypes::FSeekTo>().ToTick);
-			SeekTo(Event.BlockFrameIndex, Tick, PrerollBars);
+			if (const MidiClockMessageTypes::FSpeedChange* AsSpeedChange = ClockEvent.TryGet<MidiClockMessageTypes::FSpeedChange>())
+			{
+				if (ClockEvent.BlockFrameIndex > FrameIndex)
+				{
+					break;
+				}
+				SpeedCandidate = AsSpeedChange->Speed;
+			}
 		}
-		else if (Event.Msg.IsType<MidiClockMessageTypes::FSeekThru>())
+		return SpeedCandidate;
+	}
+
+	float FMidiClock::GetTempoAtBlockSampleFrame(int32 FrameIndex) const
+	{
+		if (NumTempoChangeInBlock == 0 || FrameIndex == 0)
 		{
-			const int32 Tick = CalculateMappedTick(Event.Msg.Get<MidiClockMessageTypes::FSeekThru>().ThruTick);
-			SeekTo(Event.BlockFrameIndex, Tick + 1, PrerollBars);
+			return TempoAtBlockStart;
 		}
-		else if (Event.Msg.IsType<MidiClockMessageTypes::FAdvanceThru>())
+
+		float TempoCandidate = TempoAtBlockStart;
+		for (const FMidiClockEvent& ClockEvent : MidiClockEventsInBlock)
 		{
-			const MidiClockMessageTypes::FAdvanceThru& AdvanceThru = Event.Msg.Get<MidiClockMessageTypes::FAdvanceThru>();
-			
-			if (AdvanceThru.IsPreRoll)
-			{
-				const int32 Tick = CalculateMappedTick(AdvanceThru.ThruTick);
-				SeekTo(Event.BlockFrameIndex, Tick + 1, PrerollBars);
-			}
-			else
+			if (const MidiClockMessageTypes::FTempoChange* AsTempoChange = ClockEvent.TryGet<MidiClockMessageTypes::FTempoChange>())
 			{
-				// Advance based on the delta ticks, and not based on the absolute tick
-				const int32 Tick = GetCurrentMidiTick() + (AdvanceThru.ThruTick - AdvanceThru.FromTick);
-				const float Ms = GetSongMaps().TickToMs(Tick);
-
-				const float ClockInSpeed = DrivingClock.GetSpeedAtBlockSampleFrame(Event.BlockFrameIndex);
-				const float AdvanceRatio = DrivingClock.GetSongMaps().GetTempoAtTick(AdvanceThru.FromTick)
-								   / GetSongMaps().GetTempoAtTick(GetCurrentMidiTick());
-				InformOfCurrentAdvanceRate(ClockInSpeed * Speed * AdvanceRatio);
-				AdvanceHiResToMs(Event.BlockFrameIndex, Ms, true);
+				if (ClockEvent.BlockFrameIndex > FrameIndex)
+				{
+					break;
+				}
+				TempoCandidate = AsTempoChange->Tempo;
 			}
 		}
+		return TempoCandidate;
 	}
 
-	void FMidiClock::WriteAdvance(int32 StartFrameIndex, int32 EndFrameIndex, float InSpeed /*= 1.0f*/)
+	float FMidiClock::GetCurrentSongPosMs() const
 	{
-		int32 FramesToProcess = EndFrameIndex - StartFrameIndex;
-		const FTempoMap& TempoMap = DrivingMidiPlayCursorMgr->GetSongMaps().GetTempoMap();
-		while (FramesToProcess > FramesUntilNextProcess)
-		{
-			StartFrameIndex += FramesUntilNextProcess;
-			FSampleCount AdvanceToFrame = SampleCount + (FSampleCount)((float)kMidiGranularity * InSpeed);
-			FramesUntilNextProcess = kMidiGranularity;
-			float AdvanceToMs = ((float)AdvanceToFrame * 1000.0f) / SampleRate;
-			DrivingMidiPlayCursorMgr->InformOfHiResAdvanceRate(InSpeed);
-			AdvanceHiResToMs(StartFrameIndex, AdvanceToMs, true);
-
-			
-			// Account for looping
-			// recalculate sampleCount based on our new time, if clock looped back, this will be different from AdvanceToFrame
-			if (GetCurrentHiResMs() < AdvanceToMs)
-			{
-				SampleCount = FMath::Max<FSampleCount>(FSampleCount(GetCurrentHiResMs() / 1000.0f * SampleRate), 0);
-			}
-			else
-			{
-				SampleCount = AdvanceToFrame;
-			}
-			
+		// We have processed THRU MostRecentProcessedMidiTick,
+		// so our Ms position is AFTER that tick, just before the 
+		// next tick...
+		return SongMapEvaluator->TickToMs(NextMidiTickToProcess);
+	}
 
-			FramesToProcess = EndFrameIndex - StartFrameIndex;
-		}
-		FramesUntilNextProcess -= FramesToProcess;
+	FMusicTimestamp FMidiClock::GetMusicTimestampAtBlockEnd() const
+	{
+		return SongMapEvaluator->TickToMusicTimestamp(NextMidiTickToProcess);
 	}
 
-	void FMidiClock::SeekTo(int32 BlockFrameIndex, const FMusicSeekTarget& InTarget, int32 InPrerollBars)
+	FMusicTimestamp FMidiClock::GetMusicTimestampAtBlockOffset(const int32 Offset) const
 	{
-		CurrentBlockFrameIndex = BlockFrameIndex;
-		SeekTo(InTarget, InPrerollBars);
-		SampleCount = FMath::Max<FSampleCount>(FSampleCount(GetCurrentHiResMs() / 1000.0f * SampleRate), 0);
-		FramesUntilNextProcess = 0;
+		return SongMapEvaluator->TickToMusicTimestamp(GetNextTickToProcessAtBlockFrame(Offset));
 	}
 
-	void FMidiClock::SeekTo(int32 BlockFrameIndex, int32 Tick, int32 InPrerollBars)
+	float FMidiClock::GetSongPosMsAtBlockOffset(int32 Offset) const
 	{
-		CurrentBlockFrameIndex = BlockFrameIndex;
-		DrivingMidiPlayCursorMgr->SeekTo(Tick, InPrerollBars, true, false);
-		SampleCount = FMath::Max<FSampleCount>(FSampleCount(GetCurrentHiResMs() / 1000.0f * SampleRate), 0);
-		FramesUntilNextProcess = 0;
+		return SongMapEvaluator->TickToMs(GetNextTickToProcessAtBlockFrame(Offset));
 	}
 
-	int32 FMidiClock::CalculateMappedTick(int32 Tick) const
+	int32 FMidiClock::WrapTickIfLooping(int32 Tick) const
 	{
-		if (DoesLoop())
+		if (HasPersistentLoop())
 		{
-			const int32 LoopStartTick = GetLoopStartTick();
-			const int32 LoopEndTick = GetLoopEndTick();
-			const int32 LoopLengthTicks = LoopEndTick - LoopStartTick;
 			// only wrap the tick if it we're passed the loop end tick
-			if (LoopLengthTicks > 0 && Tick >= LoopEndTick)
+			int32 LastTickInLoop = FirstTickInLoop + LoopLengthTicks - 1;
+			if (Tick > LastTickInLoop)
 			{
-				return LoopStartTick + (Tick - LoopStartTick) % LoopLengthTicks;
+				return FirstTickInLoop + ((Tick - LastTickInLoop - 1) % LoopLengthTicks);
 			}
 		}
 		return Tick;
 	}
-	
+
 	TSharedPtr<FMidiFileData> FMidiClock::MakeClockConductorMidiData(float InTempoBPM, int32 InTimeSigNum, int32 InTimeSigDen)
 	{
 		TSharedPtr<FMidiFileData> OutMidiData = MakeShared<FMidiFileData>();
@@ -639,113 +601,420 @@ namespace HarmonixMetasound
 		return OutMidiData;
 	}
 
-	FMidiClock::FMidiClockEventCursor::FMidiClockEventCursor(FMidiClock* MidiClock)
-		: MyMidiClock(MidiClock)
+	int32 FMidiClock::GetNextTickToProcessAtBlockFrame(int32 BlockFrame) const
 	{
-		check(MidiClock);
+		int32 FoundNextTick = NextMidiTickToProcess;
+		int32 AtBlockIndex = 0;
+
+		using namespace MidiClockMessageTypes;
+		// we're going to have to look through the clock events...
+		for (const FMidiClockEvent& Event : MidiClockEventsInBlock)
+		{
+			if (Event.BlockFrameIndex >= BlockFrame)
+			{
+				break;
+			}
+
+			if (const FSeek* AsSeek = Event.TryGet<FSeek>())
+			{
+				FoundNextTick = AsSeek->NewNextTick;
+				AtBlockIndex = Event.BlockFrameIndex;
+			}
+			else if (const FAdvance* AsProcess = Event.TryGet<FAdvance>())
+			{
+				FoundNextTick = AsProcess->FirstTickToProcess + AsProcess->NumberOfTicksToProcess;
+				AtBlockIndex = Event.BlockFrameIndex;
+			}
+			else if (const FLoop* AsLoop = Event.TryGet<FLoop>())
+			{
+				FoundNextTick = AsLoop->FirstTickInLoop;
+				AtBlockIndex = Event.BlockFrameIndex;
+			}
+		}
 
-		// Ignore midi messages, test events, and preroll notes
-		SetMessageFilter(EFilterPassFlags::All & ~(EFilterPassFlags::MidiMessage | EFilterPassFlags::Text | EFilterPassFlags::PreRollNoteOn));
+		return FoundNextTick;
 	}
 
-	void FMidiClock::FMidiClockEventCursor::Reset(bool ForceNoBroadcast /*= false*/)
+	void FMidiClock::AddEvent(const FMidiClockEvent& InEvent, bool bRequireSequential)
 	{
-		const int32 FromTick = CurrentTick;
-		FMidiPlayCursor::Reset(ForceNoBroadcast);
-		AddEvent(FMidiClockEvent(MyMidiClock->GetCurrentBlockFrameIndex(), MidiClockMessageTypes::FReset(FromTick, CurrentTick, ForceNoBroadcast)));
+		if (bRequireSequential && !MidiClockEventsInBlock.IsEmpty())
+		{
+			check(MidiClockEventsInBlock.Last().BlockFrameIndex <= InEvent.BlockFrameIndex);
+		}
+
+		MidiClockEventsInBlock.Add(InEvent);
 	}
 
-	void FMidiClock::FMidiClockEventCursor::OnLoop(int32 LoopStartTick, int32 LoopEndTick)
+	void FMidiClock::HandleClockEvent(const FMidiClock& DrivingClock, const FMidiClockEvent& Event)
 	{
-		FMidiPlayCursor::OnLoop(LoopStartTick, LoopEndTick);
-		AddEvent(FMidiClockEvent(MyMidiClock->GetCurrentBlockFrameIndex(), MidiClockMessageTypes::FLoop(LoopStartTick, LoopEndTick)));
+		using namespace MidiClockMessageTypes;
+
+		if (const FLoop* AsLoop = Event.TryGet<FLoop>())
+		{
+			const int32 Tick = WrapTickIfLooping(AsLoop->FirstTickInLoop);
+			SeekTo(Event.BlockFrameIndex, Tick);
+			TickResidualWhenDriven = 0.0f;
+		}
+		else if (const FSeek* AsSeek = Event.TryGet<FSeek>())
+		{
+			const int32 Tick = WrapTickIfLooping(AsSeek->NewNextTick);
+			SeekTo(Event.BlockFrameIndex, Tick);
+			TickResidualWhenDriven = 0.0f;
+		}
+		else if (const FAdvance* AsAdvance = Event.TryGet<FAdvance>())
+		{
+			// Advance based on the delta ticks, and not based on the absolute tick
+			int32 UpToTick;
+			if (FMath::IsNearlyEqual(SpeedAtBlockEnd, 1.0f, 0.0001))
+			{
+				UpToTick = NextMidiTickToProcess + AsAdvance->NumberOfTicksToProcess;
+			}
+			else
+			{
+				float FractionalToTick = (float)NextMidiTickToProcess + ((float)AsAdvance->NumberOfTicksToProcess * CurrentLocalSpeed) + TickResidualWhenDriven;
+				UpToTick = FMath::FloorToInt32(FractionalToTick);
+				TickResidualWhenDriven = FMath::Fractional(FractionalToTick);
+			}
+			AdvanceToTick(Event.BlockFrameIndex, UpToTick);
+		}
+		else if (const FTempoChange* AsTempoChange = Event.TryGet<FTempoChange>())
+		{
+			AddTempoChangeToBlock(Event.BlockFrameIndex, NextMidiTickToProcess, AsTempoChange->Tempo);
+		}
+		else if (const FSpeedChange* AsSpeedChange = Event.TryGet<FSpeedChange>())
+		{
+			AddSpeedChangeToBlock(Event.BlockFrameIndex, AsSpeedChange->Speed, false);
+		}
+		else if (const FTransportChange* AsTransportChange = Event.TryGet<FTransportChange>())
+		{
+			AddTransportStateChangeToBlock(Event.BlockFrameIndex, AsTransportChange->TransportState);
+		}
 	}
 
-	void FMidiClock::FMidiClockEventCursor::SeekToTick(int32 Tick) 
+	void FMidiClock::PostTempoOrTimeSignatureEventsIfNeeded()
 	{
-		const int32 FromTick = CurrentTick;
-		FMidiPlayCursor::SeekToTick(Tick);
-		AddEvent(FMidiClockEvent(MyMidiClock->GetCurrentBlockFrameIndex(), MidiClockMessageTypes::FSeekTo(FromTick, Tick)));
+		if (SongMapEvaluator->GetNumTempoChanges() == 0)
+		{
+			CurrentTempoInfoPointIndex = -1;
+			AddTempoChangeToBlock(CurrentBlockFrameIndex, NextMidiTickToProcess, 120.0f);
+			NextTempoChangeTick = std::numeric_limits<int32>::max();
+		}
+		else
+		{
+			CurrentTempoInfoPointIndex = SongMapEvaluator->GetTempoPointIndexForTick(NextMidiTickToProcess > 0 ? NextMidiTickToProcess : 0);
+			check(CurrentTempoInfoPointIndex >= 0 && CurrentTempoInfoPointIndex < SongMapEvaluator->GetNumTempoChanges());
+			if (!ExternalClockDriver)
+			{
+				AddTempoChangeToBlock(CurrentBlockFrameIndex, NextMidiTickToProcess, SongMapEvaluator->GetTempoInfoPoint(CurrentTempoInfoPointIndex)->GetBPM());
+			}
+			if ((CurrentTempoInfoPointIndex + 1) < SongMapEvaluator->GetNumTempoChanges())
+			{
+				NextTempoChangeTick = SongMapEvaluator->GetTempoChangePointTick(CurrentTempoInfoPointIndex + 1);
+			}
+			else
+			{
+				NextTempoChangeTick = std::numeric_limits<int32>::max();
+			}
+		}
+
+		if (SongMapEvaluator->GetNumTimeSignatureChanges() == 0)
+		{
+			CurrentTimeSignaturePointIndex = -1;
+			AddTimeSignatureChangeToBlock(CurrentBlockFrameIndex, NextMidiTickToProcess, FTimeSignature(4,4));
+			NextTimeSigChangeTick = std::numeric_limits<int32>::max();
+		}
+		else
+		{
+			CurrentTimeSignaturePointIndex = SongMapEvaluator->GetTimeSignaturePointIndexForTick(NextMidiTickToProcess > 0 ? NextMidiTickToProcess : 0);
+			check(CurrentTimeSignaturePointIndex >= 0 && CurrentTimeSignaturePointIndex < SongMapEvaluator->GetNumTimeSignatureChanges());
+			AddTimeSignatureChangeToBlock(CurrentBlockFrameIndex, NextMidiTickToProcess, SongMapEvaluator->GetTimeSignaturePoint(CurrentTimeSignaturePointIndex)->TimeSignature);
+			if ((CurrentTimeSignaturePointIndex + 1) < SongMapEvaluator->GetNumTimeSignatureChanges())
+			{
+				NextTimeSigChangeTick = SongMapEvaluator->GetTimeSignatureChangePointTick(CurrentTempoInfoPointIndex + 1);
+			}
+			else
+			{
+				NextTimeSigChangeTick = std::numeric_limits<int32>::max();
+			}
+			NextTempoOrTimeSigChangeTick = FMath::Min(NextTimeSigChangeTick, NextTempoChangeTick);
+		}
+
+		if (CurrentLocalSpeed < 0.0f)
+		{
+			SetSpeed(CurrentBlockFrameIndex, 1.0f);
+		}
 	}
 
-	void FMidiClock::FMidiClockEventCursor::SeekThruTick(int32 Tick)
+	void FMidiClock::AddTransportStateChangeToBlock(int32 BlockFrameIndex, EMusicPlayerTransportState NewTransportState)
 	{
-		const int32 FromTick = CurrentTick;
-		FMidiPlayCursor::SeekThruTick(Tick);
-		AddEvent(FMidiClockEvent(MyMidiClock->GetCurrentBlockFrameIndex(), MidiClockMessageTypes::FSeekThru(FromTick, Tick)));
+		check(BlockFrameIndex >= CurrentBlockFrameIndex);
+
+		CurrentBlockFrameIndex = BlockFrameIndex;
+
+		if (TransportAtBlockEnd == NewTransportState)
+		{
+			// no need to add the transport message... it is already the current transport. 
+			return;
+		}
+
+		if (NewTransportState == EMusicPlayerTransportState::Playing || NewTransportState == EMusicPlayerTransportState::Prepared)
+		{
+			PostTempoOrTimeSignatureEventsIfNeeded();
+			FramesUntilNextProcess = 0;
+		}
+
+		AddEvent(FMidiClockEvent(BlockFrameIndex, MidiClockMessageTypes::FTransportChange(NewTransportState)));
+		NumTransportChangeInBlock++;
+
+		if (BlockFrameIndex == 0)
+		{
+			TransportAtBlockStart = NewTransportState;
+		}
+
+		TransportAtBlockEnd = NewTransportState;
+	}
+
+	void FMidiClock::AddTimeSignatureChangeToBlock(int32 BlockFrameIndex, int32 Tick, const FTimeSignature& TimeSignature)
+	{
+		check(BlockFrameIndex >= CurrentBlockFrameIndex);
+
+		check(Tick >= NextMidiTickToProcess);
+
+		CurrentBlockFrameIndex = BlockFrameIndex;
+		NextMidiTickToProcess = Tick;
+
+		if (TimeSignatureAtBlockEnd == TimeSignature)
+		{
+			// no need to add the time signature... it is already the current time signature. 
+			return;
+		}
+
+		if (MidiClockMessageTypes::FTimeSignatureChange* PreviousTimeSigChangeOnSameTick = LookForEventOnMidiTick<MidiClockMessageTypes::FTimeSignatureChange>(BlockFrameIndex))
+		{
+			PreviousTimeSigChangeOnSameTick->TimeSignature = TimeSignature;
+		}
+		else
+		{
+			AddEvent(FMidiClockEvent(BlockFrameIndex, MidiClockMessageTypes::FTimeSignatureChange(Tick, FTimeSignature(TimeSignature))));
+			NumTimeSignatureChangeInBlock++;
+		}
+
+		if (BlockFrameIndex == 0)
+		{
+			TimeSignatureAtBlockStart = TimeSignature;
+		}
+
+		TimeSignatureAtBlockEnd = TimeSignature;
 	}
 
-	void FMidiClock::FMidiClockEventCursor::AdvanceThruTick(int32 Tick, bool IsPreRoll)
+	void FMidiClock::AddTempoChangeToBlock(int32 BlockFrameIndex, int32 Tick, float NewTempo)
 	{
-		CurrentAdvanceStartTick = CurrentTick;
-		FMidiPlayCursor::AdvanceThruTick(Tick, IsPreRoll);
-		AddEvent(FMidiClockEvent(MyMidiClock->GetCurrentBlockFrameIndex(), MidiClockMessageTypes::FAdvanceThru(CurrentAdvanceStartTick, Tick, IsPreRoll)));
+		check(BlockFrameIndex >= CurrentBlockFrameIndex);
+
+		check(Tick >= NextMidiTickToProcess);
+
+		CurrentBlockFrameIndex = BlockFrameIndex;
+		NextMidiTickToProcess = Tick;
+
+		if (TempoAtBlockEnd == NewTempo)
+		{
+			// no need to add the tempo... it is already the current tempo. 
+			return;
+		}
+
+		if (MidiClockMessageTypes::FTempoChange* PreviousTempoChangeOnSameTick = LookForEventOnMidiTick<MidiClockMessageTypes::FTempoChange>(BlockFrameIndex))
+		{
+			PreviousTempoChangeOnSameTick->Tempo = NewTempo;
+		}
+		else
+		{
+			AddEvent(FMidiClockEvent(BlockFrameIndex, MidiClockMessageTypes::FTempoChange(Tick, NewTempo)));
+			NumTempoChangeInBlock++;
+		}
+
+		if (BlockFrameIndex == 0)
+		{
+			TempoAtBlockStart = NewTempo;
+		}
+
+		TempoAtBlockEnd = NewTempo;
 	}
 
-	void FMidiClock::FMidiClockEventCursor::OnTempo(int32 TrackIndex, int32 Tick, int32 Tempo, bool IsPreroll)
+	void FMidiClock::AddSpeedChangeToBlock(int32 BlockFrameIndex, float NewSpeed, bool bIsNewLocalSpeed)
 	{
-		const int32 BlockFrameIndex = MyMidiClock->GetCurrentBlockFrameIndex();
-		const float Bpm = Harmonix::Midi::Constants::MidiTempoToBPM(Tempo);
+		check(BlockFrameIndex >= CurrentBlockFrameIndex);
+
+		CurrentBlockFrameIndex = BlockFrameIndex;
 
-		if (BlockFrameIndex >= MyMidiClock->TempoChangesInBlock.Last().BlockSampleFrameIndex)
+		if (ExternalClockDriver)
 		{
-			MyMidiClock->HasTempoChangeInBlock = true;
-			
-			if (MyMidiClock->TempoChangesInBlock.Last().BlockSampleFrameIndex == BlockFrameIndex)
+			// actual speed message is the 
+			if (bIsNewLocalSpeed)
 			{
-				MyMidiClock->TempoChangesInBlock.Last().Tempo = Bpm;
+				NewSpeed *= ExternalClockDriver->GetSpeedAtBlockSampleFrame(BlockFrameIndex);
 			}
 			else
-			{ 
-				MyMidiClock->TempoChangesInBlock.Add({BlockFrameIndex, 0.0f, Bpm});
+			{
+				NewSpeed *= CurrentLocalSpeed;
 			}
 		}
 
-		// CurrentAdvanceStartTick is the tick ALREADY processed in a previous advance. So the first new tick 
-		// we are processing is CurrentAdvanceStartTick + 1.
-		// 
-		// This gets fired off as a result of AdvanceThruTick, and we want the event to go in between advances,
-		// so if this tempo change is not on the first new tick of this advance we need to "cap off" the 
-		// currently in progress advance and start a new one. 
-		const int32 FirstNewTickThisAdvance = CurrentAdvanceStartTick + 1;
-		const int32 TickPriorToThisEvent = Tick - 1;
-		if (Tick > FirstNewTickThisAdvance)
+		if (SpeedAtBlockEnd == NewSpeed)
 		{
-			AddEvent(FMidiClockEvent(BlockFrameIndex, MidiClockMessageTypes::FAdvanceThru(CurrentAdvanceStartTick, TickPriorToThisEvent, IsPreroll)));
-			CurrentAdvanceStartTick = TickPriorToThisEvent;
+			// no need to add the speed... it is already the current speed. 
+			return;
 		}
-		
-		AddEvent(FMidiClockEvent(BlockFrameIndex, MidiClockMessageTypes::FTempoChange(Tick, Bpm)));
+
+		if (MidiClockMessageTypes::FSpeedChange* PreviousSpeedChangeOnSameTick = LookForEventOnBlockFrameIndex<MidiClockMessageTypes::FSpeedChange>(BlockFrameIndex))
+		{
+			PreviousSpeedChangeOnSameTick->Speed = NewSpeed;
+		}
+		else
+		{
+			AddEvent(FMidiClockEvent(BlockFrameIndex, MidiClockMessageTypes::FSpeedChange(NewSpeed)));
+			NumSpeedChangeInBlock++;
+		}
+
+		if (BlockFrameIndex == 0)
+		{
+			SpeedAtBlockStart = NewSpeed;
+		}
+
+		SpeedAtBlockEnd = NewSpeed;
 	}
 
-	void FMidiClock::FMidiClockEventCursor::OnTimeSig(int32 TrackIndex, int32 Tick, int32 Numerator, int32 Denominator, bool IsPreroll)
+	void FMidiClock::AddLoopToBlock(int32 BlockFrameIndex, int32 FirstTick, int32 LoopLength)
 	{
-		// CurrentAdvanceStartTick is the tick ALREADY processed in a previous advance. So the first new tick 
-		// we are processing is CurrentAdvanceStartTick + 1.
-		// 
-		// This gets fired off as a result of AdvanceThruTick, and we want the event to go in between advances,
-		// so if this time signature change is not on the first new tick of this advance we need to "cap off" the 
-		// currently in progress advance and start a new one. 
-		const int32 FirstNewTickThisAdvance = CurrentAdvanceStartTick + 1;
-		const int32 TickPriorToThisEvent = Tick - 1;
-		if (Tick > FirstNewTickThisAdvance)
+		check(BlockFrameIndex >= CurrentBlockFrameIndex);
+		check(LastProcessedMidiTick == FirstTick + LoopLength - 1);
+
+		CurrentBlockFrameIndex = BlockFrameIndex;
+
+		if (NextMidiTickToProcess == FirstTick)
 		{
-			AddEvent(FMidiClockEvent(MyMidiClock->GetCurrentBlockFrameIndex(), MidiClockMessageTypes::FAdvanceThru(CurrentAdvanceStartTick, TickPriorToThisEvent, IsPreroll)));
-			CurrentAdvanceStartTick = TickPriorToThisEvent;
+			return;
 		}
-		
-		AddEvent(FMidiClockEvent(MyMidiClock->GetCurrentBlockFrameIndex(), MidiClockMessageTypes::FTimeSignatureChange(Tick, { Numerator, Denominator })));
+
+		if (MidiClockMessageTypes::FLoop* PreviousLoopOnSameTick = LookForEventOnBlockFrameIndex<MidiClockMessageTypes::FLoop>(BlockFrameIndex))
+		{
+			PreviousLoopOnSameTick->LengthInTicks = LoopLength;
+			PreviousLoopOnSameTick->FirstTickInLoop = FirstTick;
+		}
+		else
+		{
+			AddEvent(FMidiClockEvent(BlockFrameIndex, MidiClockMessageTypes::FLoop(FirstTick, LoopLength)));
+		}
+
+		NextMidiTickToProcess = FirstTickInLoop;
+
+		// The tempo and/or time signature may be different in the location where we are going... so... update...
+		PostTempoOrTimeSignatureEventsIfNeeded();
 	}
 
-	void FMidiClock::FMidiClockEventCursor::AddEvent(const FMidiClockEvent& InEvent) const
+	void FMidiClock::AddSeekToBlock(int32 BlockFrameIndex, int32 ToTick)
 	{
-		TArray<FMidiClockEvent>& Events = MyMidiClock->MidiClockEventsInBlock;
-		if (!Events.IsEmpty())
+		check(BlockFrameIndex >= CurrentBlockFrameIndex);
+
+		CurrentBlockFrameIndex = BlockFrameIndex;
+
+		if (NextMidiTickToProcess == ToTick)
+		{
+			return;
+		}
+
+		if (MidiClockMessageTypes::FSeek* PreviousSeekOnSameTick = LookForEventOnBlockFrameIndex<MidiClockMessageTypes::FSeek>(BlockFrameIndex))
 		{
-			check(Events.Last().BlockFrameIndex <= InEvent.BlockFrameIndex);
+			check(PreviousSeekOnSameTick->LastTickProcessedBeforeSeek == LastProcessedMidiTick);
+			PreviousSeekOnSameTick->NewNextTick = ToTick;
 		}
+		else
+		{
+			AddEvent(FMidiClockEvent(BlockFrameIndex, MidiClockMessageTypes::FSeek(LastProcessedMidiTick, ToTick)));
+		}
+
+		NextMidiTickToProcess = ToTick;
+
+		// The tempo and/or time signature may be different in the location where we are going... so... update...
+		PostTempoOrTimeSignatureEventsIfNeeded();
+	}
 
-		Events.Add(InEvent);
+	void FMidiClock::AddAdvanceToBlock(int32 BlockFrameIndex, int32 FirstTick, int32 NumTicks)
+	{
+		check(BlockFrameIndex >= CurrentBlockFrameIndex);
+		check(FirstTick == NextMidiTickToProcess);
+
+		CurrentBlockFrameIndex = BlockFrameIndex;
+
+		// moving forward may cause us to move into a new tempo and/or time signature...
+		while (FirstTick <= NextTempoOrTimeSigChangeTick && NextTempoOrTimeSigChangeTick < (FirstTick + NumTicks))
+		{
+			// process ticks UP TO the tempo or time signature change...
+			int32 SpanNumTicks = NextTempoOrTimeSigChangeTick - FirstTick;
+			if (SpanNumTicks > 0)
+			{
+				AddEvent(FMidiClockEvent(BlockFrameIndex, MidiClockMessageTypes::FAdvance(FirstTick, SpanNumTicks)));
+				NextMidiTickToProcess = FirstTick + SpanNumTicks;
+				LastProcessedMidiTick = NextMidiTickToProcess - 1;
+				FirstTick = NextMidiTickToProcess;
+				NumTicks -= SpanNumTicks;
+			}
+			check(NextMidiTickToProcess == NextTempoOrTimeSigChangeTick);
+
+			// time signature or tempo?
+			if (NextMidiTickToProcess == NextTempoChangeTick && CurrentTempoInfoPointIndex != -1)
+			{
+				CurrentTempoInfoPointIndex++;
+				const FTempoInfoPoint* TempoPoint = SongMapEvaluator->GetTempoInfoPoint(CurrentTempoInfoPointIndex);
+				check(TempoPoint);
+				if (!ExternalClockDriver)
+				{
+					AddTempoChangeToBlock(BlockFrameIndex, TempoPoint->StartTick, TempoPoint->GetBPM());
+				}
+				int32 NextTempoInfoPointIndex = CurrentTempoInfoPointIndex + 1;
+				if (NextTempoInfoPointIndex < SongMapEvaluator->GetNumTempoChanges())
+				{
+					const FTempoInfoPoint* NextTempoPoint = SongMapEvaluator->GetTempoInfoPoint(NextTempoInfoPointIndex);
+					check(NextTempoPoint);
+					NextTempoChangeTick = NextTempoPoint->StartTick;
+				}
+				else
+				{
+					NextTempoChangeTick = std::numeric_limits<int32>::max();
+				}
+			}
+			// time signature or tempo?
+			if (NextMidiTickToProcess == NextTimeSigChangeTick && CurrentTimeSignaturePointIndex != -1)
+			{
+				CurrentTimeSignaturePointIndex++;
+				const FTimeSignaturePoint* TimeSignaturePoint = SongMapEvaluator->GetTimeSignaturePoint(CurrentTimeSignaturePointIndex);
+				check(TimeSignaturePoint);
+				AddTimeSignatureChangeToBlock(BlockFrameIndex, TimeSignaturePoint->StartTick, TimeSignaturePoint->TimeSignature);
+				int32 NextTimeSignaturePointIndex = CurrentTimeSignaturePointIndex + 1;
+				if (NextTimeSignaturePointIndex < SongMapEvaluator->GetNumTimeSignatureChanges())
+				{
+					const FTimeSignaturePoint* NextTimeSigPoint = SongMapEvaluator->GetTimeSignaturePoint(NextTimeSignaturePointIndex);
+					check(NextTimeSigPoint);
+					NextTimeSigChangeTick = NextTimeSigPoint->StartTick;
+				}
+				else
+				{
+					NextTimeSigChangeTick = std::numeric_limits<int32>::max();
+				}
+			}
+			NextTempoOrTimeSigChangeTick = FMath::Min(NextTimeSigChangeTick, NextTempoChangeTick);
+		}
+
+		if (NumTicks > 0)
+		{
+			AddEvent(FMidiClockEvent(BlockFrameIndex, MidiClockMessageTypes::FAdvance(FirstTick, NumTicks)));
+			NextMidiTickToProcess = FirstTick + NumTicks;
+			LastProcessedMidiTick = NextMidiTickToProcess - 1;
+		}
+	}
+
+	void FMidiClock::RebuildSongMapEvaluator(const TSharedPtr<const FMidiFileData>& MidiWithTempo, const TSharedPtr<const FMidiFileData>& MidiWithOtherMaps)
+	{
+		SongMapEvaluator = MakeShared<FSongMapsWithAlternateTempoSource>(MidiWithTempo, MidiWithOtherMaps);
 	}
 
 }
diff --git a/Engine/Plugins/Runtime/Harmonix/Source/HarmonixMetasound/Private/HarmonixMetasound/DataTypes/MusicTimeInterval.cpp b/Engine/Plugins/Runtime/Harmonix/Source/HarmonixMetasound/Private/HarmonixMetasound/DataTypes/MusicTimeInterval.cpp
--- a/Engine/Plugins/Runtime/Harmonix/Source/HarmonixMetasound/Private/HarmonixMetasound/DataTypes/MusicTimeInterval.cpp
+++ b/Engine/Plugins/Runtime/Harmonix/Source/HarmonixMetasound/Private/HarmonixMetasound/DataTypes/MusicTimeInterval.cpp
@@ -1,13 +1,13 @@
 // Copyright Epic Games, Inc. All Rights Reserved.
 
 #include "HarmonixMetasound/DataTypes/MusicTimeInterval.h"
-#include "HarmonixMidi/SongMaps.h"
+#include "HarmonixMidi/MidiConstants.h"
 
 namespace Harmonix
 {
 	float GetIntervalInBeats(const FMusicTimeInterval& Interval, const FTimeSignature& TimeSignature)
 	{
-		return FSongMaps::SubdivisionToBeats(Interval.Interval, TimeSignature) * Interval.IntervalMultiplier;
+		return Harmonix::Midi::Constants::SubdivisionToBeats(Interval.Interval, TimeSignature) * Interval.IntervalMultiplier;
 	}
 
 	void IncrementTimestampByBeats(FMusicTimestamp& Timestamp, float Beats, const FTimeSignature& TimeSignature)
@@ -43,7 +43,7 @@ namespace Harmonix
 	float GetOffsetInBeats(const FMusicTimeInterval& Interval, const FTimeSignature& TimeSignature)
 	{
 		const float IntervalBeats = GetIntervalInBeats(Interval, TimeSignature);
-		const float OffsetBeats = FSongMaps::SubdivisionToBeats(Interval.Offset, TimeSignature) * Interval.OffsetMultiplier;
+		const float OffsetBeats = Harmonix::Midi::Constants::SubdivisionToBeats(Interval.Offset, TimeSignature) * Interval.OffsetMultiplier;
 		return FMath::Fmod(OffsetBeats, IntervalBeats);
 	}
 
diff --git a/Engine/Plugins/Runtime/Harmonix/Source/HarmonixMetasound/Private/HarmonixMetasound/DataTypes/MusicTransport.cpp b/Engine/Plugins/Runtime/Harmonix/Source/HarmonixMetasound/Private/HarmonixMetasound/DataTypes/MusicTransport.cpp
--- a/Engine/Plugins/Runtime/Harmonix/Source/HarmonixMetasound/Private/HarmonixMetasound/DataTypes/MusicTransport.cpp
+++ b/Engine/Plugins/Runtime/Harmonix/Source/HarmonixMetasound/Private/HarmonixMetasound/DataTypes/MusicTransport.cpp
@@ -239,4 +239,23 @@ namespace HarmonixMetasound
 		       TransportState == EMusicPlayerTransportState::Stopping ||
 		       TransportState == EMusicPlayerTransportState::Killing;
 	}
+
+	FString HARMONIXMETASOUND_API MusicPlayerTransportStateToString(EMusicPlayerTransportState State)
+	{
+		switch (State)
+		{
+		case EMusicPlayerTransportState::Invalid:    return TEXT("Invalid");
+		case EMusicPlayerTransportState::Preparing:	 return TEXT("Preparing");
+		case EMusicPlayerTransportState::Prepared:	 return TEXT("Prepared");
+		case EMusicPlayerTransportState::Starting:	 return TEXT("Starting");
+		case EMusicPlayerTransportState::Playing:	 return TEXT("Playing");
+		case EMusicPlayerTransportState::Seeking:	 return TEXT("Seeking");
+		case EMusicPlayerTransportState::Pausing:	 return TEXT("Pausing");
+		case EMusicPlayerTransportState::Paused:	 return TEXT("Paused");
+		case EMusicPlayerTransportState::Continuing: return TEXT("Continuing");
+		case EMusicPlayerTransportState::Stopping:	 return TEXT("Stopping");
+		case EMusicPlayerTransportState::Killing:	 return TEXT("Killing");
+		}											
+		return TEXT("<Unrecognized>");
+	}
 }
diff --git a/Engine/Plugins/Runtime/Harmonix/Source/HarmonixMetasound/Private/HarmonixMetasound/MidiOps/MidiStreamWriter.cpp b/Engine/Plugins/Runtime/Harmonix/Source/HarmonixMetasound/Private/HarmonixMetasound/MidiOps/MidiStreamWriter.cpp
--- a/Engine/Plugins/Runtime/Harmonix/Source/HarmonixMetasound/Private/HarmonixMetasound/MidiOps/MidiStreamWriter.cpp
+++ b/Engine/Plugins/Runtime/Harmonix/Source/HarmonixMetasound/Private/HarmonixMetasound/MidiOps/MidiStreamWriter.cpp
@@ -1,4 +1,4 @@
-﻿// Copyright Epic Games, Inc. All Rights Reserved.
+// Copyright Epic Games, Inc. All Rights Reserved.
 
 #include "HarmonixMetasound/MidiOps/MidiStreamWriter.h"
 #include "HarmonixMetasound/DataTypes/MidiStream.h"
@@ -16,37 +16,32 @@ namespace Harmonix::Midi::Ops
 		ensureMsgf(InStream.GetClock(), TEXT("Midi stream must have a midi clock for the MidiStreamWriter to process it"));
 		if (TSharedPtr<const FMidiClock, ESPMode::NotThreadSafe> MidiClock = InStream.GetClock())
 		{
+			using namespace MidiClockMessageTypes;
 			for (const FMidiClockEvent& ClockEvent : MidiClock->GetMidiClockEventsInBlock())
 			{
-				if (ClockEvent.Msg.IsType<MidiClockMessageTypes::FAdvanceThru>())
+				if (const FAdvance* AsAdvance = ClockEvent.TryGet<MidiClockMessageTypes::FAdvance>())
 				{
-					const MidiClockMessageTypes::FAdvanceThru& AdvanceThru = ClockEvent.Msg.Get<MidiClockMessageTypes::FAdvanceThru>();
-					
-					if (!AdvanceThru.IsPreRoll)
-					{
-						Process(InStream, AdvanceThru.FromTick, AdvanceThru.ThruTick);
-					}
+					Process(InStream, AsAdvance->FirstTickToProcess, AsAdvance->LastTickToProcess());
 				}
 			}
 		}
 	}
 
-
-	void FMidiStreamWriter::Process(const FMidiStream& InStream, int32 FromTick, int32 ThruTick)
+	void FMidiStreamWriter::Process(const FMidiStream& InStream, int32 FirstTickToProcess, int32 LastTickToProcess)
 	{
 		bool AddedEvents = false;
 		for (const FMidiStreamEvent& MidiEvent : InStream.GetEventsInBlock())
 		{
-			if (MidiEvent.CurrentMidiTick > FromTick && MidiEvent.CurrentMidiTick <= ThruTick)
+			if (FirstTickToProcess <= MidiEvent.CurrentMidiTick && MidiEvent.CurrentMidiTick <= LastTickToProcess)
 			{
-				const int32 OffsetTick = MidiEvent.CurrentMidiTick - FromTick;
-				const int32 MidiTick = CurrentWriteTick + OffsetTick;
+				const int32 OffsetTick = MidiEvent.CurrentMidiTick - FirstTickToProcess;
+				const int32 MidiTick = NextWriteTick + OffsetTick;
 				FMidiTrack& MidiTrack = MidiTracks.FindOrAdd(MidiEvent.TrackIndex);
 				MidiTrack.AddEvent(FMidiEvent(MidiTick, MidiEvent.MidiMessage));
 				AddedEvents = true;
 			}
 		}
-		CurrentWriteTick += ThruTick - FromTick;
+		NextWriteTick += LastTickToProcess - FirstTickToProcess + 1;
 
 		if (AddedEvents)
 		{
diff --git a/Engine/Plugins/Runtime/Harmonix/Source/HarmonixMetasound/Private/HarmonixMetasound/MidiOps/PulseGenerator.cpp b/Engine/Plugins/Runtime/Harmonix/Source/HarmonixMetasound/Private/HarmonixMetasound/MidiOps/PulseGenerator.cpp
--- a/Engine/Plugins/Runtime/Harmonix/Source/HarmonixMetasound/Private/HarmonixMetasound/MidiOps/PulseGenerator.cpp
+++ b/Engine/Plugins/Runtime/Harmonix/Source/HarmonixMetasound/Private/HarmonixMetasound/MidiOps/PulseGenerator.cpp
@@ -38,9 +38,9 @@ namespace Harmonix::Midi::Ops
 
 		for (const FMidiClockEvent& ClockEvent : PinnedClock->GetMidiClockEventsInBlock())
 		{
-			if (const FAdvanceThru* AsAdvanceThru = ClockEvent.TryGet<FAdvanceThru>())
+			if (const FAdvance* AsAdvance = ClockEvent.TryGet<FAdvance>())
 			{
-				HandleAdvanceThru(ClockEvent.BlockFrameIndex, AsAdvanceThru, OutStream);
+				HandleAdvance(ClockEvent.BlockFrameIndex, AsAdvance, OutStream);
 			}
 			else if (const FTimeSignatureChange* AsTimeSigChange = ClockEvent.TryGet<FTimeSignatureChange>())
 			{
@@ -49,7 +49,7 @@ namespace Harmonix::Midi::Ops
 		}
 	}
 
-	void FPulseGenerator::HandleAdvanceThru(const int32 BlockFrameIndex, const HarmonixMetasound::MidiClockMessageTypes::FAdvanceThru* AdvanceThru, HarmonixMetasound::FMidiStream& OutStream)
+	void FPulseGenerator::HandleAdvance(const int32 BlockFrameIndex, const HarmonixMetasound::MidiClockMessageTypes::FAdvance* Advance, HarmonixMetasound::FMidiStream& OutStream)
 	{
 		if (!NextPulseTimestamp.IsValid())
 		{
@@ -58,15 +58,15 @@ namespace Harmonix::Midi::Ops
 
 		check(Clock.IsValid()); // if we're here and we don't have a clock, we have problems
 		const auto PinnedClock = Clock.Pin();
-		int32 NextPulseTick = PinnedClock->GetBarMap().MusicTimestampToTick(NextPulseTimestamp);
+		int32 NextPulseTick = PinnedClock->GetSongMapEvaluator().MusicTimestampToTick(NextPulseTimestamp);
 
-		while (AdvanceThru->ThruTick >= NextPulseTick)
+		while (Advance->LastTickToProcess() >= NextPulseTick)
 		{
 			DoPulse(BlockFrameIndex, NextPulseTick, OutStream);
 
 			IncrementTimestampByInterval(NextPulseTimestamp, Interval, CurrentTimeSignature);
 
-			NextPulseTick = PinnedClock->GetBarMap().MusicTimestampToTick(NextPulseTimestamp);
+			NextPulseTick = PinnedClock->GetSongMapEvaluator().MusicTimestampToTick(NextPulseTimestamp);
 		}
 	}
 
@@ -78,7 +78,7 @@ namespace Harmonix::Midi::Ops
 		const auto PinnedClock = Clock.Pin();
 		// Time sig changes will come on the downbeat, and if we change time signature,
 		// we want to reset the pulse, so the next pulse is now plus the offset
-		NextPulseTimestamp = PinnedClock->GetBarMap().TickToMusicTimestamp(TimeSigChange->Tick);
+		NextPulseTimestamp = PinnedClock->GetSongMapEvaluator().TickToMusicTimestamp(TimeSigChange->Tick);
 		IncrementTimestampByOffset(NextPulseTimestamp, Interval, CurrentTimeSignature);
 	}
 
diff --git a/Engine/Plugins/Runtime/Harmonix/Source/HarmonixMetasound/Private/HarmonixMetasound/Nodes/DelayNode.cpp b/Engine/Plugins/Runtime/Harmonix/Source/HarmonixMetasound/Private/HarmonixMetasound/Nodes/DelayNode.cpp
--- a/Engine/Plugins/Runtime/Harmonix/Source/HarmonixMetasound/Private/HarmonixMetasound/Nodes/DelayNode.cpp
+++ b/Engine/Plugins/Runtime/Harmonix/Source/HarmonixMetasound/Private/HarmonixMetasound/Nodes/DelayNode.cpp
@@ -359,26 +359,29 @@ namespace HarmonixMetasound::DelayNode
 				int32 SamplesRendered = 0;
 				
 				// if the tempo or speed changes, render in chunks
-				for (int SampleIdx = 0; SampleIdx < NumSamplesInBlock; ++SampleIdx)
+				for (const FMidiClockEvent& ClockEvent : Clock->GetMidiClockEventsInBlock())
 				{
 					bool ShouldRenderNow = false;
 					
-					if (const float Tempo = Clock->GetTempoAtBlockSampleFrame(SampleIdx); Tempo != LastTempo)
+					const MidiClockMessageTypes::FTempoChange* AsTempoChange = ClockEvent.TryGet<MidiClockMessageTypes::FTempoChange>();
+					const MidiClockMessageTypes::FSpeedChange* AsSpeedChange = ClockEvent.TryGet<MidiClockMessageTypes::FSpeedChange>();
+
+					if (AsTempoChange && AsTempoChange->Tempo != LastTempo)
 					{
-						LastTempo = Tempo;
+						LastTempo = AsTempoChange->Tempo;
 						ShouldRenderNow = true;
 					}
 
-					if (const float Speed = Clock->GetSpeedAtBlockSampleFrame(SampleIdx); Speed != LastSpeed)
+					if (AsSpeedChange && AsSpeedChange->Speed != LastSpeed)
 					{
-						LastSpeed = Speed;
+						LastSpeed = AsSpeedChange->Speed;
 						ShouldRenderNow = true;
 					}
 
 					if (ShouldRenderNow)
 					{
 						// render the delay
-						const int32 SamplesToRender = SampleIdx - SamplesRendered;
+						const int32 SamplesToRender = ClockEvent.BlockFrameIndex - SamplesRendered;
 						RenderSubBlock(SamplesRendered, SamplesToRender);
 						SamplesRendered += SamplesToRender;
 
diff --git a/Engine/Plugins/Runtime/Harmonix/Source/HarmonixMetasound/Private/HarmonixMetasound/Nodes/FusionSamplerNode.cpp b/Engine/Plugins/Runtime/Harmonix/Source/HarmonixMetasound/Private/HarmonixMetasound/Nodes/FusionSamplerNode.cpp
--- a/Engine/Plugins/Runtime/Harmonix/Source/HarmonixMetasound/Private/HarmonixMetasound/Nodes/FusionSamplerNode.cpp
+++ b/Engine/Plugins/Runtime/Harmonix/Source/HarmonixMetasound/Private/HarmonixMetasound/Nodes/FusionSamplerNode.cpp
@@ -663,8 +663,15 @@ namespace HarmonixMetasound
 				SetSpeed(ClockSpeed, !(*ClockSpeedAffectsPitchInPin));
 				const float ClockTempo = MidiClock->GetTempoAtBlockSampleFrame(CurrentBlockFrameIndex);
 				SetTempo(ClockTempo);
-				const float Beat = MidiClock->GetQuarterNoteIncludingCountIn();
-				SetBeat(Beat);
+				// This next thing is *close* to being accurate, but not really. 
+				// It gets the last processed beat in this metasound render block, which may or may not
+				// be correct for this specific sample in the block we are processing at the moment. 
+				// It should be close enough. If not, we will need to ask the clock specifically what
+				// its tick was at this point in the metasound rendering block. 
+				const int32 Tick = MidiClock->GetNextMidiTickToProcess();
+				const ISongMapEvaluator& ClocksMap = MidiClock->GetSongMapEvaluator();
+				const float MidiQuarterNote = (float)Tick / (float)ClocksMap.GetTicksPerQuarterNote();
+				SetQuarterNote(MidiQuarterNote);
 			}
 
 			if (FramesRequired > AudioRendering::kFramesPerRenderBuffer)
diff --git a/Engine/Plugins/Runtime/Harmonix/Source/HarmonixMetasound/Private/HarmonixMetasound/Nodes/MetronomeNode.cpp b/Engine/Plugins/Runtime/Harmonix/Source/HarmonixMetasound/Private/HarmonixMetasound/Nodes/MetronomeNode.cpp
--- a/Engine/Plugins/Runtime/Harmonix/Source/HarmonixMetasound/Private/HarmonixMetasound/Nodes/MetronomeNode.cpp
+++ b/Engine/Plugins/Runtime/Harmonix/Source/HarmonixMetasound/Private/HarmonixMetasound/Nodes/MetronomeNode.cpp
@@ -42,8 +42,6 @@ namespace HarmonixMetasound
 
 		virtual void BindInputs(FInputVertexInterfaceData& InVertexData) override;
 		virtual void BindOutputs(FOutputVertexInterfaceData& InVertexData) override;
-		virtual FDataReferenceCollection GetInputs() const override;
-		virtual FDataReferenceCollection GetOutputs() const override;
 
 		void Reset(const FResetParams& Params);
 		void Execute();
@@ -65,22 +63,26 @@ namespace HarmonixMetasound
 		FMidiClockWriteRef MidiClockOutPin;
 
 		//** DATA
-		FMidiClock MetronomeClock;
+		TSharedPtr<FMidiClock, ESPMode::NotThreadSafe> MonotonicallyIncreasingClock;
 		TSharedPtr<FMidiFileData> MidiData;
 		FSampleCount BlockSize;
 		float        SampleRate;
 		float        CurrentTempo;
 		int32        CurrentTimeSigNum;
 		int32        CurrentTimeSigDenom;
-		int32		 LastClockTickUpdate = -1;
+		int32		 LastProcessedClockTick = -1;
+		int32		 NextClockTickToProcess = 0;
+		bool		 bClocksArePreparedForExecute = true;
 
 		void BuildMidiData(bool ResetToStart = true);
 		void UpdateMidi();
 		void AddTempoChangeForMidi(float TempoBPM);
 		void AddTimeSigChangeForMidi(int32 TimeSigNum, int32 TimeSigDenom);
 		void HandleTransportChange(int32 StartFrameIndex, EMusicPlayerTransportState NewTransportState);
+		void PrepareClocksForExecute();
+		void MarkClocksAsExecuted();
 
-		FMidiClock& GetDrivingMidiClock() { return LoopInPin ? MetronomeClock : (*MidiClockOutPin); }
+		FMidiClock& GetDrivingMidiClock() { return LoopInPin ? *MonotonicallyIncreasingClock : (*MidiClockOutPin); }
 	};
 
 	class FMetronomeNode : public FNodeFacade
@@ -180,20 +182,14 @@ namespace HarmonixMetasound
 		, SpeedMultInPin(InSpeedMultiplier)
 		, SeekPreRollBarsInPin(InPreRollBars)
 		, MidiClockOutPin(FMidiClockWriteRef::CreateNew(InParams.OperatorSettings))
-		, MetronomeClock(InParams.OperatorSettings)
+		, MonotonicallyIncreasingClock(MakeShared<FMidiClock, ESPMode::NotThreadSafe>(InParams.OperatorSettings))
 		, BlockSize(InParams.OperatorSettings.GetNumFramesPerBlock())
 		, SampleRate(InParams.OperatorSettings.GetSampleRate())
 		, CurrentTempo(*TempoInPin)
 		, CurrentTimeSigNum(FMath::Clamp(*TimeSigNumInPin, 1, 64))
 		, CurrentTimeSigDenom(FMath::Clamp(*TimeSigDenomInPin, 1, 64))
 	{
-		if (LoopInPin)
-		{
-			MidiClockOutPin->AttachToTimeAuthority(MetronomeClock);
-		}
-
 		Reset(InParams);
-
 		Init();
 	}
 
@@ -218,32 +214,27 @@ namespace HarmonixMetasound
 		InVertexData.BindReadVertex(METASOUND_GET_PARAM_NAME(Outputs::MidiClock), MidiClockOutPin);
 	}
 
-	FDataReferenceCollection FMetronomeOperator::GetInputs() const
-	{
-		// This should never be called. Bind(...) is called instead. This method
-		// exists as a stop-gap until the API can be deprecated and removed.
-		checkNoEntry();
-		return {};
-	}
-
-	FDataReferenceCollection FMetronomeOperator::GetOutputs() const
-	{
-		// This should never be called. Bind(...) is called instead. This method
-		// exists as a stop-gap until the API can be deprecated and removed.
-		checkNoEntry();
-		return {};
-	}
-
 	void FMetronomeOperator::Reset(const FResetParams& Params)
 	{
 		BlockSize = Params.OperatorSettings.GetNumFramesPerBlock();
 		SampleRate = Params.OperatorSettings.GetSampleRate();
 		
-		LastClockTickUpdate = -1;
+		LastProcessedClockTick = -1;
+		NextClockTickToProcess = 0;
 	}
 
 	void FMetronomeOperator::Init()
 	{
+		bClocksArePreparedForExecute = false;
+		PrepareClocksForExecute();
+		
+		MonotonicallyIncreasingClock->SetSpeed(0, 1.0f);
+		MidiClockOutPin->SetSpeed(0, 1.0f);
+		if (LoopInPin)
+		{
+			MidiClockOutPin->SetDrivingClock(MonotonicallyIncreasingClock);
+		}
+
 		BuildMidiData();
 		
 		FTransportInitFn InitFn = [this](EMusicPlayerTransportState CurrentState)
@@ -252,11 +243,11 @@ namespace HarmonixMetasound
 			switch (CurrentState)
 			{
 			case EMusicPlayerTransportState::Starting:
-				DrivingMidiClock.ResetAndStart(0, true);
-				break;
-
-			case EMusicPlayerTransportState::Playing:
-				DrivingMidiClock.WriteAdvance(0, 0, *SpeedMultInPin);
+				MidiClockOutPin->SeekTo(0, 0);
+				if (LoopInPin)
+				{
+					MonotonicallyIncreasingClock->SeekTo(0, 0);
+				}
 				break;
 			}
 
@@ -270,37 +261,32 @@ namespace HarmonixMetasound
 
 	void FMetronomeOperator::Execute()
 	{
-		MidiClockOutPin->PrepareBlock();
+		PrepareClocksForExecute();
 
 		FMidiClock& DrivingMidiClock = GetDrivingMidiClock();
-		DrivingMidiClock.PrepareBlock();
 
 		if (*SpeedMultInPin != DrivingMidiClock.GetSpeedAtEndOfBlock())
 		{
-			DrivingMidiClock.AddSpeedChangeToBlock({ 0, 0.0f, *SpeedMultInPin });
+			DrivingMidiClock.SetSpeed(0, *SpeedMultInPin);
 		}
 
 		// only update our midi data if the clock is advancing
 		// update the tempo and time sig before we advance our clock 
-		int32 ClockTick = DrivingMidiClock.GetCurrentHiResTick();
-		if (ClockTick >= 0 && ClockTick > LastClockTickUpdate)
+		int32 ClockTick = DrivingMidiClock.GetLastProcessedMidiTick();
+		if (ClockTick >= 0 && ClockTick > LastProcessedClockTick)
 		{
 			UpdateMidi();
-			LastClockTickUpdate = ClockTick;
+			LastProcessedClockTick = ClockTick;
 		}
 
-		TransportSpanPostProcessor HandleMidiClockEvents = [this](int32 StartFrameIndex, int32 EndFrameIndex, EMusicPlayerTransportState CurrentState)
+		TransportSpanPostProcessor HandleMidiClockEvents = [this, &DrivingMidiClock](int32 StartFrameIndex, int32 EndFrameIndex, EMusicPlayerTransportState CurrentState)
 		{
 			int32 NumFrames = EndFrameIndex - StartFrameIndex;
 			HandleTransportChange(StartFrameIndex, CurrentState);
-			if (MidiClockOutPin->DoesLoop())
-			{
-				MetronomeClock.Process(StartFrameIndex, NumFrames, SeekPreRollBarsInPin, *SpeedMultInPin);
-				MidiClockOutPin->Process(MetronomeClock, StartFrameIndex, NumFrames, SeekPreRollBarsInPin, *SpeedMultInPin);
-			}
-			else
+			DrivingMidiClock.SetSpeed(StartFrameIndex, *SpeedMultInPin);
+			if (CurrentState == EMusicPlayerTransportState::Playing || CurrentState == EMusicPlayerTransportState::Continuing)
 			{
-				MidiClockOutPin->Process(StartFrameIndex, NumFrames, SeekPreRollBarsInPin, *SpeedMultInPin);
+				DrivingMidiClock.Advance(StartFrameIndex, NumFrames);
 			}
 		};
 
@@ -313,15 +299,20 @@ namespace HarmonixMetasound
 				if (!ReceivedSeekWhileStopped())
 				{
 					BuildMidiData(true);
-					LastClockTickUpdate = -1;
+					LastProcessedClockTick = -1;
+					NextClockTickToProcess = 0;
+				}
+				if (!ReceivedSeekWhileStopped())
+				{
+					DrivingMidiClock.SeekTo(StartFrameIndex, 0);
 				}
-				DrivingMidiClock.ResetAndStart(StartFrameIndex, !ReceivedSeekWhileStopped());
 				return EMusicPlayerTransportState::Playing;
 				
 			case EMusicPlayerTransportState::Seeking:
 				BuildMidiData(false);
-				DrivingMidiClock.SeekTo(StartFrameIndex, TransportInPin->GetNextSeekDestination(), SeekPreRollBarsInPin);
-				LastClockTickUpdate = DrivingMidiClock.GetCurrentMidiTick();
+				DrivingMidiClock.SeekTo(StartFrameIndex, TransportInPin->GetNextSeekDestination());
+				LastProcessedClockTick = DrivingMidiClock.GetLastProcessedMidiTick();
+				NextClockTickToProcess = DrivingMidiClock.GetNextMidiTickToProcess();
 				// Here we will return that we want to be in the same state we were in before this request to 
 				// seek since we can seek "instantaneously"...
 				return GetTransportState();
@@ -331,10 +322,12 @@ namespace HarmonixMetasound
 		};
 		ExecuteTransportSpans(TransportInPin, BlockSize, TransportHandler, HandleMidiClockEvents);
 
-		if (LoopInPin)
+		if (MidiClockOutPin->HasPersistentLoop())
 		{
-			MidiClockOutPin->CopySpeedAndTempoChanges(&MetronomeClock);
+			MidiClockOutPin->Advance(*MonotonicallyIncreasingClock, 0, BlockSize);
 		}
+
+		MarkClocksAsExecuted();
 	}
 
 	void FMetronomeOperator::BuildMidiData(bool ResetToStart)
@@ -348,22 +341,23 @@ namespace HarmonixMetasound
 		
 		if (LoopInPin)
 		{
-			MetronomeClock.AttachToMidiResource(MidiData, ResetToStart);
+			MonotonicallyIncreasingClock->AttachToMidiFile(MidiData, ResetToStart);
 
-			// midi clock out will follow the tempo of the metronome
+			// midi clock out will follow the tempo of the monotonically increasing clock
 			// so just assign it to some reasonable values. They will be ignored
 			TSharedPtr<FMidiFileData> MidiDataOut = FMidiClock::MakeClockConductorMidiData(120.0f, CurrentTimeSigNum, CurrentTimeSigDenom);
-			MidiClockOutPin->AttachToMidiResource(MidiDataOut, ResetToStart);
+			MidiClockOutPin->AttachToMidiFile(MidiDataOut, ResetToStart);
 			
-			int32 LoopEndTick = MidiDataOut->SongMaps.GetBarMap().BarIncludingCountInToTick(FMath::Max(LoopLengthBarsInPin, 1));
-			MidiClockOutPin->SetLoop(0, LoopEndTick);
+			int32 LoopEndTick = MidiDataOut->SongMaps.BarIncludingCountInToTick(FMath::Max(LoopLengthBarsInPin, 1));
+			// LoopEndTick == Loop Length since we are starting the loop at tick 0...
+			MidiClockOutPin->SetupPersistentLoop(0, LoopEndTick);
 		}
 		else
 		{
 			// if we're not looping, then the MidiClockOut is going to be the driving midi clock
 			// so attach the midi data directly to it
-			MidiClockOutPin->AttachToMidiResource(MidiData);
-			MidiClockOutPin->ClearLoop();
+			MidiClockOutPin->AttachToMidiFile(MidiData);
+			MidiClockOutPin->ClearPersistentLoop();
 		}
 	}
 
@@ -373,10 +367,6 @@ namespace HarmonixMetasound
 		bool HasMidiChanges = false;
 		if (*TempoInPin > 0 && !FMath::IsNearlyEqual(CurrentTempo, *TempoInPin))
 		{
-			if (!HasMidiChanges)
-			{
-				DrivingMidiClock.LockForMidiDataChanges();
-			}
 			AddTempoChangeForMidi(*TempoInPin);
 			HasMidiChanges = true;
 		};
@@ -387,12 +377,8 @@ namespace HarmonixMetasound
 		{
 			CurrentTimeSigNum = InTimeSigNum;
 			CurrentTimeSigDenom = InTimeSigDenom;
-			if (ensure(!MidiClockOutPin->DoesLoop()))
+			if (ensure(!MidiClockOutPin->HasPersistentLoop()))
 			{
-				if (!HasMidiChanges)
-				{
-					DrivingMidiClock.LockForMidiDataChanges();
-				}
 				AddTimeSigChangeForMidi(InTimeSigNum, InTimeSigDenom);
 				HasMidiChanges = true;
 			}
@@ -405,14 +391,14 @@ namespace HarmonixMetasound
 
 		if (HasMidiChanges)
 		{
-			DrivingMidiClock.MidiDataChangesComplete();
+			DrivingMidiClock.MidiChanged();
 		}
 	}
 
 	void FMetronomeOperator::AddTempoChangeForMidi(float InTempoBPM)
 	{
 		CurrentTempo = InTempoBPM;
-		int32 AtTick = GetDrivingMidiClock().GetCurrentHiResTick() + 1;
+		int32 AtTick = GetDrivingMidiClock().GetNextMidiTickToProcess();
 		MidiData->AddTempoChange(0, AtTick, CurrentTempo);
 	}
 
@@ -420,19 +406,20 @@ namespace HarmonixMetasound
 	{
 		CurrentTimeSigNum = InTimeSigNum;
 		CurrentTimeSigDenom = InTimeSigDenom;
-		int32 AtTick = GetDrivingMidiClock().GetCurrentHiResTick() + 1;
+		int32 AtTick = GetDrivingMidiClock().GetNextMidiTickToProcess();
 		// round to the next bar boundary, the bar we're actually going to apply the time sig change to
 		int32 AtBar = FMath::CeilToInt32(MidiData->SongMaps.GetBarIncludingCountInAtTick(AtTick));
-		int32 NumTimeSigPoints = MidiData->SongMaps.GetBarMap().GetNumTimeSignaturePoints();
+		int32 NumTimeSigPoints = MidiData->SongMaps.GetNumTimeSignatureChanges();
 
 		// check if there's already a time signature point at the bar we're trying to update
 		// the metronome clock increases monotonically, so we just have to check the _last_ point
-		FTimeSignaturePoint& Point = MidiData->SongMaps.GetBarMap().GetTimeSignaturePoint(NumTimeSigPoints - 1);
-		if (Point.BarIndex == AtBar)
+		FTimeSignaturePoint* Point = MidiData->SongMaps.GetMutableTimeSignaturePoint(NumTimeSigPoints - 1);
+		check(Point);
+		if (Point->BarIndex == AtBar)
 		{
 			// update that instead of adding a new one
-			Point.TimeSignature.Numerator = InTimeSigNum;
-			Point.TimeSignature.Denominator = InTimeSigDenom;
+			Point->TimeSignature.Numerator = InTimeSigNum;
+			Point->TimeSignature.Denominator = InTimeSigDenom;
 		}
 		else
 		{
@@ -442,12 +429,32 @@ namespace HarmonixMetasound
 
 	void FMetronomeOperator::HandleTransportChange(int32 StartFrameIndex, EMusicPlayerTransportState NewTransportState)
 	{
-		MidiClockOutPin->HandleTransportChange(StartFrameIndex, NewTransportState);
+		GetDrivingMidiClock().SetTransportState(StartFrameIndex, NewTransportState);
+	}
+
+	void FMetronomeOperator::PrepareClocksForExecute()
+	{
+		if (bClocksArePreparedForExecute)
+		{
+			return;
+		}
+
+		MidiClockOutPin->PrepareBlock();
+
 		if (LoopInPin)
 		{
-			MetronomeClock.HandleTransportChange(StartFrameIndex, NewTransportState);
+			MonotonicallyIncreasingClock->PrepareBlock();
 		}
+
+		bClocksArePreparedForExecute = true;
+
 	}
+
+	void FMetronomeOperator::MarkClocksAsExecuted()
+	{
+		bClocksArePreparedForExecute = false;
+	}
+
 }
 
 #undef LOCTEXT_NAMESPACE // "HarmonixMetaSound"
diff --git a/Engine/Plugins/Runtime/Harmonix/Source/HarmonixMetasound/Private/HarmonixMetasound/Nodes/MidiClockOffsetNode.cpp b/Engine/Plugins/Runtime/Harmonix/Source/HarmonixMetasound/Private/HarmonixMetasound/Nodes/MidiClockOffsetNode.cpp
--- a/Engine/Plugins/Runtime/Harmonix/Source/HarmonixMetasound/Private/HarmonixMetasound/Nodes/MidiClockOffsetNode.cpp
+++ b/Engine/Plugins/Runtime/Harmonix/Source/HarmonixMetasound/Private/HarmonixMetasound/Nodes/MidiClockOffsetNode.cpp
@@ -110,7 +110,9 @@ namespace HarmonixMetasound::Nodes::MidiClockOffset
 			, BlockSize(InSettings.GetNumFramesPerBlock())
 		{
 			TSharedPtr<FMidiFileData> ConductorMidiData = FMidiClock::MakeClockConductorMidiData(120.0f, 4, 4);
-			MidiClockOut->AttachToMidiResource(ConductorMidiData, true, PrerollBars);
+			MidiClockOut->AttachToMidiFile(ConductorMidiData, true);
+			MidiClockOut->SetDrivingClock(MidiClockIn->AsShared().ToSharedPtr());
+			bClockOutNeedsPrepare = false;
 		}
 
 		virtual void BindInputs(FInputVertexInterfaceData& InVertexData) override
@@ -119,7 +121,7 @@ namespace HarmonixMetasound::Nodes::MidiClockOffset
 			InVertexData.BindReadVertex(METASOUND_GET_PARAM_NAME(Inputs::OffsetBars), OffsetBarsInPin);
 			InVertexData.BindReadVertex(METASOUND_GET_PARAM_NAME(Inputs::OffsetBeats), OffsetBeatsInPin);
 			InVertexData.BindReadVertex(METASOUND_GET_PARAM_NAME(Inputs::OffsetMs), OffsetMsInPin);
-			MidiClockOut->AttachToTimeAuthority(*MidiClockIn);
+			MidiClockOut->SetDrivingClock(MidiClockIn->AsShared().ToSharedPtr());
 		}
 		
 		virtual void BindOutputs(FOutputVertexInterfaceData& InVertexData) override
@@ -131,7 +133,9 @@ namespace HarmonixMetasound::Nodes::MidiClockOffset
 		{
 			BlockSize = Params.OperatorSettings.GetNumFramesPerBlock();
 			CurrentBlockSpanStart = 0;
-			MidiClockOut->ResetAndStart(0, true);
+			MidiClockOut->SeekTo(0,0);
+			MidiClockOut->SetTransportState(0, EMusicPlayerTransportState::Playing);
+			bClockOutNeedsPrepare = false;
 
 			PrevOffsetBars = 0;
 			PrevOffsetBeats = 0.0f;
@@ -144,66 +148,61 @@ namespace HarmonixMetasound::Nodes::MidiClockOffset
 			if (!FMath::IsNearlyZero(OffsetBeats) || OffsetBars != 0)
 			{
 				int32 BeatsPerBar = 0;
-				FMusicTimestamp OffsetTimestamp = MidiClockIn->GetBarMap().TickToMusicTimestamp(InTick, &BeatsPerBar);
+				FMusicTimestamp OffsetTimestamp = MidiClockIn->GetSongMapEvaluator().TickToMusicTimestamp(InTick, &BeatsPerBar);
 				if (BeatsPerBar != 0)
 				{
 					float Beats = ((OffsetTimestamp.Bar - 1) * BeatsPerBar) + (OffsetTimestamp.Beat - 1);
 					Beats += OffsetBars * BeatsPerBar + OffsetBeats;
 					OffsetTimestamp.Bar = 1 + FMath::Floor(Beats / BeatsPerBar);
 					OffsetTimestamp.Beat = 1 + FMath::Fmod(BeatsPerBar + FMath::Fmod(Beats, BeatsPerBar), BeatsPerBar);
-					OffsetTick = MidiClockIn->GetBarMap().MusicTimestampToTick(OffsetTimestamp);
+					OffsetTick = MidiClockIn->GetSongMapEvaluator().MusicTimestampToTick(OffsetTimestamp);
 				}
 			}
-			const float Ms = MidiClockIn->GetSongMaps().TickToMs(OffsetTick);
-			OffsetTick = MidiClockIn->GetSongMaps().MsToTick(OffsetMs + Ms);
+			if (!FMath::IsNearlyZero(OffsetMs))
+			{
+				const float Ms = MidiClockIn->GetSongMapEvaluator().TickToMs(OffsetTick);
+				OffsetTick = MidiClockIn->GetSongMapEvaluator().MsToTick(OffsetMs + Ms);
+			}
 
 			return OffsetTick;
 		}
 		
 		virtual void Execute()
 		{
+			using namespace MidiClockMessageTypes;
+
 			float OffsetMs = *OffsetMsInPin;
 			int32 OffsetBars = *OffsetBarsInPin;
 			float OffsetBeats = *OffsetBeatsInPin;
-			MidiClockOut->PrepareBlock();
-			MidiClockOut->CopySpeedAndTempoChanges(MidiClockIn.Get());
+
+			// We might not want to "prepare" the clock output because
+			// we may have just initialized it with a tempo map, etc.
+			// in which case it already has events in it for this block!
+			if (bClockOutNeedsPrepare)
+			{
+				MidiClockOut->PrepareBlock();
+			}
+			// next time we definitely want to prepare the block.
+			bClockOutNeedsPrepare = true;
+
 			const TArray<FMidiClockEvent>& MidiClockEvents = MidiClockIn->GetMidiClockEventsInBlock();
 			for (int32 EventIndex = 0; EventIndex < MidiClockEvents.Num(); ++EventIndex)
 			{
 				const FMidiClockEvent& Event = MidiClockEvents[EventIndex];
-
-				if (Event.Msg.IsType<MidiClockMessageTypes::FSeekTo>())
+				if (Event.Msg.IsType<FSeek>())
 				{
-					const int32 Tick = GetTickWithOffset(Event.Msg.Get<MidiClockMessageTypes::FSeekTo>().ToTick, OffsetBars, OffsetBeats, OffsetMs);
-					FMusicSeekTarget SeekTarget;
-					SeekTarget.Type = ESeekPointType::Millisecond;
-					SeekTarget.Ms = MidiClockOut->GetSongMaps().TickToMs(Tick);
-								
-					MidiClockOut->SeekTo(Event.BlockFrameIndex, SeekTarget, PrerollBars);
+					const int32 NewNextTick = GetTickWithOffset(Event.Msg.Get<FSeek>().NewNextTick, OffsetBars, OffsetBeats, OffsetMs);
+					MidiClockOut->SeekTo(Event.BlockFrameIndex, NewNextTick);
 				}
-				else if (Event.Msg.IsType<MidiClockMessageTypes::FReset>())
+				else if (const FLoop* AsLoop = Event.TryGet<FLoop>())
 				{
-					const int32 Tick = GetTickWithOffset(Event.Msg.Get<MidiClockMessageTypes::FReset>().ToTick, OffsetBars, OffsetBeats, OffsetMs);
-					FMusicSeekTarget SeekTarget;
-					SeekTarget.Type = ESeekPointType::Millisecond;
-					SeekTarget.Ms = MidiClockOut->GetSongMaps().TickToMs(Tick);
-								
-					MidiClockOut->SeekTo(Event.BlockFrameIndex, SeekTarget, PrerollBars);
+					const int32 FirstTickInLoop = GetTickWithOffset(AsLoop->FirstTickInLoop, OffsetBars, OffsetBeats, OffsetMs);
+					const int32 LastTickAfterLoop = GetTickWithOffset(AsLoop->FirstTickInLoop + AsLoop->LengthInTicks, OffsetBars, OffsetBeats, OffsetMs);
+					MidiClockOut->AdvanceToTick(Event.BlockFrameIndex, LastTickAfterLoop);
+					MidiClockOut->AddTransientLoop(Event.BlockFrameIndex, FirstTickInLoop, LastTickAfterLoop - FirstTickInLoop);
 				}
-				else if (Event.Msg.IsType<MidiClockMessageTypes::FSeekThru>())
+				else if (const FAdvance* AsAdvance = Event.TryGet<FAdvance>())
 				{
-					const int32 Tick = GetTickWithOffset(Event.Msg.Get<MidiClockMessageTypes::FSeekThru>().ThruTick, OffsetBars, OffsetBeats, OffsetMs);
-					FMusicSeekTarget SeekTarget;
-					SeekTarget.Type = ESeekPointType::Millisecond;
-					SeekTarget.Ms = MidiClockOut->GetSongMaps().TickToMs(Tick + 1);
-							
-					MidiClockOut->SeekTo(Event.BlockFrameIndex, SeekTarget, PrerollBars);
-				}
-				else if (Event.Msg.IsType<MidiClockMessageTypes::FAdvanceThru>())
-				{
-					const int32 ThruTick = Event.Msg.Get<MidiClockMessageTypes::FAdvanceThru>().ThruTick;
-					const int32 Tick = GetTickWithOffset(ThruTick, OffsetBars, OffsetBeats, OffsetMs);
-						
 					// if our offset changed while we were advancing, seek to the new offset first
 					if (!FMath::IsNearlyEqual(PrevOffsetMs, OffsetMs) || !FMath::IsNearlyEqual(PrevOffsetBeats, OffsetBeats) || PrevOffsetBars != OffsetBars)
 					{
@@ -211,18 +210,37 @@ namespace HarmonixMetasound::Nodes::MidiClockOffset
 						PrevOffsetBars = OffsetBars;
 						PrevOffsetBeats = OffsetBeats;
 
-						FMusicSeekTarget SeekTarget;
-						SeekTarget.Type = ESeekPointType::Millisecond;
-						SeekTarget.Ms = MidiClockOut->GetSongMaps().TickToMs(Tick);
-								
-						MidiClockOut->SeekTo(Event.BlockFrameIndex, SeekTarget, PrerollBars);
 					}
-
-					const float AdvanceToMs = MidiClockOut->GetSongMaps().TickToMs(Tick);
-					const float ClockInSpeed = MidiClockIn->GetSpeedAtBlockSampleFrame(EventIndex);
-					const float AdvanceRatio = MidiClockIn->GetSongMaps().GetTempoAtTick(ThruTick) / MidiClockOut->GetSongMaps().GetTempoAtTick(Tick);
-					MidiClockOut->InformOfCurrentAdvanceRate(ClockInSpeed * AdvanceRatio);
-					MidiClockOut->AdvanceHiResToMs(Event.BlockFrameIndex, AdvanceToMs, true);
+					const int32 FirstickToProcess = GetTickWithOffset(AsAdvance->FirstTickToProcess, OffsetBars, OffsetBeats, OffsetMs);
+					// ONLY seek if there is a discontinuity AND the transport changed. Otherwise we 
+					// just want to advance from where we sit to the appropriate destination.
+					if (!bAdvancedSinceTransportChange && FirstickToProcess != MidiClockOut->GetNextMidiTickToProcess())
+					{
+						MidiClockOut->SeekTo(Event.BlockFrameIndex, FirstickToProcess);
+					}
+					const int32 ProcessUpToTick = GetTickWithOffset(AsAdvance->FirstTickToProcess + AsAdvance->NumberOfTicksToProcess, OffsetBars, OffsetBeats, OffsetMs);
+					MidiClockOut->AdvanceToTick(Event.BlockFrameIndex, ProcessUpToTick);
+					bAdvancedSinceTransportChange = true;
+				}
+				else if (const MidiClockMessageTypes::FTempoChange* AsTempoChange = Event.TryGet<MidiClockMessageTypes::FTempoChange>())
+				{
+					MidiClockOut->SetTempo(Event.BlockFrameIndex, MidiClockOut->GetNextMidiTickToProcess(), AsTempoChange->Tempo);
+				}
+				else if (const MidiClockMessageTypes::FTimeSignatureChange* AsTimeSigChange = Event.TryGet<MidiClockMessageTypes::FTimeSignatureChange>())
+				{
+					MidiClockOut->SetTimeSignature(Event.BlockFrameIndex, MidiClockOut->GetNextMidiTickToProcess(), AsTimeSigChange->TimeSignature);
+				}
+				else if (const MidiClockMessageTypes::FSpeedChange* AsSpeedChange = Event.TryGet<MidiClockMessageTypes::FSpeedChange>())
+				{
+					MidiClockOut->SetSpeed(Event.BlockFrameIndex, AsSpeedChange->Speed);
+				}
+				else if (const FTransportChange* AsTransportChange = Event.TryGet<FTransportChange>())
+				{
+					CurrentTransportState = AsTransportChange->TransportState;
+					// We need to know that the transport changed when we get the next advance
+					// so that we know if we should seek on a discontinuity. Otherwise we just
+					// want to advance from where the clock last left off!
+					bAdvancedSinceTransportChange = false;
 				}
 			}
 		}
@@ -241,7 +259,9 @@ namespace HarmonixMetasound::Nodes::MidiClockOffset
 		//** DATA
 		FSampleCount BlockSize      = 0;
 		int32 CurrentBlockSpanStart = 0;
-
+		EMusicPlayerTransportState CurrentTransportState = EMusicPlayerTransportState::Prepared;
+		bool bAdvancedSinceTransportChange = false;
+		bool bClockOutNeedsPrepare = false;
 		float PrevOffsetMs = 0.0f;
 		int32 PrevOffsetBars = 0;
 		float PrevOffsetBeats = 0.0f;
diff --git a/Engine/Plugins/Runtime/Harmonix/Source/HarmonixMetasound/Private/HarmonixMetasound/Nodes/MidiClockStateNode.cpp b/Engine/Plugins/Runtime/Harmonix/Source/HarmonixMetasound/Private/HarmonixMetasound/Nodes/MidiClockStateNode.cpp
--- a/Engine/Plugins/Runtime/Harmonix/Source/HarmonixMetasound/Private/HarmonixMetasound/Nodes/MidiClockStateNode.cpp
+++ b/Engine/Plugins/Runtime/Harmonix/Source/HarmonixMetasound/Private/HarmonixMetasound/Nodes/MidiClockStateNode.cpp
@@ -212,7 +212,7 @@ namespace HarmonixMetasound
 			return;
 		}
 	
-		const int32 CurrentTick = MidiClockInPin->GetCurrentMidiTick();
+		const int32 NextTick = MidiClockInPin->GetNextMidiTickToProcess();
 
 		bool HasTempoSpeedOrTimeSigChange = false;
 		if (MidiClockInPin->HasTempoChangesInBlock())
@@ -226,24 +226,24 @@ namespace HarmonixMetasound
 			*CurrentSpeedOutPin = FMath::Max(kTinyTempoSpeed, MidiClockInPin->GetSpeedAtEndOfBlock());
 		}
 			
-		const FBarMap& BarMap = MidiClockInPin->GetBarMap();
-		if (!BarMap.IsEmpty())
+		const ISongMapEvaluator& Map = MidiClockInPin->GetSongMapEvaluator();
+		const FTimeSignature* TimeSignature = Map.GetTimeSignatureAtTick(NextTick);
+		if (ensure(TimeSignature))
 		{
-			const FTimeSignature& TimeSignature = BarMap.GetTimeSignatureAtTick(CurrentTick);
-			if (TimeSignature.Numerator != *CurrentTimeSigNumOutPin || TimeSignature.Denominator != *CurrentTimeSigDenomOutPin)
+			if (TimeSignature->Numerator != *CurrentTimeSigNumOutPin || TimeSignature->Denominator != *CurrentTimeSigDenomOutPin)
 			{
 				HasTempoSpeedOrTimeSigChange = true;
-				*CurrentTimeSigNumOutPin = TimeSignature.Numerator;
-				*CurrentTimeSigDenomOutPin = FMath::Max(1, TimeSignature.Denominator);
+				*CurrentTimeSigNumOutPin = TimeSignature->Numerator;
+				*CurrentTimeSigDenomOutPin = FMath::Max(1, TimeSignature->Denominator);
 			}
-		}
 
-		if (CurrentTick != LastMidiTick)
-		{
-			LastMidiTick = CurrentTick;
-			FMusicTimestamp CurrentTimestamp = BarMap.TickToMusicTimestamp(CurrentTick);
-			*CurrentBarOutPin = CurrentTimestamp.Bar;
-			*CurrentBeatOutPin = CurrentTimestamp.Beat;
+			if (NextTick != LastMidiTick)
+			{
+				LastMidiTick = NextTick;
+				FMusicTimestamp CurrentTimestamp = Map.TickToMusicTimestamp(NextTick);
+				*CurrentBarOutPin = CurrentTimestamp.Bar;
+				*CurrentBeatOutPin = CurrentTimestamp.Beat;
+			}
 		}
 
 		if (HasTempoSpeedOrTimeSigChange)
diff --git a/Engine/Plugins/Runtime/Harmonix/Source/HarmonixMetasound/Private/HarmonixMetasound/Nodes/MidiClockSubdivisionTriggerNode.cpp b/Engine/Plugins/Runtime/Harmonix/Source/HarmonixMetasound/Private/HarmonixMetasound/Nodes/MidiClockSubdivisionTriggerNode.cpp
--- a/Engine/Plugins/Runtime/Harmonix/Source/HarmonixMetasound/Private/HarmonixMetasound/Nodes/MidiClockSubdivisionTriggerNode.cpp
+++ b/Engine/Plugins/Runtime/Harmonix/Source/HarmonixMetasound/Private/HarmonixMetasound/Nodes/MidiClockSubdivisionTriggerNode.cpp
@@ -208,9 +208,9 @@ namespace HarmonixMetasound
 		GridOffsetUnits = *GridOffsetUnitsInPin;
 		GridOffsetMultiplier = *GridOffsetMultInPin;
 
-		const FSongMaps& SongMaps = MidiClockInPin->GetSongMaps();
-		GridOffsetTicks = SongMaps.SubdivisionToMidiTicks(GridOffsetUnits, 0) * GridOffsetMultiplier;
-		GridSizeTicks = SongMaps.SubdivisionToMidiTicks(GridSizeUnits, 0) * GridSizeMultiplier;
+		const ISongMapEvaluator& Map = MidiClockInPin->GetSongMapEvaluator();
+		GridOffsetTicks = Map.SubdivisionToMidiTicks(GridOffsetUnits, 0) * GridOffsetMultiplier;
+		GridSizeTicks = Map.SubdivisionToMidiTicks(GridSizeUnits, 0) * GridSizeMultiplier;
 	}
 
 	void FMidiClockSubdivisionTriggerOperator::Execute()
@@ -232,37 +232,32 @@ namespace HarmonixMetasound
 		{
 			GridOffsetUnits = *GridOffsetUnitsInPin;
 			GridOffsetMultiplier = *GridOffsetMultInPin;
-			const FSongMaps& SongMaps = MidiClockInPin->GetSongMaps();
-			GridOffsetTicks = SongMaps.SubdivisionToMidiTicks(GridOffsetUnits, 0) * GridOffsetMultiplier;
+			const ISongMapEvaluator& Map = MidiClockInPin->GetSongMapEvaluator();
+			GridOffsetTicks = Map.SubdivisionToMidiTicks(GridOffsetUnits, 0) * GridOffsetMultiplier;
 		}
 
 		if (*EnableInPin)
 		{
+			using namespace MidiClockMessageTypes;
+
 			for (const FMidiClockEvent& ClockEvent : MidiClockInPin->GetMidiClockEventsInBlock())
 			{
-				if (ClockEvent.Msg.IsType<MidiClockMessageTypes::FAdvanceThru>())
+				if (const FAdvance* AsAdvance = ClockEvent.TryGet<FAdvance>())
 				{
-					const MidiClockMessageTypes::FAdvanceThru& AdvanceThru = ClockEvent.Msg.Get<MidiClockMessageTypes::FAdvanceThru>();
-					
-					if (AdvanceThru.IsPreRoll)
-					{
-						continue;
-					}
-					
-					int32 TickPreceedingThisAdvance = AdvanceThru.FromTick - GridOffsetTicks;
-					int32 LastTickProcessed = AdvanceThru.ThruTick - GridOffsetTicks;
-					const FSongMaps& SongMaps = MidiClockInPin->GetSongMaps();
+					int32 TickPreceedingThisAdvance = (AsAdvance->FirstTickToProcess - 1) - GridOffsetTicks;
+					int32 LastTickProcessed = AsAdvance->LastTickToProcess() - GridOffsetTicks;
+					const ISongMapEvaluator& Map = MidiClockInPin->GetSongMapEvaluator();
 					
 					// calculate new GridSizeTicks and GridOffsetTicks at the tick we're evaluating based on any time signature changes
-					GridSizeTicks = SongMaps.SubdivisionToMidiTicks(GridSizeUnits, AdvanceThru.ThruTick) * GridSizeMultiplier;
+					GridSizeTicks = Map.SubdivisionToMidiTicks(GridSizeUnits, AsAdvance->LastTickToProcess()) * GridSizeMultiplier;
 					
 					// if this is one of the quantization types that is variable in size we have to do it a little different...
 					if (GridSizeUnits == EMidiClockSubdivisionQuantization::Bar ||
 						GridSizeUnits == EMidiClockSubdivisionQuantization::Beat)
 					{
 						
-						FMusicTimestamp Start = SongMaps.GetBarMap().TickToMusicTimestamp(TickPreceedingThisAdvance);
-						FMusicTimestamp End = SongMaps.GetBarMap().TickToMusicTimestamp(LastTickProcessed);
+						FMusicTimestamp Start = Map.TickToMusicTimestamp(TickPreceedingThisAdvance);
+						FMusicTimestamp End = Map.TickToMusicTimestamp(LastTickProcessed);
 						if (GridSizeUnits == EMidiClockSubdivisionQuantization::Bar && Start.Bar != End.Bar)
 						{
 							if (End.Bar % GridSizeMultiplier == 0)
diff --git a/Engine/Plugins/Runtime/Harmonix/Source/HarmonixMetasound/Private/HarmonixMetasound/Nodes/MidiClockTimestampTriggerNode.cpp b/Engine/Plugins/Runtime/Harmonix/Source/HarmonixMetasound/Private/HarmonixMetasound/Nodes/MidiClockTimestampTriggerNode.cpp
--- a/Engine/Plugins/Runtime/Harmonix/Source/HarmonixMetasound/Private/HarmonixMetasound/Nodes/MidiClockTimestampTriggerNode.cpp
+++ b/Engine/Plugins/Runtime/Harmonix/Source/HarmonixMetasound/Private/HarmonixMetasound/Nodes/MidiClockTimestampTriggerNode.cpp
@@ -201,8 +201,8 @@ namespace HarmonixMetasound
 
 	void FMidiClockTimestampTriggerOperator::CalculateTriggerTick()
 	{
-		const FSongMaps& SongMaps = MidiClockInPin->GetSongMaps();
-		TriggerTick = SongMaps.CalculateMidiTick(CurrentTimestamp, Quantize);
+		const ISongMapEvaluator& Map = MidiClockInPin->GetSongMapEvaluator();
+		TriggerTick = Map.CalculateMidiTick(CurrentTimestamp, Quantize);
 	}
 
 	void FMidiClockTimestampTriggerOperator::Execute()
@@ -225,20 +225,17 @@ namespace HarmonixMetasound
 
 				bool ShouldTrigger = false;
 
-				if (ClockEvent.Msg.IsType<FAdvanceThru>())
+				if (const FAdvance* AsAdvance = ClockEvent.TryGet<FAdvance>())
 				{
-					const FAdvanceThru AdvanceThru = ClockEvent.Msg.Get<FAdvanceThru>();
-					ShouldTrigger = !AdvanceThru.IsPreRoll && AdvanceThru.FromTick < TriggerTick && AdvanceThru.ThruTick >= TriggerTick;
+					ShouldTrigger = AsAdvance->FirstTickToProcess <= TriggerTick && TriggerTick <= AsAdvance->LastTickToProcess();
 				}
-				else if (ClockEvent.Msg.IsType<FSeekThru>())
+				else if (const FSeek* AsSeekTo = ClockEvent.TryGet<FSeek>())
 				{
 					if (*TriggerDuringSeekInPin)
 					{
-						const FSeekThru& SeekThru = ClockEvent.Msg.Get<FSeekThru>();
-						ShouldTrigger = SeekThru.FromTick < TriggerTick && SeekThru.ThruTick >= TriggerTick;
+						ShouldTrigger = AsSeekTo->LastTickProcessedBeforeSeek < TriggerTick && AsSeekTo->NewNextTick > TriggerTick;
 					}
 				}
-				// TODO: SeekTo?
 				
 				if (ShouldTrigger)
 				{
