Massive Midi Clock Refactor

diff --git a/Engine/Plugins/Runtime/Harmonix/Source/HarmonixMetasound/Private/HarmonixMetasound/Nodes/MidiNoteTriggerNode.cpp b/Engine/Plugins/Runtime/Harmonix/Source/HarmonixMetasound/Private/HarmonixMetasound/Nodes/MidiNoteTriggerNode.cpp
--- a/Engine/Plugins/Runtime/Harmonix/Source/HarmonixMetasound/Private/HarmonixMetasound/Nodes/MidiNoteTriggerNode.cpp
+++ b/Engine/Plugins/Runtime/Harmonix/Source/HarmonixMetasound/Private/HarmonixMetasound/Nodes/MidiNoteTriggerNode.cpp
@@ -181,6 +181,7 @@ namespace HarmonixMetasound::Nodes::MidiNoteTriggerNode
 					{
 						// Stop sounding note
 						TriggerNoteOff(Event.BlockSampleFrameIndex, SoundingNote);
+						NoteOffTriggerFrame = Event.BlockSampleFrameIndex;
 					}
 
 					// Play new note
@@ -188,7 +189,7 @@ namespace HarmonixMetasound::Nodes::MidiNoteTriggerNode
 					*VelOutPin     = Event.MidiMessage.GetStdData2();
 					*NoteNumOutPin = Event.MidiMessage.GetStdData1();
 					NoteOnOutPin->TriggerFrame(NoteOffTriggerFrame == Event.BlockSampleFrameIndex ? NoteOffTriggerFrame + 1 : Event.BlockSampleFrameIndex);
-					SoundingNote   = Event.MidiMessage.GetStdData1();;
+					SoundingNote   = Event.MidiMessage.GetStdData1();
 				}
 				else if (Event.MidiMessage.IsNoteOff())
 				{
@@ -196,6 +197,19 @@ namespace HarmonixMetasound::Nodes::MidiNoteTriggerNode
 					{
 						TriggerNoteOff(Event.BlockSampleFrameIndex, Event.MidiMessage.GetStdData1());
 						NoteOffTriggerFrame = Event.BlockSampleFrameIndex;
+						PlayingId = FMidiVoiceId::None();
+						SoundingNote = -1;
+					}
+				}
+				else if (Event.MidiMessage.IsAllNotesOff() || Event.MidiMessage.IsAllNotesKill())
+				{
+					if (SoundingNote > -1)
+					{
+						// Stop sounding note
+						TriggerNoteOff(Event.BlockSampleFrameIndex, SoundingNote);
+						NoteOffTriggerFrame = Event.BlockSampleFrameIndex;
+						PlayingId = FMidiVoiceId::None();
+						SoundingNote = -1;
 					}
 				}
 			}
diff --git a/Engine/Plugins/Runtime/Harmonix/Source/HarmonixMetasound/Private/HarmonixMetasound/Nodes/MidiPlayerNode.cpp b/Engine/Plugins/Runtime/Harmonix/Source/HarmonixMetasound/Private/HarmonixMetasound/Nodes/MidiPlayerNode.cpp
--- a/Engine/Plugins/Runtime/Harmonix/Source/HarmonixMetasound/Private/HarmonixMetasound/Nodes/MidiPlayerNode.cpp
+++ b/Engine/Plugins/Runtime/Harmonix/Source/HarmonixMetasound/Private/HarmonixMetasound/Nodes/MidiPlayerNode.cpp
@@ -15,8 +15,8 @@
 #include "HarmonixMetasound/DataTypes/MidiAsset.h"
 #include "HarmonixMetasound/DataTypes/MusicTransport.h"
 
-#include "HarmonixMidi/MidiPlayCursor.h"
 #include "HarmonixMidi/MidiVoiceId.h"
+#include "HarmonixMidi/MidiCursor.h"
 
 #define LOCTEXT_NAMESPACE "HarmonixMetaSound"
 
@@ -30,7 +30,7 @@ namespace HarmonixMetasound
 		METASOUND_PARAM(KillVoicesOnMidiChange, "Kill Voices On MIDI File Change", "If true, a \"Kill All Voices\" MIDI message will be sent when the MIDI file asset is changed. Otherwise an \"All Notes Off\" will be sent which allows to ADSR release phases.")
 	}
 
-	class FMidiPlayerOperator : public TExecutableOperator<FMidiPlayerOperator>, public FMidiPlayCursor, public FMidiVoiceGeneratorBase, public FMusicTransportControllable
+	class FMidiPlayerOperator : public TExecutableOperator<FMidiPlayerOperator>, public FMidiCursor::FReceiver, public FMidiVoiceGeneratorBase, public FMusicTransportControllable
 	{
 	public:
 		static const FNodeClassMetadata& GetNodeInfo();
@@ -45,7 +45,6 @@ namespace HarmonixMetasound
                             const int32 InPrerollBars,
 							const bool bInKillVoicesOnSeek,
 							const bool bInKillVoicesOnMidiChange);
-		~FMidiPlayerOperator();
 
 		virtual void BindInputs(FInputVertexInterfaceData& InVertexData) override;
 		virtual void BindOutputs(FOutputVertexInterfaceData& InVertexData) override;
@@ -54,18 +53,12 @@ namespace HarmonixMetasound
 		virtual void Execute();
 		bool IsPlaying() const;
 
-		//~ BEGIN FMidiPlayCursor Overrides
-		virtual void Reset(bool ForceNoBroadcast = false) override;
-		virtual void OnLoop(int32 LoopStartTick, int32 LoopEndTick) override;
-		virtual void SeekToTick(int32 Tick) override;
-		virtual void SeekThruTick(int32 Tick) override;
-		virtual void AdvanceThruTick(int32 Tick, bool IsPreRoll) override;
-
-		virtual void OnMidiMessage(int32 TrackIndex, int32 Tick, uint8 Status, uint8 Data1, uint8 Data2, bool IsPreroll) override;
-		virtual void OnTempo(int32 TrackIndex, int32 Tick, int32 Tempo, bool IsPreroll) override;
-		virtual void OnText(int32 TrackIndex, int32 Tick, int32 TextIndex, const FString& Str, uint8 Type, bool IsPreroll = false) override;
+		//~ BEGIN FMidiCursor::FReceiver Overrides
+		virtual void OnMidiMessage(int32 TrackIndex, int32 Tick, uint8 Status, uint8 Data1, uint8 Data2, bool bIsPreroll) override;
+		virtual void OnTempo(int32 TrackIndex, int32 Tick, int32 Tempo, bool bIsPreroll) override;
+		virtual void OnText(int32 TrackIndex, int32 Tick, int32 TextIndex, const FString& Str, uint8 Type, bool bIsPreroll) override;
 		virtual void OnPreRollNoteOn(int32 TrackIndex, int32 EventTick, int32 CurrentTick, float PreRollMs, uint8 Status, uint8 Data1, uint8 Data2) override;
-		//~ END FMidiPlayCursor Overrides
+		//~ END FMidiCursor::FReceiver Overrides
 
 	protected:
 		//** INPUTS
@@ -83,9 +76,11 @@ namespace HarmonixMetasound
 
 		//** DATA
 		FMidiFileProxyPtr CurrentMidiFile;
-		FSampleCount BlockSize      = 0;
+		FMidiCursor MidiCursor;
+		FSampleCount BlockSize = 0;
 		int32 CurrentBlockSpanStart = 0;
 		bool NeedsTransportInit = true;
+		int32 CurrentRenderBlockFrame = 0;
 		
 		virtual void SetupNewMidiFile(const FMidiFileProxyPtr& NewMidi);
 
@@ -98,6 +93,9 @@ namespace HarmonixMetasound
 			}
 		}
 		virtual void InitTransportImpl() = 0;
+
+		void RenderMidiForClockEvents();
+		void SendAllNotesOff(int32 BlockFrameIndex, int32 Tick);
 	};
 
 	class FExternallyClockedMidiPlayerOperator : public FMidiPlayerOperator
@@ -254,16 +252,9 @@ namespace HarmonixMetasound
 		, MidiClockOut(FMidiClockWriteRef::CreateNew(InSettings))
 		, BlockSize(InSettings.GetNumFramesPerBlock())
 	{
-		MidiClockOut->RegisterHiResPlayCursor(this);
 		MidiOutPin->SetClock(*MidiClockOut);
 	}
 
-	FMidiPlayerOperator::~FMidiPlayerOperator()
-	{
-		MidiClockOut->UnregisterPlayCursor(this);
-	}
-
-
 	FExternallyClockedMidiPlayerOperator::FExternallyClockedMidiPlayerOperator(const FOperatorSettings& InSettings,
 																			   const FMidiAssetReadRef& InMidiAsset, 
 																			   const FMusicTransportEventStreamReadRef& InTransport,
@@ -276,6 +267,7 @@ namespace HarmonixMetasound
 		: FMidiPlayerOperator(InSettings, InMidiAsset, InTransport, InLoop, InSpeedMultiplier, InPrerollBars, bInKillVoicesOnSeek, bInKillVoicesOnMidiChange)
 		, MidiClockIn(InMidiClock)
 	{
+		MidiClockOut->SetDrivingClock(MidiClockIn->AsShared().ToSharedPtr());
 	}
 
 	FSelfClockedMidiPlayerOperator::FSelfClockedMidiPlayerOperator(const FOperatorSettings& InSettings, 
@@ -299,8 +291,6 @@ namespace HarmonixMetasound
 		FMidiPlayerOperator::BindInputs(InVertexData);
 		using namespace CommonPinNames;
 		InVertexData.BindReadVertex(METASOUND_GET_PARAM_NAME(Inputs::MidiClock), MidiClockIn);
-
-		MidiClockOut->AttachToTimeAuthority(*MidiClockIn);
 	}
 
 	void FExternallyClockedMidiPlayerOperator::BindOutputs(FOutputVertexInterfaceData& InVertexData)
@@ -335,13 +325,17 @@ namespace HarmonixMetasound
 			case EMusicPlayerTransportState::Prepared:
 			case EMusicPlayerTransportState::Stopping:
 			case EMusicPlayerTransportState::Killing:
-				MidiClockOut->AddTransportStateChangeToBlock({ 0, 0.0f, EMusicPlayerTransportState::Prepared });
+				MidiClockOut->SetTransportState(0, EMusicPlayerTransportState::Prepared);
 				return EMusicPlayerTransportState::Prepared;
 
 			case EMusicPlayerTransportState::Starting:
 			case EMusicPlayerTransportState::Playing:
 			case EMusicPlayerTransportState::Continuing:
-				MidiClockOut->ResetAndStart(0, !ReceivedSeekWhileStopped());
+				MidiClockOut->SetTransportState(0, EMusicPlayerTransportState::Playing);
+				if (!ReceivedSeekWhileStopped())
+				{
+					MidiClockOut->SeekTo(0, 0);
+				}
 				return EMusicPlayerTransportState::Playing;
 
 			case EMusicPlayerTransportState::Seeking: // seeking is omitted from init, shouldn't happen
@@ -350,7 +344,7 @@ namespace HarmonixMetasound
 
 			case EMusicPlayerTransportState::Pausing:
 			case EMusicPlayerTransportState::Paused:
-				MidiClockOut->AddTransportStateChangeToBlock({ 0, 0.0f, EMusicPlayerTransportState::Paused });
+				MidiClockOut->SetTransportState(0, EMusicPlayerTransportState::Paused);
 				return EMusicPlayerTransportState::Paused;
 
 			default:
@@ -364,10 +358,11 @@ namespace HarmonixMetasound
 	void FMidiPlayerOperator::Execute()
 	{
 		MidiOutPin->PrepareBlock();
-		
+		CurrentRenderBlockFrame = 0;
+
 		MidiClockOut->PrepareBlock();
 
-		if (MidiClockOut->DoesLoop() != *LoopInPin || CurrentMidiFile != MidiAssetInPin->GetMidiProxy())
+		if (MidiClockOut->HasPersistentLoop() != *LoopInPin || CurrentMidiFile != MidiAssetInPin->GetMidiProxy())
 		{
 			SetupNewMidiFile(MidiAssetInPin->GetMidiProxy());
 		}
@@ -386,26 +381,34 @@ namespace HarmonixMetasound
 	{
 		FMidiPlayerOperator::Execute();
 
-		MidiClockOut->CopySpeedAndTempoChanges(MidiClockIn.Get(), *SpeedMultInPin);
-
 		TransportSpanPostProcessor HandleMidiClockEventsInBlock = [this](int32 StartFrameIndex, int32 EndFrameIndex, EMusicPlayerTransportState CurrentState)
 		{
 			// clock should always process in post processor
 			int32 NumFrames = EndFrameIndex - StartFrameIndex;
-			MidiClockOut->HandleTransportChange(StartFrameIndex, CurrentState);
-			MidiClockOut->Process(*MidiClockIn, StartFrameIndex, NumFrames, PrerollBars, *SpeedMultInPin);
+			MidiClockOut->SetTransportState(StartFrameIndex, CurrentState);
+			if (CurrentState == EMusicPlayerTransportState::Playing || CurrentState == EMusicPlayerTransportState::Continuing)
+			{
+				MidiClockOut->SetSpeed(StartFrameIndex, *SpeedMultInPin);
+				MidiClockOut->Advance(*MidiClockIn, StartFrameIndex, NumFrames);
+			}
 		};
 
 		TransportSpanProcessor TransportHandler = [&](int32 StartFrameIndex, int32 EndFrameIndex, EMusicPlayerTransportState CurrentState)
 		{
 			switch (CurrentState)
 			{
+			case EMusicPlayerTransportState::Starting:
+				if (!ReceivedSeekWhileStopped())
+				{
+					MidiClockOut->SeekTo(StartFrameIndex, 0);
+				}
+				break;
 			case EMusicPlayerTransportState::Stopping:
 				{
 					FMidiStreamEvent MidiEvent(this, FMidiMsg::CreateAllNotesOff());
-					MidiEvent.BlockSampleFrameIndex = MidiClockOut->GetCurrentBlockFrameIndex();
-					MidiEvent.AuthoredMidiTick = MidiClockOut->GetCurrentHiResTick();
-					MidiEvent.CurrentMidiTick = MidiClockOut->GetCurrentHiResTick();
+					MidiEvent.BlockSampleFrameIndex = StartFrameIndex;
+					MidiEvent.AuthoredMidiTick = MidiCursor.GetNextTick();
+					MidiEvent.CurrentMidiTick = MidiCursor.GetNextTick();
 					MidiEvent.TrackIndex = 0;
 					MidiOutPin->AddMidiEvent(MidiEvent);
 				}
@@ -414,6 +417,7 @@ namespace HarmonixMetasound
 			return GetNextTransportState(CurrentState);
 		};
 		ExecuteTransportSpans(TransportInPin, BlockSize, TransportHandler, HandleMidiClockEventsInBlock);
+		RenderMidiForClockEvents();
 	}
 
 	void FExternallyClockedMidiPlayerOperator::InitTransportImpl()
@@ -438,13 +442,18 @@ namespace HarmonixMetasound
 	{
 		FMidiPlayerOperator::Execute();
 
-		MidiClockOut->AddSpeedChangeToBlock({ 0, 0.0f, *SpeedMultInPin });
+		MidiClockOut->SetSpeed(0, *SpeedMultInPin);
 
 		TransportSpanPostProcessor MidiClockTransportHandler = [&](int32 StartFrameIndex, int32 EndFrameIndex, EMusicPlayerTransportState CurrentState)
 		{
 			int32 NumFrames = EndFrameIndex - StartFrameIndex;
-			MidiClockOut->HandleTransportChange(StartFrameIndex, CurrentState);
-			MidiClockOut->Process(StartFrameIndex, NumFrames, PrerollBars, *SpeedMultInPin);
+			MidiClockOut->SetTransportState(StartFrameIndex, CurrentState);
+			switch (CurrentState)
+			{
+			case EMusicPlayerTransportState::Playing:
+			case EMusicPlayerTransportState::Continuing:
+				MidiClockOut->Advance(StartFrameIndex, NumFrames);
+			}
 		};
 
 		TransportSpanProcessor TransportHandler = [this](int32 StartFrameIndex, int32 EndFrameIndex, EMusicPlayerTransportState CurrentState)
@@ -452,15 +461,19 @@ namespace HarmonixMetasound
 			switch (CurrentState)
 			{
 			case EMusicPlayerTransportState::Starting:
-				MidiClockOut->ResetAndStart(StartFrameIndex, !ReceivedSeekWhileStopped());
+				if (!ReceivedSeekWhileStopped())
+				{
+					MidiClockOut->SeekTo(StartFrameIndex, 0);
+				}
 				break;
 			case EMusicPlayerTransportState::Seeking:
-				MidiClockOut->SeekTo(StartFrameIndex, TransportInPin->GetNextSeekDestination(), PrerollBars);
+				MidiClockOut->SeekTo(StartFrameIndex, TransportInPin->GetNextSeekDestination());
 				break;
 			}
 			return GetNextTransportState(CurrentState);
 		};
 		ExecuteTransportSpans(TransportInPin, BlockSize, TransportHandler, MidiClockTransportHandler);
+		RenderMidiForClockEvents();
 	}
 
 	void FMidiPlayerOperator::BindInputs(FInputVertexInterfaceData& InVertexData)
@@ -494,65 +507,18 @@ namespace HarmonixMetasound
 		CurrentBlockSpanStart = 0;
 
 		MidiOutPin->SetClock(*MidiClockOut);
-		MidiClockOut->ResetAndStart(0, true);
+		MidiClockOut->SeekTo(0, 0);
+		MidiClockOut->SetTransportState(0, EMusicPlayerTransportState::Prepared);
 
 		NeedsTransportInit = true;
 	}
 
-	void FMidiPlayerOperator::Reset(bool ForceNoBroadcast /*= false*/)
-	{
-		FMidiPlayCursor::Reset(ForceNoBroadcast);
-	}
-
-	void FMidiPlayerOperator::OnLoop(int32 LoopStartTick, int32 LoopEndTick)
-	{
-		//TRACE_BOOKMARK(TEXT("Midi Looping"));
-		FMidiPlayCursor::OnLoop(LoopStartTick, LoopEndTick);
-	}
-
-	void FMidiPlayerOperator::SeekToTick(int32 Tick)
-	{
-		//TRACE_BOOKMARK(TEXT("Midi Seek To Tick"));
-		FMidiPlayCursor::SeekToTick(Tick);
-
-		if (IsPlaying())
-		{
-			FMidiStreamEvent MidiEvent(this, bKillVoicesOnSeek ? FMidiMsg::CreateAllNotesKill() : FMidiMsg::CreateAllNotesOff());
-			MidiEvent.BlockSampleFrameIndex = MidiClockOut->GetCurrentBlockFrameIndex();
-			MidiEvent.AuthoredMidiTick = MidiClockOut->GetCurrentHiResTick();
-			MidiEvent.CurrentMidiTick = MidiClockOut->GetCurrentHiResTick();
-			MidiEvent.TrackIndex = 0;
-			MidiOutPin->AddMidiEvent(MidiEvent);
-		}
-	}
-
-	void FMidiPlayerOperator::SeekThruTick(int32 Tick)
-	{
-		//TRACE_BOOKMARK(TEXT("Midi Seek Thru Tick"));
-		FMidiPlayCursor::SeekThruTick(Tick);
-
-		if (IsPlaying())
-		{
-			FMidiStreamEvent MidiEvent(this, bKillVoicesOnSeek ? FMidiMsg::CreateAllNotesKill() : FMidiMsg::CreateAllNotesOff());
-			MidiEvent.BlockSampleFrameIndex = MidiClockOut->GetCurrentBlockFrameIndex();
-			MidiEvent.AuthoredMidiTick = MidiClockOut->GetCurrentHiResTick();
-			MidiEvent.CurrentMidiTick = MidiClockOut->GetCurrentHiResTick();
-			MidiEvent.TrackIndex = 0;
-			MidiOutPin->AddMidiEvent(MidiEvent);
-		}
-	}
-
-	void FMidiPlayerOperator::AdvanceThruTick(int32 Tick, bool IsPreRoll)
-	{
-		FMidiPlayCursor::AdvanceThruTick(Tick, IsPreRoll);
-	}
-
-	void FMidiPlayerOperator::OnMidiMessage(int32 TrackIndex, int32 Tick, uint8 Status, uint8 Data1, uint8 Data2, bool IsPreroll)
+	void FMidiPlayerOperator::OnMidiMessage(int32 TrackIndex, int32 Tick, uint8 Status, uint8 Data1, uint8 Data2, bool bIsPreroll)
 	{
 		if (IsPlaying())
 		{
 			FMidiStreamEvent MidiEvent(this, FMidiMsg(Status, Data1, Data2));
-			MidiEvent.BlockSampleFrameIndex = MidiClockOut->GetCurrentBlockFrameIndex();
+			MidiEvent.BlockSampleFrameIndex = CurrentRenderBlockFrame;
 			MidiEvent.AuthoredMidiTick = Tick;
 			MidiEvent.CurrentMidiTick = Tick;
 			MidiEvent.TrackIndex = TrackIndex;
@@ -560,12 +526,12 @@ namespace HarmonixMetasound
 		}
 	}
 
-	void FMidiPlayerOperator::OnTempo(int32 TrackIndex, int32 Tick, int32 tempo, bool isPreroll)
+	void FMidiPlayerOperator::OnTempo(int32 TrackIndex, int32 Tick, int32 tempo, bool bIsPreroll)
 	{
 		if (IsPlaying())
 		{
 			FMidiStreamEvent MidiEvent(this, FMidiMsg((int32)tempo));
-			MidiEvent.BlockSampleFrameIndex  = MidiClockOut->GetCurrentBlockFrameIndex();
+			MidiEvent.BlockSampleFrameIndex  = CurrentRenderBlockFrame;
 			MidiEvent.AuthoredMidiTick = Tick;
 			MidiEvent.CurrentMidiTick  = Tick;
 			MidiEvent.TrackIndex       = TrackIndex;
@@ -573,12 +539,12 @@ namespace HarmonixMetasound
 		}
 	}
 
-	void FMidiPlayerOperator::OnText(int32 TrackIndex, int32 Tick, int32 TextIndex, const FString& Str, uint8 Type, bool IsPreroll)
+	void FMidiPlayerOperator::OnText(int32 TrackIndex, int32 Tick, int32 TextIndex, const FString& Str, uint8 Type, bool bIsPreroll)
 	{
-		if (!IsPreroll && IsPlaying())
+		if (!bIsPreroll && IsPlaying())
 		{
 			FMidiStreamEvent MidiEvent(this, FMidiMsg::CreateText(TextIndex, Type));
-			MidiEvent.BlockSampleFrameIndex = MidiClockOut->GetCurrentBlockFrameIndex();
+			MidiEvent.BlockSampleFrameIndex = CurrentRenderBlockFrame;
 			MidiEvent.AuthoredMidiTick = Tick;
 			MidiEvent.CurrentMidiTick = Tick;
 			MidiEvent.TrackIndex = TrackIndex;
@@ -586,13 +552,12 @@ namespace HarmonixMetasound
 		}
 	}
 
-
 	void FMidiPlayerOperator::OnPreRollNoteOn(int32 TrackIndex, int32 EventTick, int32 InCurrentTick, float InPrerollMs, uint8 InStatus, uint8 Data1, uint8 Data2)
 	{
 		if (IsPlaying())
 		{
 			FMidiStreamEvent MidiEvent(this, FMidiMsg(InStatus, Data1, Data2));
-			MidiEvent.BlockSampleFrameIndex = MidiClockOut->GetCurrentBlockFrameIndex();
+			MidiEvent.BlockSampleFrameIndex = CurrentRenderBlockFrame;
 			MidiEvent.AuthoredMidiTick = EventTick;
 			MidiEvent.CurrentMidiTick = InCurrentTick;
 			MidiEvent.TrackIndex = TrackIndex;
@@ -614,7 +579,8 @@ namespace HarmonixMetasound
 
 		if (CurrentMidiFile.IsValid())
 		{			
-			MidiClockOut->AttachToMidiResource(CurrentMidiFile->GetMidiFile(), !IsPlaying(), PrerollBars);
+			MidiClockOut->AttachToMidiFile(CurrentMidiFile->GetMidiFile(), !IsPlaying());
+			MidiCursor.Prepare(CurrentMidiFile->GetMidiFile());
 			MidiOutPin->SetTicksPerQuarterNote(CurrentMidiFile->GetMidiFile()->TicksPerQuarterNote);
 			if (*LoopInPin)
 			{
@@ -623,35 +589,72 @@ namespace HarmonixMetasound
 				int32 LoopEndTick = SongLengthData.LengthTicks;
 
 				// Round the content authored ticks to bar boundaries based on the bar map. 
-				const FBarMap& BarMap = CurrentMidiFile->GetMidiFile()->SongMaps.GetBarMap();
-				int32 LoopStartBarIndex = FMath::RoundToInt32(BarMap.TickToFractionalBarIncludingCountIn(LoopStartTick));
-				int32 LoopEndBarIndex = FMath::RoundToInt32(BarMap.TickToFractionalBarIncludingCountIn(LoopEndTick));
+				const ISongMapEvaluator& SongMap = CurrentMidiFile->GetMidiFile()->SongMaps;
+				int32 LoopStartBarIndex = FMath::RoundToInt32(SongMap.TickToFractionalBarIncludingCountIn(LoopStartTick));
+				int32 LoopEndBarIndex = FMath::RoundToInt32(SongMap.TickToFractionalBarIncludingCountIn(LoopEndTick));
 
 				// Make sure there's at least 1 bar of looping
 				if (LoopEndBarIndex <= LoopStartBarIndex)
 				{
 					LoopEndBarIndex = LoopStartBarIndex + 1;
 				}
-				LoopStartTick = BarMap.BarIncludingCountInToTick(LoopStartBarIndex);
-				LoopEndTick = BarMap.BarIncludingCountInToTick(LoopEndBarIndex);
+				LoopStartTick = SongMap.BarIncludingCountInToTick(LoopStartBarIndex);
+				LoopEndTick = SongMap.BarIncludingCountInToTick(LoopEndBarIndex);
 				
-				MidiClockOut->SetLoop(LoopStartTick, LoopEndTick);
+				MidiClockOut->SetupPersistentLoop(LoopStartTick, LoopEndTick - LoopStartTick);
 				
 				// remap our current tick based on the looping behavior
 				// maybe this should happen automatically when resetting a loop or changing midi files?
-				int32 NewTick = MidiClockOut->CalculateMappedTick(CurrentTick);
-				MidiClockOut->SeekTo(NewTick, PrerollBars);
+				int32 NewTick = MidiClockOut->WrapTickIfLooping(MidiClockOut->GetNextMidiTickToProcess());
+				MidiClockOut->SeekTo(0, NewTick);
 			}
 			else
 			{
-				MidiClockOut->ClearLoop();
+				MidiClockOut->ClearPersistentLoop();
 			}
+			MidiCursor.SeekToNextTick(MidiClockOut->GetNextMidiTickToProcess(), PrerollBars, this);
 		}
 		else
 		{
-			MidiClockOut->AttachToMidiResource(nullptr, !IsPlaying(), 0);
+			MidiClockOut->AttachToMidiFile(nullptr, !IsPlaying());
+			MidiCursor.Prepare(nullptr);
+		}
+	}
+
+	void FMidiPlayerOperator::RenderMidiForClockEvents()
+	{
+		const TArray<FMidiClockEvent>& ClockEvents = MidiClockOut->GetMidiClockEventsInBlock();
+		for (const FMidiClockEvent& Event : ClockEvents)
+		{
+			CurrentRenderBlockFrame = Event.BlockFrameIndex;
+			if (auto AsAdvance = Event.TryGet<MidiClockMessageTypes::FAdvance>())
+			{
+				MidiCursor.Process(AsAdvance->FirstTickToProcess, AsAdvance->LastTickToProcess(), *this);
+			}
+			else if (auto AsSeekTo = Event.TryGet<MidiClockMessageTypes::FSeek>())
+			{
+				SendAllNotesOff(Event.BlockFrameIndex, AsSeekTo->NewNextTick);
+				MidiCursor.SeekToNextTick(AsSeekTo->NewNextTick, PrerollBars, this);
+			}
+			else if (auto AsLoop = Event.TryGet<MidiClockMessageTypes::FLoop>())
+			{
+				// When looping we don't preroll the events prior to the 
+				// loop start point.
+				SendAllNotesOff(Event.BlockFrameIndex, AsLoop->FirstTickInLoop);
+				MidiCursor.SeekToNextTick(AsLoop->FirstTickInLoop);
+			}
 		}
 	}
+
+	void FMidiPlayerOperator::SendAllNotesOff(int32 BlockFrameIndex, int32 Tick)
+	{
+		FMidiStreamEvent MidiEvent(this, bKillVoicesOnSeek ? FMidiMsg::CreateAllNotesKill() : FMidiMsg::CreateAllNotesOff());
+		MidiEvent.BlockSampleFrameIndex = BlockFrameIndex;
+		MidiEvent.AuthoredMidiTick = Tick;
+		MidiEvent.CurrentMidiTick = Tick;
+		MidiEvent.TrackIndex = 0;
+		MidiOutPin->AddMidiEvent(MidiEvent);
+	}
 }
 
 #undef LOCTEXT_NAMESPACE // "HarmonixMetaSound"
diff --git a/Engine/Plugins/Runtime/Harmonix/Source/HarmonixMetasound/Private/HarmonixMetasound/Nodes/MorphingLfoNode.cpp b/Engine/Plugins/Runtime/Harmonix/Source/HarmonixMetasound/Private/HarmonixMetasound/Nodes/MorphingLfoNode.cpp
--- a/Engine/Plugins/Runtime/Harmonix/Source/HarmonixMetasound/Private/HarmonixMetasound/Nodes/MorphingLfoNode.cpp
+++ b/Engine/Plugins/Runtime/Harmonix/Source/HarmonixMetasound/Private/HarmonixMetasound/Nodes/MorphingLfoNode.cpp
@@ -254,7 +254,15 @@ namespace HarmonixMetasound::Nodes::MorphingLFO
 			MusicTimingInfo.Tempo = LastTempo;
 			MusicTimingInfo.Speed = LastSpeed;
 			MusicTimingInfo.Timestamp = (*Inputs.MidiClock)->GetMusicTimestampAtBlockOffset(StartSample);
-			MusicTimingInfo.TimeSignature = (*Inputs.MidiClock)->GetBarMap().GetTimeSignatureAtBar(MusicTimingInfo.Timestamp.Bar);
+			const FTimeSignature* TimeSigPtr = (*Inputs.MidiClock)->GetSongMapEvaluator().GetTimeSignatureAtBar(MusicTimingInfo.Timestamp.Bar);
+			if (!TimeSigPtr)
+			{
+				MusicTimingInfo.TimeSignature = FTimeSignature(4, 4);
+			}
+			else
+			{
+				MusicTimingInfo.TimeSignature = *TimeSigPtr;
+			}
 			LFO.Advance(NumSamples, *Outputs.LFO.Get(), &MusicTimingInfo);
 		}
 		else
@@ -288,7 +296,15 @@ namespace HarmonixMetasound::Nodes::MorphingLFO
 			MusicTimingInfo.Tempo = LastTempo;
 			MusicTimingInfo.Speed = LastSpeed;
 			MusicTimingInfo.Timestamp = (*Inputs.MidiClock)->GetMusicTimestampAtBlockOffset(StartSample);
-			MusicTimingInfo.TimeSignature = (*Inputs.MidiClock)->GetBarMap().GetTimeSignatureAtBar(MusicTimingInfo.Timestamp.Bar);
+			const FTimeSignature* TimeSigPtr = (*Inputs.MidiClock)->GetSongMapEvaluator().GetTimeSignatureAtBar(MusicTimingInfo.Timestamp.Bar);
+			if (!TimeSigPtr)
+			{
+				MusicTimingInfo.TimeSignature = FTimeSignature(4, 4);
+			}
+			else
+			{
+				MusicTimingInfo.TimeSignature = *TimeSigPtr;
+			}
 			LFO.Advance(Output, NumSamples, &MusicTimingInfo);
 		}
 		else
diff --git a/Engine/Plugins/Runtime/Harmonix/Source/HarmonixMetasound/Private/HarmonixMetasound/Nodes/StepSequencePlayerNode.cpp b/Engine/Plugins/Runtime/Harmonix/Source/HarmonixMetasound/Private/HarmonixMetasound/Nodes/StepSequencePlayerNode.cpp
--- a/Engine/Plugins/Runtime/Harmonix/Source/HarmonixMetasound/Private/HarmonixMetasound/Nodes/StepSequencePlayerNode.cpp
+++ b/Engine/Plugins/Runtime/Harmonix/Source/HarmonixMetasound/Private/HarmonixMetasound/Nodes/StepSequencePlayerNode.cpp
@@ -129,7 +129,7 @@ namespace HarmonixMetasound::Nodes::StepSequencePlayer
 
 		void SeekToTick(int32 BlockFrameIndex, int32 Tick);
 		void SeekThruTick(int32 BlockFrameIndex, int32 Tick);
-		void AdvanceThruTick(int32 BlockFrameIndex, int32 Tick, bool IsPreRoll);
+		void AdvanceThruTick(int32 BlockFrameIndex, int32 Tick);
 
 	};
 
@@ -344,7 +344,7 @@ namespace HarmonixMetasound::Nodes::StepSequencePlayer
 		{
 			if (CurrentCellNotes.Num() > 0)
 			{
-				AllNotesOff(0, MidiClockInPin->GetCurrentMidiTick(), true);
+				AllNotesOff(0, MidiClockInPin->GetLastProcessedMidiTick(), true);
 				CurrentCellNotes.Empty();
 			}
 			return;
@@ -358,6 +358,8 @@ namespace HarmonixMetasound::Nodes::StepSequencePlayer
 
 		TransportSpanPostProcessor ClockEventHandler = [&](int32 StartFrameIndex, int32 EndFrameIndex, EMusicPlayerTransportState State)
 			{
+				using namespace MidiClockMessageTypes;
+
 				switch (State)
 				{
 				case EMusicPlayerTransportState::Playing:
@@ -371,18 +373,19 @@ namespace HarmonixMetasound::Nodes::StepSequencePlayer
 
 						if (Event.BlockFrameIndex >= StartFrameIndex)
 						{
-							if (Event.Msg.IsType<MidiClockMessageTypes::FAdvanceThru>())
+							if (const FAdvance* AsAdvance = Event.TryGet<FAdvance>())
 							{
-								const MidiClockMessageTypes::FAdvanceThru& AdvanceThru = Event.Msg.Get<MidiClockMessageTypes::FAdvanceThru>();
-								AdvanceThruTick(Event.BlockFrameIndex, AdvanceThru.ThruTick, AdvanceThru.IsPreRoll);
+								AdvanceThruTick(Event.BlockFrameIndex, AsAdvance->LastTickToProcess());
 							}
+/*
 							else if (Event.Msg.IsType<MidiClockMessageTypes::FSeekThru>())
 							{
 								SeekThruTick(Event.BlockFrameIndex, Event.Msg.Get<MidiClockMessageTypes::FSeekThru>().ThruTick);
 							}
-							else if (Event.Msg.IsType<MidiClockMessageTypes::FSeekTo>())
+*/
+							else if (const FSeek* AsSeekTo = Event.TryGet<FSeek>())
 							{
-								SeekToTick(Event.BlockFrameIndex, Event.Msg.Get<MidiClockMessageTypes::FSeekTo>().ToTick);
+								SeekToTick(Event.BlockFrameIndex, AsSeekTo->NewNextTick);
 							}
 						}
 					}
@@ -420,11 +423,11 @@ namespace HarmonixMetasound::Nodes::StepSequencePlayer
 					return EMusicPlayerTransportState::Paused;
 
 				case EMusicPlayerTransportState::Stopping:
-					AllNotesOff(StartFrameIndex, MidiClockInPin->GetCurrentMidiTick(), true);
+					AllNotesOff(StartFrameIndex, MidiClockInPin->GetLastProcessedMidiTick(), true);
 					return EMusicPlayerTransportState::Prepared;
 
 				case EMusicPlayerTransportState::Killing:
-					AllNotesOff(StartFrameIndex, MidiClockInPin->GetCurrentMidiTick(), true);
+					AllNotesOff(StartFrameIndex, MidiClockInPin->GetLastProcessedMidiTick(), true);
 					return EMusicPlayerTransportState::Prepared;
 
 				default:
@@ -473,11 +476,11 @@ namespace HarmonixMetasound::Nodes::StepSequencePlayer
 				return EMusicPlayerTransportState::Paused;
 
 			case EMusicPlayerTransportState::Stopping:
-				AllNotesOff(0, MidiClockInPin->GetCurrentMidiTick(), true);
+				AllNotesOff(0, MidiClockInPin->GetLastProcessedMidiTick(), true);
 				return EMusicPlayerTransportState::Prepared;
 
 			case EMusicPlayerTransportState::Killing:
-				AllNotesOff(0, MidiClockInPin->GetCurrentMidiTick(), true);
+				AllNotesOff(0, MidiClockInPin->GetLastProcessedMidiTick(), true);
 				return EMusicPlayerTransportState::Prepared;
 
 			default:
@@ -596,7 +599,7 @@ namespace HarmonixMetasound::Nodes::StepSequencePlayer
 		}
 	}
 
-	void FStepSequencePlayerOperator::AdvanceThruTick(int32 BlockFrameIndex, int32 Tick, bool IsPreRoll)
+	void FStepSequencePlayerOperator::AdvanceThruTick(int32 BlockFrameIndex, int32 Tick)
 	{
 		if (!SequenceTable || SequenceTable->Pages.IsEmpty())
 		{
diff --git a/Engine/Plugins/Runtime/Harmonix/Source/HarmonixMetasound/Private/HarmonixMetasound/Subsystems/MidiClockUpdateSubsystem.cpp b/Engine/Plugins/Runtime/Harmonix/Source/HarmonixMetasound/Private/HarmonixMetasound/Subsystems/MidiClockUpdateSubsystem.cpp
--- a/Engine/Plugins/Runtime/Harmonix/Source/HarmonixMetasound/Private/HarmonixMetasound/Subsystems/MidiClockUpdateSubsystem.cpp
+++ b/Engine/Plugins/Runtime/Harmonix/Source/HarmonixMetasound/Private/HarmonixMetasound/Subsystems/MidiClockUpdateSubsystem.cpp
@@ -26,8 +26,13 @@ namespace MidiClockUpdateSubsystem
 					UpdateMethod = (EUpdateMethod)NewValue;
 				}
 			}));
+
+	int32 kClockHistorySize = 100;
 }
 
+FCriticalSection UMidiClockUpdateSubsystem::ClockHistoryMapLocker;
+TMap<uint32, TWeakPtr<HarmonixMetasound::Analysis::FMidiClockSongPositionHistory>> UMidiClockUpdateSubsystem::ClockHistories;
+
 bool UMidiClockUpdateSubsystem::IsTickable() const
 {
 	using namespace MidiClockUpdateSubsystem;
@@ -38,8 +43,8 @@ bool UMidiClockUpdateSubsystem::IsTickable() const
 		// In either of these cases we need our tick function called IF there are tracked clocks. 
 		{
 			FScopeLock Lock{ &TrackedMidiClocksMutex };
-			return TrackedMidiClocks.Num() > 0 || TrackedMusicClockComponents.Num() > 0;
-		}
+			return TrackedMusicClockComponents.Num() > 0 || ClockHistories.Num() > 0;
+	}
 	default:
 		// Midi clocks and music clock components are ticked elsewhere.
 		return false;
@@ -53,12 +58,10 @@ void UMidiClockUpdateSubsystem::Tick(float DeltaTime)
 	{
 	case EUpdateMethod::EngineTickableObject:
 		// We tick BOTH the midi clocks and the music clock components here.
-		UpdateFMidiClocks();
 		UpdateUMusicClockComponents();
 		return;
 	case EUpdateMethod::EngineTickableObjectAndTickComponent:
-		// The original method... ONLY the midi clocks are ticked here. 
-		UpdateFMidiClocks();
+		// The original method... ONLY the midi clocks were ticked here. 
 		return;
 	default:
 		// Midi clocks and music clock components are ticked elsewhere.
@@ -83,7 +86,7 @@ void UMidiClockUpdateSubsystem::Deinitialize()
 	FCoreDelegates::OnBeginFrame.Remove(EngineBeginFrameDelegate);
 }
 
-void UMidiClockUpdateSubsystem::TrackMidiClock(HarmonixMetasound::FMidiClock* Clock)
+void UMidiClockUpdateSubsystem::TrackMusicClockComponent(UMusicClockComponent* Clock)
 {
 	check(GEngine);
 
@@ -91,68 +94,59 @@ void UMidiClockUpdateSubsystem::TrackMidiClock(HarmonixMetasound::FMidiClock* Cl
 
 	check(UpdateSubsystem);
 
-	UpdateSubsystem->TrackMidiClockImpl(Clock);
+	UpdateSubsystem->TrackMusicClockComponentImpl(Clock);
 }
 
-void UMidiClockUpdateSubsystem::StopTrackingMidiClock(HarmonixMetasound::FMidiClock* Clock)
+void UMidiClockUpdateSubsystem::StopTrackingMusicClockComponent(UMusicClockComponent* Clock)
 {
 	if (GEngine)
 	{
 		if (UMidiClockUpdateSubsystem* UpdateSubsystem = GEngine->GetEngineSubsystem<UMidiClockUpdateSubsystem>())
 		{
-			UpdateSubsystem->StopTrackingMidiClockImpl(Clock);
+			UpdateSubsystem->StopTrackingMusicClockComponentImpl(Clock);
 		}
 	}
 }
 
-void UMidiClockUpdateSubsystem::TrackMidiClockImpl(HarmonixMetasound::FMidiClock* Clock)
+uint32 UMidiClockUpdateSubsystem::MakeMidiSongPosAnalyzerAddressHash(const Metasound::Frontend::FAnalyzerAddress& ForAddress)
 {
-	check(nullptr != Clock);
-
-	FScopeLock Lock{ &TrackedMidiClocksMutex };
-
-	TrackedMidiClocks.AddUnique(Clock);
+	uint32 AddressHash = GetTypeHashHelper(ForAddress.AnalyzerMemberName);
+	AddressHash = HashCombineFast(AddressHash, GetTypeHashHelper(ForAddress.AnalyzerName));
+	AddressHash = HashCombineFast(AddressHash, GetTypeHashHelper(ForAddress.DataType));
+	AddressHash = HashCombineFast(AddressHash, GetTypeHashHelper(ForAddress.InstanceID));
+	AddressHash = HashCombineFast(AddressHash, ForAddress.NodeID.A);
+	AddressHash = HashCombineFast(AddressHash, GetTypeHashHelper(ForAddress.OutputName));
+	return AddressHash;
 }
 
-void UMidiClockUpdateSubsystem::UpdateFMidiClocks()
+UMidiClockUpdateSubsystem::FClockHistoryPtr UMidiClockUpdateSubsystem::GetOrCreateClockHistory(const Metasound::Frontend::FAnalyzerAddress& ForAddress)
 {
-	TRACE_CPUPROFILER_EVENT_SCOPE(UpdateFMidiClocks);
+	FScopeLock Lock(&ClockHistoryMapLocker);
 
-	FScopeLock Lock{ &TrackedMidiClocksMutex };
+	uint32 AddressHash = MakeMidiSongPosAnalyzerAddressHash(ForAddress);
 
-	for (HarmonixMetasound::FMidiClock* Clock : TrackedMidiClocks)
+	if (ClockHistories.Contains(AddressHash))
 	{
-		Clock->UpdateLowResCursors();
+		if (FClockHistoryPtr HistoryPtr = ClockHistories[AddressHash].Pin())
+		{
+			return HistoryPtr;
+		}
+		FClockHistoryPtr NewHistory = MakeShared<HarmonixMetasound::Analysis::FMidiClockSongPositionHistory>(MidiClockUpdateSubsystem::kClockHistorySize);
+		ClockHistories[AddressHash] = NewHistory;
+		return NewHistory;
 	}
-}
-
-void UMidiClockUpdateSubsystem::StopTrackingMidiClockImpl(HarmonixMetasound::FMidiClock* Clock)
-{
-	FScopeLock Lock{ &TrackedMidiClocksMutex };
 
-	TrackedMidiClocks.Remove(Clock);
-}
-
-void UMidiClockUpdateSubsystem::TrackMusicClockComponent(UMusicClockComponent* Clock)
-{
-	check(GEngine);
-
-	UMidiClockUpdateSubsystem* UpdateSubsystem = GEngine->GetEngineSubsystem<UMidiClockUpdateSubsystem>();
-
-	check(UpdateSubsystem);
-
-	UpdateSubsystem->TrackMusicClockComponentImpl(Clock);
-}
-
-void UMidiClockUpdateSubsystem::StopTrackingMusicClockComponent(UMusicClockComponent* Clock)
-{
-	if (GEngine)
+	for (auto It = ClockHistories.CreateIterator(); It; ++It)
 	{
-		if (UMidiClockUpdateSubsystem* UpdateSubsystem = GEngine->GetEngineSubsystem<UMidiClockUpdateSubsystem>())
+		if (!It.Value().IsValid())
 		{
-			UpdateSubsystem->StopTrackingMusicClockComponentImpl(Clock);
+			It.RemoveCurrent();
 		}
 	}
+
+	FClockHistoryPtr NewHistory = MakeShared<HarmonixMetasound::Analysis::FMidiClockSongPositionHistory>(MidiClockUpdateSubsystem::kClockHistorySize);
+	ClockHistories.Add(AddressHash, NewHistory);
+	return NewHistory;
 }
 
 void UMidiClockUpdateSubsystem::TrackMusicClockComponentImpl(UMusicClockComponent* Clock)
@@ -171,7 +165,6 @@ void UMidiClockUpdateSubsystem::CoreDelegatesBeginFrame()
 	switch (UpdateMethod)
 	{
 	case EUpdateMethod::EngineSubsystemCoreDelegatesOnBeginFrame:
-		UpdateFMidiClocks();
 		UpdateUMusicClockComponents();
 		return;
 	default:
@@ -185,7 +178,6 @@ void UMidiClockUpdateSubsystem::CoreDelegatesSamplingInput()
 	switch (UpdateMethod)
 	{
 	case EUpdateMethod::EngineSubsystemCoreDelegatesOnSamplingInput:
-		UpdateFMidiClocks();
 		UpdateUMusicClockComponents();
 		return;
 	default:
@@ -209,6 +201,14 @@ void UMidiClockUpdateSubsystem::UpdateUMusicClockComponents()
 			ClockIterator.RemoveCurrentSwap();
 		}
 	}
+
+	for (auto It = ClockHistories.CreateIterator(); It; ++It)
+	{
+		if (!It.Value().IsValid())
+		{
+			It.RemoveCurrent();
+		}
+	}
 }
 
 // Implement a "tick" method that can be used during automated testing so that
@@ -219,18 +219,14 @@ void UMidiClockUpdateSubsystem::TickForTesting()
 	switch (UpdateMethod)
 	{
 	case MidiClockUpdateSubsystem::EUpdateMethod::EngineTickableObjectAndTickComponent:
-		UpdateFMidiClocks();
 		break;
 	case MidiClockUpdateSubsystem::EUpdateMethod::EngineSubsystemCoreDelegatesOnBeginFrame:
-		UpdateFMidiClocks();
 		UpdateUMusicClockComponents();
 		break;
 	case MidiClockUpdateSubsystem::EUpdateMethod::EngineTickableObject:
-		UpdateFMidiClocks();
 		UpdateUMusicClockComponents();
 		break;
 	case MidiClockUpdateSubsystem::EUpdateMethod::EngineSubsystemCoreDelegatesOnSamplingInput:
-		UpdateFMidiClocks();
 		UpdateUMusicClockComponents();
 		break;
 	default:
diff --git a/Engine/Plugins/Runtime/Harmonix/Source/HarmonixMetasound/Private/HarmonixMetasoundModule.cpp b/Engine/Plugins/Runtime/Harmonix/Source/HarmonixMetasound/Private/HarmonixMetasoundModule.cpp
--- a/Engine/Plugins/Runtime/Harmonix/Source/HarmonixMetasound/Private/HarmonixMetasoundModule.cpp
+++ b/Engine/Plugins/Runtime/Harmonix/Source/HarmonixMetasound/Private/HarmonixMetasoundModule.cpp
@@ -7,6 +7,7 @@
 #include "Analysis/MetasoundFrontendAnalyzerRegistry.h"
 
 #include "HarmonixMetasound/Analysis/MidiClockVertexAnalyzer.h"
+#include "HarmonixMetasound/Analysis/MidiSongPosVertexAnalyzer.h"
 #include "HarmonixMetasound/Analysis/MidiStreamVertexAnalyzer.h"
 #include "HarmonixMetasound/Analysis/FFTAnalyzerResultVertexAnalyzer.h"
 #include "HarmonixMetasound/DataTypes/FFTAnalyzerResult.h"
@@ -40,6 +41,7 @@ void FHarmonixMetasoundModule::StartupModule()
 	// Register vertex analyzer factories
 	METASOUND_REGISTER_VERTEX_ANALYZER_FACTORY(HarmonixMetasound::Analysis::FMidiStreamVertexAnalyzer)
 	METASOUND_REGISTER_VERTEX_ANALYZER_FACTORY(HarmonixMetasound::Analysis::FMidiClockVertexAnalyzer)
+	METASOUND_REGISTER_VERTEX_ANALYZER_FACTORY(HarmonixMetasound::Analysis::FMidiSongPosVertexAnalyzer)
 	METASOUND_REGISTER_VERTEX_ANALYZER_FACTORY(HarmonixMetasound::Analysis::FFFTAnalyzerResultVertexAnalyzer)
 
 	// The first redirect for thie module
diff --git a/Engine/Plugins/Runtime/Harmonix/Source/HarmonixMetasound/Public/HarmonixMetasound/Components/MusicClockComponent.h b/Engine/Plugins/Runtime/Harmonix/Source/HarmonixMetasound/Public/HarmonixMetasound/Components/MusicClockComponent.h
--- a/Engine/Plugins/Runtime/Harmonix/Source/HarmonixMetasound/Public/HarmonixMetasound/Components/MusicClockComponent.h
+++ b/Engine/Plugins/Runtime/Harmonix/Source/HarmonixMetasound/Public/HarmonixMetasound/Components/MusicClockComponent.h
@@ -126,6 +126,12 @@ public:
 	UFUNCTION(BlueprintPure, Category = "MusicClock")
 	float GetBeatsIncludingCountIn(ECalibratedMusicTimebase Timebase = ECalibratedMusicTimebase::VideoRenderTime) const;
 
+	// NOTE: Working in ticks is a little risky. Midi files have have different numbers of ticks per quarter note,
+	// Ticks change duration with tempo changes, etc. So we don't expose ticks to blueprints and recommend using them
+	// in c++ code. That said, sometimes there is a use for them... so... 
+	float GetTicksFromBarOne(ECalibratedMusicTimebase Timebase = ECalibratedMusicTimebase::VideoRenderTime) const;
+	float GetTicksIncludingCountIn(ECalibratedMusicTimebase Timebase = ECalibratedMusicTimebase::VideoRenderTime) const;
+
 	// Returns the "classic" musical timestamp in the form Bar (int) & Beat (float). In this form...
 	//    - Bar 1, Beat 1.0 is the "beginning of the song" AFTER count-in/pickups
 	//    - Bar 0, Beat 1.0 would be one bar BEFORE the "beginning of the song"... eg. a bar of count-in or pickup.
@@ -252,7 +258,7 @@ public:
 	UFUNCTION(BlueprintPure, Category = "Song Data")
 	float GetSongLengthBars() const;
 
-	const FSongMaps& GetSongMaps() const;
+	const ISongMapEvaluator& GetSongMaps() const;
 
 	UPROPERTY(EditAnywhere, BlueprintReadWrite, Category = "MusicClock")
 	ECalibratedMusicTimebase TimebaseForBarAndBeatEvents = ECalibratedMusicTimebase::VideoRenderTime;
@@ -368,7 +374,7 @@ struct FMusicClockDriverBase : public TSharedFromThis<FMusicClockDriverBase>
 public:
 	FMusicClockDriverBase() = delete;
 	FMusicClockDriverBase(UMusicClockComponent* DrivenClock)
-		: Clock(DrivenClock)
+		: ClockComponent(DrivenClock)
 	{}
 	virtual ~FMusicClockDriverBase() = default;
 
@@ -381,10 +387,11 @@ public:
 	virtual void OnPause() = 0;
 	virtual void OnContinue() = 0;
 	virtual void OnStop() = 0;
-	virtual const FSongMaps* GetCurrentSongMaps() const = 0;
+	virtual const ISongMapEvaluator* GetCurrentSongMapEvaluator() const = 0;
 
 protected:
-	UMusicClockComponent* Clock;
+	UMusicClockComponent* ClockComponent;
+
 private:
 	virtual bool RefreshCurrentSongPos() = 0;
 };
diff --git a/Engine/Plugins/Runtime/Harmonix/Source/HarmonixMetasound/Public/HarmonixMetasound/DataTypes/MidiClock.h b/Engine/Plugins/Runtime/Harmonix/Source/HarmonixMetasound/Public/HarmonixMetasound/DataTypes/MidiClock.h
--- a/Engine/Plugins/Runtime/Harmonix/Source/HarmonixMetasound/Public/HarmonixMetasound/DataTypes/MidiClock.h
+++ b/Engine/Plugins/Runtime/Harmonix/Source/HarmonixMetasound/Public/HarmonixMetasound/DataTypes/MidiClock.h
@@ -15,8 +15,6 @@
 #include "HarmonixMidi/SongMaps.h"
 #include "HarmonixMetasound/DataTypes/MusicTransport.h"
 #include "HarmonixMetasound/DataTypes/MidiClockEvent.h"
-#include "HarmonixMidi/MidiPlayCursorMgr.h"
-#include "HarmonixMidi/MidiPlayCursor.h"
 #include "Sound/QuartzQuantizationUtilities.h"
 
 namespace Metasound
@@ -34,28 +32,9 @@ namespace Metasound
 
 namespace HarmonixMetasound
 {
-	struct HARMONIXMETASOUND_API FMidiTimestampTransportState
-	{
-		int32 BlockSampleFrameIndex = 0;
-		float BlockSampleFrameOffset = 0.0f;
-		EMusicPlayerTransportState TransportState = EMusicPlayerTransportState::Invalid;
-	};
-
-	struct HARMONIXMETASOUND_API FMidiTimestampSpeed
-	{
-		int32 BlockSampleFrameIndex = 0;
-		float BlockSampleFrameOffset = 0.0f;
-		float Speed = 1.0f;
-	};
-
-	struct HARMONIXMETASOUND_API FMidiTimestampTempo
-	{
-		int32 BlockSampleFrameIndex = 0;
-		float BlockSampleFrameOffset = 0.0f;
-		float Tempo = 120.0f;
-	};
-
-	constexpr FMidiTimestampTempo InvalidMidiTimestampTempo{ -1, 0.0f, 0.0f };
+	class FMidiClock;
+	using FConstSharedMidiClockPtr = TSharedPtr<const FMidiClock, ESPMode::NotThreadSafe>;
+	using FMidiClockEvents = TArray<FMidiClockEvent>;
 
 	class HARMONIXMETASOUND_API FMidiClock : public TSharedFromThis<FMidiClock, ESPMode::NotThreadSafe>
 	{
@@ -63,69 +42,83 @@ namespace HarmonixMetasound
 		static constexpr int32 kMidiGranularity = 128;
 		
 		explicit FMidiClock(const Metasound::FOperatorSettings& InSettings);
-		virtual ~FMidiClock();
-
 		FMidiClock(const FMidiClock& Other);
+		virtual ~FMidiClock();
 		FMidiClock& operator=(const FMidiClock& Other);
 
-		void ResetAndStart(int32 FrameIndex, bool SeekToStart = true);
+		void AttachToMidiFile(TSharedPtr<FMidiFileData> MidiData, bool ResetToStart = true);
+		void MidiChanged();
+		void DetachFromMidiFile();
+		const ISongMapEvaluator& GetSongMapEvaluator() const { return *SongMapEvaluator; }
+
+		void SetDrivingClock(FConstSharedMidiClockPtr NewExternalClockDriver);
+		FConstSharedMidiClockPtr GetDrivingClock() const { return ExternalClockDriver; }
 
 		void PrepareBlock();
 
-		bool HasLowResCursors() const;
-		void UpdateLowResCursors();
+		void SetTransportState(int32 BlockFrameIndex, EMusicPlayerTransportState TransportState);
+		void SetSpeed(int32 BlockFrameIndex, float Speed);
+		void SetTempo(int32 BlockFrameIndex, int32 Tick, float Bpm);
+		void SetTimeSignature(int32 BlockFrameIndex, int32 Tick, const FTimeSignature& TimeSignature);
 
-		void AddTransportStateChangeToBlock(const FMidiTimestampTransportState& NewTransportState);
-		void AddSpeedChangeToBlock(const FMidiTimestampSpeed& NewSpeed);
+		// directly seek this clock with a musical seek target or a specific tick
+		void SeekTo(int32 BlockFrameIndex, const FMusicSeekTarget& InTarget);
+		void SeekTo(int32 BlockFrameIndex, int32 Tick);
+
+		// This will add a loop event to the clock event stream WITHOUT having 
+		// to set this clock to looping. This is used when this clock is being 
+		// driven by an external clock and THAT clock's looping setup causes this
+		// clock to loop. 
+		void AddTransientLoop(int32 BlockFrameIndex, int32 NewFirstTiokInLoop, int32 NewLoopLengthTicks);
+		void SetupPersistentLoop(int32 NewFirstTickInLoop, int32 NewLoopLengthTicks);
+		void ClearPersistentLoop();
+		bool HasPersistentLoop() const;
+		int32 GetFirstTickInLoop() const { return FirstTickInLoop; }
+		int32 GetLoopLengthTicks() const { return LoopLengthTicks; }
+		float GetLoopStartMs() const;
+		float GetLoopEndMs() const;
+		float GetLoopLengthMs() const;
 
-		bool HasTransportStateChangeInBlock() const
-		{
-			return TransportChangesInBlock.IsEmpty() == false;
-		}
-		bool HasSpeedChangesInBlock() const
-		{
-			return HasSpeedChangeInBlock;
-		}
-		bool HasTempoChangesInBlock() const
-		{
-			return HasTempoChangeInBlock;
-		}
+		// process and advance the clock based on the driving clock given sample frames
+		// will handle the driving clock events based on the frame range
+		void Advance(const FMidiClock& DrivingClock, int32 StartFrame, int32 NumFrames);
 
-		const TArray<FMidiClockEvent>& GetMidiClockEventsInBlock() const;
-		const FMidiClockEvent* FindLastMidiClockEventAtBlockSampleFrame(int32 FrameIndex) const;
+		// process and advance the clock normally based on the given sample frames
+		void Advance(int32 StartFrame, int32 NumFrames);
 
-		EMusicPlayerTransportState GetTransportStateAtBlockSampleFrame(int32 FrameIndex) const;
-		EMusicPlayerTransportState GetTransportStateAtEndOfBlock() const;
-		const FMidiTimestampTransportState& GetTransportTimestampForBlockSampleFrame(int32 FrameIndex) const;
-		const TArray<FMidiTimestampTransportState>& GetTransportTimestampsInBlock() const
-		{
-			return TransportChangesInBlock;
-		}
+		bool AdvanceToTick(int32 BlockFrameIndex, int32 UpToTick);
+		bool AdvanceToMs(int32 BlockFrameIndex, float Ms);
+
+		bool  HasTransportStateChangesInBlock() const    { return NumTransportChangeInBlock > 0; }
+		int32 GetNumTransportStateChangesInBlock() const { return NumTransportChangeInBlock;     }
+		bool  HasSpeedChangesInBlock() const             { return NumSpeedChangeInBlock > 0;     }
+		int32 GetNumSpeedChangesInBlock() const          { return NumSpeedChangeInBlock;         }
+		bool  HasTempoChangesInBlock() const             { return NumTempoChangeInBlock > 0;     }
+		int32 GetNumTempoChangesInBlock() const          { return NumTempoChangeInBlock;         }
+
+		const FMidiClockEvents& GetMidiClockEventsInBlock() const { return MidiClockEventsInBlock; }
 
+		EMusicPlayerTransportState GetTransportStateAtStartOfBlock() const { return TransportAtBlockStart; }
+		EMusicPlayerTransportState GetTransportStateAtEndOfBlock() const { return TransportAtBlockEnd; }
+
+		float GetSpeedAtStartOfBlock() const { return SpeedAtBlockStart; }
 		float GetSpeedAtBlockSampleFrame(int32 FrameIndex) const;
-		float GetSpeedAtEndOfBlock() const;
-		const FMidiTimestampSpeed& GetSpeedTimestampForBlockSampleFrame(int32 FrameIndex) const;
-		const TArray<FMidiTimestampSpeed>& GetTempoSpeedTimestampsInBlock() const
-		{
-			return SpeedChangesInBlock;
-		}
+		float GetSpeedAtEndOfBlock() const { return SpeedAtBlockEnd; }
 
+		float GetTempoAtStartOfBlock() const { return TempoAtBlockStart; }
 		float GetTempoAtBlockSampleFrame(int32 FrameIndex) const;
-		float GetTempoAtEndOfBlock() const;
-		int32 GetNumTempoChangesInBlock() const;
-		FMidiTimestampTempo GetTempoChangeByIndex(int32 Index) const;
+		float GetTempoAtEndOfBlock() const { return TempoAtBlockEnd; }
 
-		int32 GetCurrentMidiTick() const;
+		int32 GetLastProcessedMidiTick() const { return LastProcessedMidiTick; }
+		int32 GetNextMidiTickToProcess() const { return NextMidiTickToProcess; }
 
-		int32 GetCurrentBlockFrameIndex() const;
-
-		float GetQuarterNoteIncludingCountIn() const;
+		float GetCurrentSongPosMs() const;
 
 		/**
 		 * @brief Get the timestamp after the most recent clock update
 		 * @return The current timestamp
 		 */
-		FMusicTimestamp GetCurrentMusicTimestamp() const;
+		FMusicTimestamp GetMusicTimestampAtBlockEnd() const;
 		
 		/**
 		 * @brief Get the music timestamp at a given frame offset from the last processed audio block.
@@ -135,73 +128,15 @@ namespace HarmonixMetasound
 		FMusicTimestamp GetMusicTimestampAtBlockOffset(int32 Offset) const;
 
 		/**
-		 * @brief Get the absolute time in ms for a frame within the last audio block
+		 * @brief Get the absolute "music time" in ms for a frame within the last audio block. This is the 
+		 * time in the musical content that the clock has advanced "up to". Note: This time will not be sample
+		 * accurate as midi processing advances by ticks, and the time is calculated by turning the "current tick"
+		 * at the offset provided into a time in ms.
 		 * @param Offset - The frame index from the beginning of the last processed audio block
 		 * @return The absolute time in ms
 		 */
-		float GetMsAtBlockOffset(int32 Offset) const;
-
-		//*****************************************************************************************
-		// NOTE: These next functions are a facade in front of this class's MidiPlayCursorMgr. 
-		void AttachToMidiResource(TSharedPtr<FMidiFileData> MidiDataProxy, bool ResetCursorsToStart = true, int32 PreRollBars = 0) { DrivingMidiPlayCursorMgr->AttachToMidiResource(MidiDataProxy, ResetCursorsToStart, PreRollBars); }
-		void DetachFromMidiResource() { DrivingMidiPlayCursorMgr->DetachFromMidiResource(); }
-		void AttachToTimeAuthority(const FMidiClock& MidiClockRef);
-		void DetachFromTimeAuthority();
-		void InformOfCurrentAdvanceRate(float AdvanceRate);
-		
-		void LockForMidiDataChanges();
-		void MidiDataChangesComplete(FMidiPlayCursorMgr::EMidiChangePositionCorrectMode Mode = FMidiPlayCursorMgr::EMidiChangePositionCorrectMode::MaintainTick);
-		const FSongMaps& GetSongMaps() const     { return DrivingMidiPlayCursorMgr->GetSongMaps();   }
-		const FTempoMap& GetTempoMap() const     { return DrivingMidiPlayCursorMgr->GetTempoMap();   }
-		const FBarMap& GetBarMap() const         { return DrivingMidiPlayCursorMgr->GetBarMap();     }
-		void AdvanceHiResToMs(int32 BlockFrameIndex, float Ms, bool Broadcast);
-		void SeekTo(const FMusicSeekTarget& Timestamp, int32 PreRollBars);
-		void SeekTo(int32 Tick, int32 PrerollBars);
-
-		void SetLoop(int32 StartTick, int32 EndTick) { DrivingMidiPlayCursorMgr->SetLoop(StartTick, EndTick, false, true); }
-		void ClearLoop() { DrivingMidiPlayCursorMgr->ClearLoop(true); }
-		bool DoesLoop() const { return DrivingMidiPlayCursorMgr->DoesLoop(false); }
-		float GetLoopStartMs() const { return DrivingMidiPlayCursorMgr->GetLoopStartMs(false); }
-		int32 GetLoopStartTick() const { return DrivingMidiPlayCursorMgr->GetLoopStartTick(false); }
-		float GetLoopEndMs() const { return DrivingMidiPlayCursorMgr->GetLoopEndMs(false); }
-		int32 GetLoopEndTick() const { return DrivingMidiPlayCursorMgr->GetLoopEndTick(false); }
-
-		int32 GetCurrentHiResTick() const  { return DrivingMidiPlayCursorMgr->GetCurrentHiResTick();  }
-		float GetCurrentHiResMs() const    { return DrivingMidiPlayCursorMgr->GetCurrentHiResMs();    }
-		float GetElapsedHiResMs() const    { return DrivingMidiPlayCursorMgr->GetElapsedHiResMs();    }
-		int32 GetCurrentLowResTick() const { return DrivingMidiPlayCursorMgr->GetCurrentLowResTick(); }
-		float GetCurrentLowResMs() const   { return DrivingMidiPlayCursorMgr->GetCurrentLowResMs();   }
-		float GetElapsedLowResMs() const   { return DrivingMidiPlayCursorMgr->GetElapsedLowResMs();   }
-		// These next few are const because they need to be callable by things that hold a ReadRef to us!
-		void RegisterHiResPlayCursor(FMidiPlayCursor* PlayCursor, float PreRollMs = -1.0f) const  { DrivingMidiPlayCursorMgr->RegisterHiResPlayCursor(PlayCursor, PreRollMs);  }
-		void RegisterLowResPlayCursor(FMidiPlayCursor* PlayCursor, float PreRollMs = -1.0f) const { DrivingMidiPlayCursorMgr->RegisterLowResPlayCursor(PlayCursor, PreRollMs); }
-		void UnregisterPlayCursor(FMidiPlayCursor* PlayCursor, bool WarnOnFail = true) const      { DrivingMidiPlayCursorMgr->UnregisterPlayCursor(PlayCursor, WarnOnFail);    }
-		void UnregisterAllPlayCursors()                                                           { DrivingMidiPlayCursorMgr->UnregisterAllPlayCursors();                      }
-
-		TSharedPtr<FMidiPlayCursorMgr> GetDrivingMidiPlayCursorMgr() const { return DrivingMidiPlayCursorMgr.ToSharedPtr(); }
-		//*****************************************************************************************
-		
-		// handle and add transport change event
-		// to be called within the Transport Post Processor with the new transport state after the normal Processor
-		void HandleTransportChange(int32 BlockFrameIndex, EMusicPlayerTransportState TransportState);
-
-		// handle single clock event
-		void HandleClockEvent(const FMidiClock& DrivingClock, const FMidiClockEvent& Event, int32 PrerollBars, float Speed = 1.0f);
-		
-		// process and advance the clock based on the driving clock given sample frames
-		// will handle the driving clock events based on the frame range
-		void Process(const FMidiClock& DrivingClock, int32 StartFrame, int32 NumFrames, int32 PrerollBars, float Speed = 1.0f);
+		float GetSongPosMsAtBlockOffset(int32 Offset) const;
 
-		// process and advance the clock normally based on the given sample frames
-		void Process(int32 StartFrame, int32 NumFrames, int32 PrerollBars, float Speed = 1.0f);
-		
-		// directly perform and write an advance to this clock
-		void WriteAdvance(int32 StartFrameIndex, int32 EndFrameIndex, float InSpeed = 1.0f);
-
-		// directly seek this clock with a musical seek target or a specific tick
-		void SeekTo(int32 BlockFrameIndex, const FMusicSeekTarget& InTarget, int32 InPrerollBars);
-		void SeekTo(int32 BlockFrameIndex, int32 Tick, int32 InPrerollBars);
-		
 		/**
 		 * Given an input tick, outputs a looped tick if the input tick is > the StartTick of the Loop Region
 		 * If the clock is not looping, or loop region length is 0, then the output will be unchanged.
@@ -225,71 +160,116 @@ namespace HarmonixMetasound
 		 * @param		Tick - Absolute Tick
 		 * @return		Looped Tick if Tick > LoopEnd: LoopedTick = LoopStart + (Tick - LoopStart) % (LoopEnd - LoopStart) 
 		 */
-		int32 CalculateMappedTick(int32 Tick) const;
+		int32 WrapTickIfLooping(int32 Tick) const;
 		
-		// copy speed and tempo changes from in clock to this clock
-		// with an optional Speed multiplier to adjust the out going speed on this clock
-		void CopySpeedAndTempoChanges(const FMidiClock* InClock, float InSpeedMult = 1.0f);
-
 		// Creates a new FMidiFileData with the given starting Tempo and Time Signature
 		// With max song length to be played indefinitely. Useful for Midi Clock Metronomes.
 		static TSharedPtr<FMidiFileData> MakeClockConductorMidiData(float TempoBPM, int32 TimeSigNum, int32 TimeSigDen);
 
-	private:
-		friend class FMidiClockEventCursor;
-
-		class FMidiClockEventCursor : public FMidiPlayCursor
-		{
-		public:
-			FMidiClockEventCursor(FMidiClock* MidiClock);
-		
-			//~ BEGIN FMidiPlayCursor Overrides
-			virtual void Reset(bool ForceNoBroadcast = false) override;
-			virtual void OnLoop(int32 LoopStartTick, int32 LoopEndTick) override;
-			virtual void SeekToTick(int32 Tick) override;
-			virtual void SeekThruTick(int32 Tick) override;
-			virtual void AdvanceThruTick(int32 Tick, bool IsPreRoll) override;
-			virtual void OnTempo(int32 TrackIndex, int32 Tick, int32 Tempo, bool IsPreroll = false) override;
-			virtual void OnTimeSig(int32 TrackIndex, int32 Tick, int32 Numerator, int32 Denominator, bool IsPreroll) override;
-			//~ END FMidiPlayCursor Overrides
-			
-			void AddEvent(const FMidiClockEvent& InEvent) const;
-		private:
-			FMidiClock* MyMidiClock = nullptr;
-			int32 CurrentAdvanceStartTick = -1;
-		};
-		
-		FMidiClockEventCursor MidiClockEventCursor;
+		bool GetMidiDataChangedInBlock() const { return MidiDataChangedInBlock; }
 
 	private:
-		void RegisterForGameThreadUpdates();
-		void UnregisterForGameThreadUpdates();
+		int32 GetNextTickToProcessAtBlockFrame(int32 BlockFrame) const;
+		void AddEvent(const FMidiClockEvent& InEvent, bool bRequireSequential = true);
+		void HandleClockEvent(const FMidiClock& DrivingClock, const FMidiClockEvent& Event);
+		void PostTempoOrTimeSignatureEventsIfNeeded();
 		
+		template <typename MSGTYPE>
+		MSGTYPE* LookForEventOnMidiTick(int32 Tick);
+		template <typename MSGTYPE>
+		MSGTYPE* LookForEventOnBlockFrameIndex(int32 BlockFrameIndex);
+
+		void AddTransportStateChangeToBlock(int32 BlockFrameIndex, EMusicPlayerTransportState TransportState);
+		void AddTimeSignatureChangeToBlock(int32 BlockFrameIndex, int32 Tick, const FTimeSignature& TimeSignature);
+		void AddTempoChangeToBlock(int32 BlockFrameIndex, int32 Tick, float Tempo);
+		void AddSpeedChangeToBlock(int32 BlockFrameIndex, float Speed, bool bIsNewLocalSpeed);
+		void AddLoopToBlock(int32 BlockFrameIndex, int32 FirstTick, int32 LoopLength);
+		void AddSeekToBlock(int32 BlockFrameIndex, int32 ToTick);
+		void AddAdvanceToBlock(int32 BlockFrameIndex, int32 FirstTick, int32 NumTicks);
+		void RebuildSongMapEvaluator(const TSharedPtr<const FMidiFileData>& MidiWithTempo, const TSharedPtr<const FMidiFileData>& MidiWithOtherMaps);
+
+		TSharedPtr<FSongMapsWithAlternateTempoSource> SongMapEvaluator;
+		int32 CurrentTempoInfoPointIndex;
+		int32 CurrentTimeSignaturePointIndex;
+
+		FConstSharedMidiClockPtr ExternalClockDriver;
+		float TickResidualWhenDriven;
+
 		int32 BlockSize;
 		int32 CurrentBlockFrameIndex;
+		int32 LastProcessedMidiTick;
+		int32 NextMidiTickToProcess;
 		float SampleRate;
 		Metasound::FSampleCount SampleCount;
-		int32 FramesUntilNextProcess = 0;
-		FMidiTimestampTransportState CurrentTransportState;
-
-		TArray<FMidiTimestampTransportState> TransportChangesInBlock;
-		bool								 HasSpeedChangeInBlock;
-		TArray<FMidiTimestampSpeed>          SpeedChangesInBlock;
-		bool								 HasTempoChangeInBlock;
-		TArray<FMidiTimestampTempo>          TempoChangesInBlock;
+		int32 FramesUntilNextProcess;
+		EMusicPlayerTransportState TransportAtBlockStart;
+		EMusicPlayerTransportState TransportAtBlockEnd;
+		float SpeedAtBlockStart;
+		float SpeedAtBlockEnd;
+		float CurrentLocalSpeed;
+		float TempoAtBlockStart;
+		float TempoAtBlockEnd;
+		FTimeSignature TimeSignatureAtBlockStart;
+		FTimeSignature TimeSignatureAtBlockEnd;
+
+		int32 NumTransportChangeInBlock;
+		int32 NumSpeedChangeInBlock;
+		int32 NumTempoChangeInBlock;
+		int32 NumTimeSignatureChangeInBlock;
 		
-		// midi clock events are modified by a MidiClockEventCursor
-		// which has to be manually registered to the MidiClock
-		// (this is to avoid always needlessly generating midi clock events)
-		TArray<FMidiClockEvent> MidiClockEventsInBlock;
-
-		bool SmoothingEnabled = false;
-		TSharedRef<FMidiPlayCursorMgr> DrivingMidiPlayCursorMgr;
-		bool bSeekToAuthorityOnNextProcess = false;
+		int32 NextTempoChangeTick;
+		int32 NextTimeSigChangeTick;
+		int32 NextTempoOrTimeSigChangeTick;
+
+		int32 FirstTickInLoop;
+		int32 LoopLengthTicks;
+
+		bool MidiDataChangedInBlock;
+
+		FMidiClockEvents MidiClockEventsInBlock;
 	};
 
 	// Declare aliases IN the namespace...
 	DECLARE_METASOUND_DATA_REFERENCE_ALIAS_TYPES(FMidiClock, FMidiClockTypeInfo, FMidiClockReadRef, FMidiClockWriteRef)
+
+	template <typename MSGTYPE>
+	MSGTYPE* FMidiClock::LookForEventOnMidiTick(int32 Tick)
+	{
+		for (auto It = MidiClockEventsInBlock.rbegin(); It != MidiClockEventsInBlock.rend(); ++It)
+		{
+			FMidiClockEvent& AsClockEvent = *It;
+			if (MSGTYPE* AsDesiredMsgType = AsClockEvent.TryGet<MSGTYPE>())
+			{
+				if (AsDesiredMsgType->ContainsTick(Tick))
+				{
+					return AsDesiredMsgType;
+				}
+			}
+		}
+		return nullptr;
+	}
+
+	template <typename MSGTYPE>
+	MSGTYPE* FMidiClock::LookForEventOnBlockFrameIndex(int32 BlockFrameIndex)
+	{
+		for (auto It = MidiClockEventsInBlock.rbegin(); It != MidiClockEventsInBlock.rend(); ++It)
+		{
+			FMidiClockEvent& AsClockEvent = *It;
+			if (AsClockEvent.BlockFrameIndex > BlockFrameIndex)
+			{
+				continue;
+			}
+			if (AsClockEvent.BlockFrameIndex < BlockFrameIndex)
+			{
+				return nullptr;
+			}
+			if (MSGTYPE* AsDesiredMsgType = AsClockEvent.TryGet<MSGTYPE>())
+			{
+				return AsDesiredMsgType;
+			}
+		}
+		return nullptr;
+	}
 }
 
 // Declare reference types OUT of the namespace...
diff --git a/Engine/Plugins/Runtime/Harmonix/Source/HarmonixMetasound/Public/HarmonixMetasound/DataTypes/MidiClockEvent.h b/Engine/Plugins/Runtime/Harmonix/Source/HarmonixMetasound/Public/HarmonixMetasound/DataTypes/MidiClockEvent.h
--- a/Engine/Plugins/Runtime/Harmonix/Source/HarmonixMetasound/Public/HarmonixMetasound/DataTypes/MidiClockEvent.h
+++ b/Engine/Plugins/Runtime/Harmonix/Source/HarmonixMetasound/Public/HarmonixMetasound/DataTypes/MidiClockEvent.h
@@ -1,70 +1,48 @@
 // Copyright Epic Games, Inc. All Rights Reserved.
 
 #pragma once
-
+#include "HarmonixMetasound/DataTypes/MusicTransport.h"
 #include "HarmonixMidi/MidiPlayCursor.h"
 
 namespace HarmonixMetasound
 {
 	namespace MidiClockMessageTypes
 	{
-		struct FReset
-		{
-			int32 FromTick;
-			int32 ToTick;
-			bool ForceNoBroadcast;
-
-			FReset(const int32 FromTick, const int32 ToTick, const bool ForceNoBroadcast)
-				: FromTick(FromTick)
-				, ToTick(ToTick)
-				, ForceNoBroadcast(ForceNoBroadcast)
-			{}
-		};
-
 		struct FLoop
 		{
-			int32 LoopStartTick;
-			int32 LoopEndTick;
+			int32 FirstTickInLoop;
+			int32 LengthInTicks;
 
-			FLoop(const int32 LoopStartTick, const int32 LoopEndTick)
-				: LoopStartTick(LoopStartTick)
-				, LoopEndTick(LoopEndTick)
+			FLoop(const int32 FirstTickInLoop, const int32 LengthInTicks)
+				: FirstTickInLoop(FirstTickInLoop)
+				, LengthInTicks(LengthInTicks)
 			{}
 		};
 
-		struct FSeekTo
+		struct FSeek
 		{
-			int32 FromTick;
-			int32 ToTick;
+			int32 LastTickProcessedBeforeSeek;
+			int32 NewNextTick;
 
-			FSeekTo(const int32 FromTick, const int32 ToTick)
-				: FromTick(FromTick)
-				, ToTick(ToTick)
+			FSeek(const int32 LastTickProcessedBeforeSeek, const int32 NewNextTick)
+				: LastTickProcessedBeforeSeek(LastTickProcessedBeforeSeek)
+				, NewNextTick(NewNextTick)
 			{}
 		};
 
-		struct FSeekThru
+		struct FAdvance
 		{
-			int32 FromTick;
-			int32 ThruTick;
+			int32 FirstTickToProcess;
+			int32 NumberOfTicksToProcess;
 
-			FSeekThru(const int32 FromTick, const int32 ThruTick)
-				: FromTick(FromTick)
-				, ThruTick(ThruTick)
+			FAdvance(const int32 FirstTickToProcess, const int32 NumberOfTicksToProcess)
+				: FirstTickToProcess(FirstTickToProcess)
+				, NumberOfTicksToProcess(NumberOfTicksToProcess)
 			{}
-		};
 
-		struct FAdvanceThru
-		{
-			int32 FromTick;
-			int32 ThruTick;
-			bool IsPreRoll;
-
-			FAdvanceThru(const int32 FromTick, const int32 ThruTick, const bool IsPreRoll)
-				: FromTick(FromTick)
-				, ThruTick(ThruTick)
-				, IsPreRoll(IsPreRoll)
-			{}
+			int32 LastTickToProcess() const { return FirstTickToProcess + NumberOfTicksToProcess - 1; }
+
+			bool ContainsTick(int32 InTick) const { return InTick >= FirstTickToProcess && InTick < (FirstTickToProcess + NumberOfTicksToProcess); }
 		};
 
 		struct FTempoChange
@@ -76,6 +54,8 @@ namespace HarmonixMetasound
 				: Tick(Tick)
 				, Tempo(Tempo)
 			{}
+
+			bool ContainsTick(int32 InTick) const { return InTick == Tick; }
 		};
 
 		struct FTimeSignatureChange
@@ -87,23 +67,43 @@ namespace HarmonixMetasound
 				: Tick(Tick)
 				, TimeSignature(MoveTemp(TimeSignature))
 			{}
+
+			bool ContainsTick(int32 InTick) const { return InTick == Tick; }
+		};
+
+		struct FTransportChange
+		{
+			EMusicPlayerTransportState TransportState;
+
+			FTransportChange(EMusicPlayerTransportState NewTransportState)
+				: TransportState(NewTransportState)
+			{}
+		};
+
+		struct FSpeedChange
+		{
+			float Speed;
+
+			FSpeedChange(float NewSpeed)
+				: Speed(NewSpeed)
+			{}
 		};
 	}
 
 	using FMidiClockMsg = TVariant<
-		MidiClockMessageTypes::FReset,
 		MidiClockMessageTypes::FLoop,
-		MidiClockMessageTypes::FSeekTo,
-		MidiClockMessageTypes::FSeekThru,
-		MidiClockMessageTypes::FAdvanceThru,
+		MidiClockMessageTypes::FSeek,
+		MidiClockMessageTypes::FAdvance,
 		MidiClockMessageTypes::FTempoChange,
-		MidiClockMessageTypes::FTimeSignatureChange
+		MidiClockMessageTypes::FTimeSignatureChange,
+		MidiClockMessageTypes::FTransportChange,
+		MidiClockMessageTypes::FSpeedChange
 	>;
 	
 	struct HARMONIXMETASOUND_API FMidiClockEvent
 	{
 		const int32 BlockFrameIndex;
-		const FMidiClockMsg Msg;
+		FMidiClockMsg Msg;
 
 		template<typename T>
 		FMidiClockEvent(const int32 InBlockFrameIndex, T&& Msg)
@@ -117,6 +117,8 @@ namespace HarmonixMetasound
 
 		template<typename T>
 		const T* TryGet() const { return Msg.TryGet<T>(); }
-	};
 
+		template<typename T>
+		T* TryGet() { return Msg.TryGet<T>(); }
+	};
 };
diff --git a/Engine/Plugins/Runtime/Harmonix/Source/HarmonixMetasound/Public/HarmonixMetasound/DataTypes/MusicTransport.h b/Engine/Plugins/Runtime/Harmonix/Source/HarmonixMetasound/Public/HarmonixMetasound/DataTypes/MusicTransport.h
--- a/Engine/Plugins/Runtime/Harmonix/Source/HarmonixMetasound/Public/HarmonixMetasound/DataTypes/MusicTransport.h
+++ b/Engine/Plugins/Runtime/Harmonix/Source/HarmonixMetasound/Public/HarmonixMetasound/DataTypes/MusicTransport.h
@@ -80,6 +80,8 @@ namespace HarmonixMetasound
 		NumStates
 	};
 
+	FString HARMONIXMETASOUND_API MusicPlayerTransportStateToString(EMusicPlayerTransportState State);
+
 	class HARMONIXMETASOUND_API FMusicTransportControllable
 	{
 	public:
diff --git a/Engine/Plugins/Runtime/Harmonix/Source/HarmonixMetasound/Public/HarmonixMetasound/MidiOps/MidiStreamWriter.h b/Engine/Plugins/Runtime/Harmonix/Source/HarmonixMetasound/Public/HarmonixMetasound/MidiOps/MidiStreamWriter.h
--- a/Engine/Plugins/Runtime/Harmonix/Source/HarmonixMetasound/Public/HarmonixMetasound/MidiOps/MidiStreamWriter.h
+++ b/Engine/Plugins/Runtime/Harmonix/Source/HarmonixMetasound/Public/HarmonixMetasound/MidiOps/MidiStreamWriter.h
@@ -1,4 +1,4 @@
-﻿// Copyright Epic Games, Inc. All Rights Reserved.
+// Copyright Epic Games, Inc. All Rights Reserved.
 
 #pragma once
 
@@ -35,13 +35,13 @@ namespace Harmonix::Midi::Ops
 		 * 
 		 * @param InStream The MIDI stream to read from
 		 */
-		void Process(const HarmonixMetasound::FMidiStream& InStream, int32 FromTick, int32 ThruTick);
+		void Process(const HarmonixMetasound::FMidiStream& InStream, int32 FirstTickToProcess, int32 LastTickToProcess);
 		
 
 	private:
 		TUniquePtr<FArchive> Archive;
 		TMap<int32, FMidiTrack> MidiTracks;
-		int32 CurrentWriteTick = 0;
+		int32 NextWriteTick = 0;
 	
 	};
 };
\ No newline at end of file
diff --git a/Engine/Plugins/Runtime/Harmonix/Source/HarmonixMetasound/Public/HarmonixMetasound/MidiOps/PulseGenerator.h b/Engine/Plugins/Runtime/Harmonix/Source/HarmonixMetasound/Public/HarmonixMetasound/MidiOps/PulseGenerator.h
--- a/Engine/Plugins/Runtime/Harmonix/Source/HarmonixMetasound/Public/HarmonixMetasound/MidiOps/PulseGenerator.h
+++ b/Engine/Plugins/Runtime/Harmonix/Source/HarmonixMetasound/Public/HarmonixMetasound/MidiOps/PulseGenerator.h
@@ -36,7 +36,7 @@ namespace Harmonix::Midi::Ops
 		void Process(HarmonixMetasound::FMidiStream& OutStream);
 
 	private:
-		void HandleAdvanceThru(const int32 BlockFrameIndex, const HarmonixMetasound::MidiClockMessageTypes::FAdvanceThru* AdvanceThru, HarmonixMetasound::FMidiStream& OutStream);
+		void HandleAdvance(const int32 BlockFrameIndex, const HarmonixMetasound::MidiClockMessageTypes::FAdvance* Advance, HarmonixMetasound::FMidiStream& OutStream);
 		void HandleTimeSignatureChange(const HarmonixMetasound::MidiClockMessageTypes::FTimeSignatureChange* TimeSigChange);
 		void DoPulse(const int32 BlockFrameIndex, const int32 PulseTick, HarmonixMetasound::FMidiStream& OutStream);
 
diff --git a/Engine/Plugins/Runtime/Harmonix/Source/HarmonixMetasound/Public/HarmonixMetasound/Subsystems/MidiClockUpdateSubsystem.h b/Engine/Plugins/Runtime/Harmonix/Source/HarmonixMetasound/Public/HarmonixMetasound/Subsystems/MidiClockUpdateSubsystem.h
--- a/Engine/Plugins/Runtime/Harmonix/Source/HarmonixMetasound/Public/HarmonixMetasound/Subsystems/MidiClockUpdateSubsystem.h
+++ b/Engine/Plugins/Runtime/Harmonix/Source/HarmonixMetasound/Public/HarmonixMetasound/Subsystems/MidiClockUpdateSubsystem.h
@@ -6,6 +6,10 @@
 #include "Delegates/IDelegateInstance.h"
 #include "Subsystems/EngineSubsystem.h"
 
+#include "HarmonixMetasound/Analysis/SpmcAnalysisResultQueue.h"
+#include "HarmonixMetasound/Analysis/MidiClockSongPos.h"
+#include "Analysis/MetasoundFrontendAnalyzerAddress.h"
+
 #include "MidiClockUpdateSubsystem.generated.h"
 
 namespace HarmonixMetasound
@@ -58,18 +62,15 @@ public:
 	virtual void Deinitialize() override;
 	// End USubsystem
 
-	static void TrackMidiClock(HarmonixMetasound::FMidiClock* Clock);
-	static void StopTrackingMidiClock(HarmonixMetasound::FMidiClock* Clock);
 	static void TrackMusicClockComponent(UMusicClockComponent* Clock);
 	static void StopTrackingMusicClockComponent(UMusicClockComponent* Clock);
 
+	using FClockHistoryPtr = TSharedPtr<HarmonixMetasound::Analysis::FMidiClockSongPositionHistory>;
+	static FClockHistoryPtr GetOrCreateClockHistory(const Metasound::Frontend::FAnalyzerAddress& ForAddress);
+
 private:
 
 	mutable FCriticalSection TrackedMidiClocksMutex;
-	TArray<HarmonixMetasound::FMidiClock*> TrackedMidiClocks;
-	void TrackMidiClockImpl(HarmonixMetasound::FMidiClock* Clock);
-	void StopTrackingMidiClockImpl(HarmonixMetasound::FMidiClock* Clock);
-	void UpdateFMidiClocks();
 
 	TArray<TWeakObjectPtr<UMusicClockComponent>> TrackedMusicClockComponents;
 	void TrackMusicClockComponentImpl(UMusicClockComponent* Clock);
@@ -82,6 +83,11 @@ private:
 	FDelegateHandle EngineSamplingInputDelegate;
 	void CoreDelegatesSamplingInput();
 
+	static FCriticalSection ClockHistoryMapLocker;
+	static TMap<uint32, TWeakPtr<HarmonixMetasound::Analysis::FMidiClockSongPositionHistory>> ClockHistories;
+
+	static uint32 MakeMidiSongPosAnalyzerAddressHash(const Metasound::Frontend::FAnalyzerAddress& ForAddress);
+
 public:
 	// Declare a "tick" method that can be used during automated testing so that
 	// the test code doesn't need knowledge of how the low-res clocks are being ticked...
diff --git a/Engine/Plugins/Runtime/Harmonix/Source/HarmonixMetasoundTests/Private/Analysis/MidiClockVertexAnalyzerTests.cpp b/Engine/Plugins/Runtime/Harmonix/Source/HarmonixMetasoundTests/Private/Analysis/MidiClockVertexAnalyzerTests.cpp
--- a/Engine/Plugins/Runtime/Harmonix/Source/HarmonixMetasoundTests/Private/Analysis/MidiClockVertexAnalyzerTests.cpp
+++ b/Engine/Plugins/Runtime/Harmonix/Source/HarmonixMetasoundTests/Private/Analysis/MidiClockVertexAnalyzerTests.cpp
@@ -44,44 +44,36 @@ namespace HarmonixMetasoundTests::MidiClockVertexAnalyzer
 		const TSharedPtr<FMidiFileData> MidiData = MakeShared<FMidiFileData>();
 		check(MidiData);
 
-		FTempoMap& TempoMap = MidiData->SongMaps.GetTempoMap();
-		TempoMap.Empty();
-		FBarMap& BarMap = MidiData->SongMaps.GetBarMap();
-		BarMap.Empty();
+		MidiData->SongMaps.EmptyAllMaps();
 		MidiData->Tracks.Empty();
 
 		MidiData->Tracks.Add(FMidiTrack(TEXT("conductor")));
 		MidiData->Tracks[0].AddEvent(FMidiEvent(0, FMidiMsg(static_cast<uint8>(TimeSigNumerator), static_cast<uint8>(TimeSigDenominator))));
-		BarMap.AddTimeSignatureAtBarIncludingCountIn(0, TimeSigNumerator, TimeSigNumerator);
+		MidiData->SongMaps.AddTimeSignatureAtBarIncludingCountIn(0, TimeSigNumerator, TimeSigDenominator);
 		const int32 MidiTempo = Harmonix::Midi::Constants::BPMToMidiTempo(Tempo);
 		MidiData->Tracks[0].AddEvent(FMidiEvent(0, FMidiMsg(MidiTempo)));
-		TempoMap.AddTempoInfoPoint(MidiTempo, 0);
+		MidiData->SongMaps.AddTempoInfoPoint(MidiTempo, 0);
 		MidiData->Tracks[0].Sort();
 		MidiData->ConformToLength(std::numeric_limits<int32>::max());
 
-		ClockInput->AttachToMidiResource(MidiData);
-		ClockInput->ResetAndStart(0);
-		ClockInput->AddSpeedChangeToBlock(HarmonixMetasound::FMidiTimestampSpeed{ 0, 0, Speed });
+		ClockInput->AttachToMidiFile(MidiData);
+		ClockInput->SeekTo(0,0);
+		ClockInput->SetSpeed(0, Speed);
+		ClockInput->SetTransportState(0, HarmonixMetasound::EMusicPlayerTransportState::Playing);
 	}
 	
 	void AdvanceClock(
+		bool bNeedsPrepare,
 		const HarmonixMetasound::FMidiClockWriteRef& ClockInput,
-		Metasound::FSampleCount& SampleCount,
-		Metasound::FSampleCount& SampleRemainder,
-		const int32 NumSamples,
-		const Metasound::FSampleRate SampleRate)
+		const int32 NumSamples)
 	{
-		SampleRemainder += NumSamples;
-		constexpr int32 MidiGranularity = 128;
-		while (SampleRemainder >= MidiGranularity)
+		if (bNeedsPrepare)
 		{
-			SampleCount += MidiGranularity;
-			SampleRemainder -= MidiGranularity;
-			const float AdvanceToMs = static_cast<float>(SampleCount) * 1000.0f / SampleRate;
-			ClockInput->AdvanceHiResToMs(0, AdvanceToMs, true);
+			ClockInput->PrepareBlock();
 		}
+		ClockInput->Advance(0, NumSamples);
 	}
-	
+
 	IMPLEMENT_SIMPLE_AUTOMATION_TEST(
 		FMidiClockVertexAnalyzerTestBasic,
 		"Harmonix.Metasound.Analysis.MidiClockVertexAnalyzer.Basic",
@@ -154,8 +146,6 @@ namespace HarmonixMetasoundTests::MidiClockVertexAnalyzer
 
 		// Render some blocks and make sure we're advancing at the expected rate
 		constexpr int32 NumBlocks = 20;
-		Metasound::FSampleCount SampleCount = 0;
-		Metasound::FSampleCount SampleRemainder = 0;
 
 		for(int32 i = 0; i < NumBlocks; ++i)
 		{
@@ -164,8 +154,8 @@ namespace HarmonixMetasoundTests::MidiClockVertexAnalyzer
 			ReceivedTimestamp.Reset();
 			
 			// Advance the clock
-			AdvanceClock(*ClockRef, SampleCount, SampleRemainder, NumSamplesPerBlock, SampleRate);
-			const FMusicTimestamp ExpectedTimestamp = (*ClockRef)->GetCurrentMusicTimestamp();
+			AdvanceClock(i != 0, *ClockRef, NumSamplesPerBlock);
+			const FMusicTimestamp ExpectedTimestamp = (*ClockRef)->GetMusicTimestampAtBlockEnd();
 
 			// Render a block
 			TArray<float> Buffer;
diff --git a/Engine/Plugins/Runtime/Harmonix/Source/HarmonixMetasoundTests/Private/Analysis/MidiStreamVertexAnalyzerTests.cpp b/Engine/Plugins/Runtime/Harmonix/Source/HarmonixMetasoundTests/Private/Analysis/MidiStreamVertexAnalyzerTests.cpp
--- a/Engine/Plugins/Runtime/Harmonix/Source/HarmonixMetasoundTests/Private/Analysis/MidiStreamVertexAnalyzerTests.cpp
+++ b/Engine/Plugins/Runtime/Harmonix/Source/HarmonixMetasoundTests/Private/Analysis/MidiStreamVertexAnalyzerTests.cpp
@@ -54,15 +54,15 @@ namespace HarmonixMetasoundTests::MidiStreamVertexAnalyzer
 		check(MidiData);
 		MidiData->Tracks.Add(FMidiTrack(TEXT("conductor")));
 		MidiData->Tracks[0].AddEvent(FMidiEvent(0, FMidiMsg(static_cast<uint8>(TimeSigNum), static_cast<uint8>(TimeSigDenom))));
-		MidiData->SongMaps.GetBarMap().AddTimeSignatureAtBarIncludingCountIn(0, TimeSigNum, TimeSigDenom);
+		MidiData->SongMaps.AddTimeSignatureAtBarIncludingCountIn(0, TimeSigNum, TimeSigDenom);
 		const int32 MidiTempo = Harmonix::Midi::Constants::BPMToMidiTempo(Tempo);
 		MidiData->Tracks[0].AddEvent(FMidiEvent(0, FMidiMsg(MidiTempo)));
-		MidiData->SongMaps.GetTempoMap().AddTempoInfoPoint(MidiTempo, 0);
+		MidiData->SongMaps.AddTempoInfoPoint(MidiTempo, 0);
 		MidiData->Tracks[0].Sort();
 		MidiData->ConformToLength(std::numeric_limits<int32>::max());
 
 		// Attach the maps
-		Clock->AttachToMidiResource(MidiData);
+		Clock->AttachToMidiFile(MidiData);
 	}
 	
 	IMPLEMENT_SIMPLE_AUTOMATION_TEST(
@@ -103,7 +103,7 @@ namespace HarmonixMetasoundTests::MidiStreamVertexAnalyzer
 		constexpr float Tempo = 87;
 		HarmonixMetasound::FMidiClockWriteRef Clock = HarmonixMetasound::FMidiClockWriteRef::CreateNew(Generator->OperatorSettings);
 		(*MidiIn)->SetClock(*Clock);
-		
+
 		bool CallbackSuccess = false;
 
 		constexpr int32 EventIntervalTicks = Harmonix::Midi::Constants::GTicksPerQuarterNoteInt;
@@ -135,7 +135,7 @@ namespace HarmonixMetasoundTests::MidiStreamVertexAnalyzer
 				&& TestEqual("Correct velocity", EventInfo.GetVelocity(), NextExpectedVelocity);
 
 				// Check the timestamp
-				const int32 Tick = Clock->GetBarMap().MusicTimestampToTick(EventInfo.Timestamp);
+				const int32 Tick = Clock->GetSongMapEvaluator().MusicTimestampToTick(EventInfo.Timestamp);
 				CallbackSuccess = TestEqual("Timestamp matches", Tick, NextExpectedEventTicks);
 
 				++NextExpectedNoteNumber;
@@ -175,6 +175,15 @@ namespace HarmonixMetasoundTests::MidiStreamVertexAnalyzer
 			
 			// Advance the clock
 			{
+				if (i == 0)
+				{
+					Clock->SetTransportState(0, HarmonixMetasound::EMusicPlayerTransportState::Playing);
+				}
+				else
+				{
+					Clock->PrepareBlock();
+				}
+				int32 BlockOffset = SampleRemainder;
 				SampleRemainder += NumSamplesPerBlock;
 				constexpr int32 MidiGranularity = 128;
 				while (SampleRemainder >= MidiGranularity)
@@ -182,7 +191,8 @@ namespace HarmonixMetasoundTests::MidiStreamVertexAnalyzer
 					SampleCount += MidiGranularity;
 					SampleRemainder -= MidiGranularity;
 					const float AdvanceToMs = static_cast<float>(SampleCount) * 1000.0f / SampleRate;
-					Clock->AdvanceHiResToMs(0, AdvanceToMs, true);
+					Clock->AdvanceToMs(BlockOffset, AdvanceToMs);
+					BlockOffset += MidiGranularity;
 				}
 			}
 			
