Massive Midi Clock Refactor

diff --git a/Engine/Plugins/Runtime/Harmonix/Source/HarmonixMetasoundTests/Private/DataTypes/MidiClock.spec.cpp b/Engine/Plugins/Runtime/Harmonix/Source/HarmonixMetasoundTests/Private/DataTypes/MidiClock.spec.cpp
--- a/Engine/Plugins/Runtime/Harmonix/Source/HarmonixMetasoundTests/Private/DataTypes/MidiClock.spec.cpp
+++ b/Engine/Plugins/Runtime/Harmonix/Source/HarmonixMetasoundTests/Private/DataTypes/MidiClock.spec.cpp
@@ -15,36 +15,36 @@ namespace HarmonixMetasoundTests::MidiClock
 		EAutomationTestFlags::EditorContext | EAutomationTestFlags::EngineFilter)
 
 	TUniquePtr<FMidiClock> TestClock;
-	TUniquePtr<FMidiClock> DrivingClock;
+	TSharedPtr<FMidiClock, ESPMode::NotThreadSafe> DrivingClock;
 	TSharedPtr<FMidiFileData> MidiFileData;
 	Metasound::FOperatorSettings OperatorSettings {48000, 100};
 
 	void AddStateAtFrame(EMusicPlayerTransportState State, int32 Frame) const
 	{
-		const FMidiTimestampTransportState NewState
-		{
-			Frame,
-			0,
-			State
-		};
-
 		check(TestClock.IsValid());
-		TestClock->AddTransportStateChangeToBlock(NewState);
+		TestClock->SetTransportState(Frame, State);
 	}
 
 	void ExecuteWriteAdvance(int32 StartFrameIndex, int32 EndFrameIndex, float Speed)
 	{
-		TestClock->SeekTo(0, 1, 0);
 		TestClock->PrepareBlock();
+		TestClock->SeekTo(0, 0);
+		TestClock->SetTransportState(0, HarmonixMetasound::EMusicPlayerTransportState::Playing);
 
-		TestEqual("Clock.GetCurrentHiResMs()", TestClock->GetCurrentHiResMs(), 0.0f);
-		TestEqual("Clock.GetCurrentMidiTick()", TestClock->GetCurrentMidiTick(), 0);
-		
-		float OldMs = TestClock->GetCurrentHiResMs();
-		int32 OldTick = TestClock->GetCurrentMidiTick();
-		TestClock->WriteAdvance(StartFrameIndex, EndFrameIndex, Speed);
-		float NewMs = TestClock->GetCurrentHiResMs();
-		int32 NewTick = TestClock->GetCurrentMidiTick();
+		TestEqual("Clock.GetCurrentSongPosMs()", TestClock->GetCurrentSongPosMs(), 0.0f);
+		TestEqual("Clock.GetLastProcessedMidiTick()", TestClock->GetLastProcessedMidiTick(), -1);
+		TestEqual("Clock.GetNextMidiTickToProcess()", TestClock->GetNextMidiTickToProcess(), 0);
+
+		// clear out the transport change message, etc. so that our
+		// tests below can just look for speed and advance events...
+		TestClock->PrepareBlock();
+
+		float OldMs = TestClock->GetCurrentSongPosMs();
+		int32 OldTick = TestClock->GetNextMidiTickToProcess();
+		TestClock->SetSpeed(StartFrameIndex, Speed);
+		TestClock->Advance(StartFrameIndex, EndFrameIndex);
+		float NewMs = TestClock->GetCurrentSongPosMs();
+		int32 NewTick = TestClock->GetNextMidiTickToProcess();
 
 		TestTrue("Non looping clock advanced forward in time", NewMs > OldMs);
 		TestTrue("Non looping Clock Advanced forward in ticks", NewTick > OldTick);
@@ -54,26 +54,54 @@ namespace HarmonixMetasoundTests::MidiClock
 		float Ms = OldMs;
 
 		int32 BlockFrameIndex = 0; 
-		for (const FMidiClockEvent& Event : TestClock->GetMidiClockEventsInBlock())
+		int32 FromTick = (int32)(TestClock->GetSongMapEvaluator().MsToTick(Ms) + 0.5f);
+		for (int32 i = 0; i < TestClock->GetMidiClockEventsInBlock().Num(); ++i)
 		{
-			int32 Tick1 = (int32)(TestClock->GetTempoMap().MsToTick(Ms) + 0.5f);
-			Ms += DeltaMs;
-			int32 Tick2 = (int32)(TestClock->GetTempoMap().MsToTick(Ms) + 0.5f);
-
-			TestEqual(FString::Printf(TEXT("Frame-%d: Event.BlockFrameIndex"), Event.BlockFrameIndex), Event.BlockFrameIndex, BlockFrameIndex);
-			TestTrue(FString::Printf(TEXT("Frame-%d: Event.Type"), Event.BlockFrameIndex), Event.Msg.IsType<MidiClockMessageTypes::FAdvanceThru>());
-			TestEqual(FString::Printf(TEXT("Frame-%d: Event.FromTick"), Event.BlockFrameIndex), Event.Msg.Get<MidiClockMessageTypes::FAdvanceThru>().FromTick, Tick1);
-			TestEqual(FString::Printf(TEXT("Frame-%d: Event.ThruTick"), Event.BlockFrameIndex), Event.Msg.Get<MidiClockMessageTypes::FAdvanceThru>().ThruTick, Tick2);
-			TestEqual(FString::Printf(TEXT("Frame-%d: Event.IsPreRoll"), Event.BlockFrameIndex), Event.Msg.Get<MidiClockMessageTypes::FAdvanceThru>().IsPreRoll, false);
+			const FMidiClockEvent& Event = TestClock->GetMidiClockEventsInBlock()[i];
+
+			if (i == 0 && Speed != 1.0f)
+			{
+				// the first message should be the speed message UNLESS the speed was 1.0...
+				TestEqual(FString::Printf(TEXT("Frame-%d: Event.BlockFrameIndex"), Event.BlockFrameIndex), Event.BlockFrameIndex, BlockFrameIndex);
+				TestTrue(FString::Printf(TEXT("Frame-%d: Event.Type"), Event.BlockFrameIndex), Event.Msg.IsType<MidiClockMessageTypes::FSpeedChange>());
+				TestEqual(FString::Printf(TEXT("Frame-%d: Event.Speed"), Event.BlockFrameIndex), Event.Msg.Get<MidiClockMessageTypes::FSpeedChange>().Speed, Speed);
+			}
+			else
+			{
+				// all other messages should be advance messages...
+				Ms += DeltaMs;
+				int32 UpToTick = FMath::RoundToInt32(TestClock->GetSongMapEvaluator().MsToTick(Ms));
+
+				TestEqual(FString::Printf(TEXT("Frame-%d: Event.BlockFrameIndex"), Event.BlockFrameIndex), Event.BlockFrameIndex, BlockFrameIndex);
+				TestTrue(FString::Printf(TEXT("Frame-%d: Event.Type"), Event.BlockFrameIndex), Event.Msg.IsType<MidiClockMessageTypes::FAdvance>());
+				TestEqual(FString::Printf(TEXT("Frame-%d: Event.FirstTickToProcess"), Event.BlockFrameIndex), Event.Msg.Get<MidiClockMessageTypes::FAdvance>().FirstTickToProcess, FromTick);
+				int32 NumberOfTicksToProcess = Event.Msg.Get<MidiClockMessageTypes::FAdvance>().NumberOfTicksToProcess;
+				TestEqual(FString::Printf(TEXT("Frame-%d: Event.NumberOfTicksToProcess"), Event.BlockFrameIndex), Event.Msg.Get<MidiClockMessageTypes::FAdvance>().NumberOfTicksToProcess, UpToTick - FromTick);
+				TestEqual(FString::Printf(TEXT("Frame-%d: Event.LastTickToProcess"), Event.BlockFrameIndex), Event.Msg.Get<MidiClockMessageTypes::FAdvance>().LastTickToProcess(), UpToTick - 1);
 			
-			BlockFrameIndex += FMidiClock::kMidiGranularity;
+				BlockFrameIndex += FMidiClock::kMidiGranularity;
+				FromTick = UpToTick;
+			}
 		}
 
-		int32 Tick = (int32)(TestClock->GetTempoMap().MsToTick(Ms) + 0.5f);
-		TestEqual("Clock.GetCurrentMidiTick()", TestClock->GetCurrentMidiTick(), Tick);
-		TestEqual("Clock.GetCurrentHiResMs()", TestClock->GetCurrentHiResMs(), Ms);
+		int32 Tick = FMath::RoundToInt32(TestClock->GetSongMapEvaluator().MsToTick(Ms));
+		TestEqual("Clock.GetNextMidiTickToProcess()", TestClock->GetLastProcessedMidiTick(), Tick - 1);
+		TestEqual("Clock.GetCurrentSongPosMs()", TestClock->GetCurrentSongPosMs(), Ms, DeltaMs / 2.0f);
 	};
 	
+	int32 GetLastTransportStateChangeBlockSample(FMidiClock* InClock)
+	{
+		const TArray<FMidiClockEvent>& ClockEvents = InClock->GetMidiClockEventsInBlock();
+		for (auto it = ClockEvents.rbegin(); it != ClockEvents.rend(); ++it)
+		{
+			if ((*it).IsType<MidiClockMessageTypes::FTransportChange>())
+			{
+				return (*it).BlockFrameIndex;
+			}
+		}
+		return -1;
+	}
+
 	END_DEFINE_SPEC(FHarmonixMetasoundMidiClockSpec)
 
 	void FHarmonixMetasoundMidiClockSpec::Define()
@@ -94,60 +122,44 @@ namespace HarmonixMetasoundTests::MidiClock
 			{
 				TestClock->PrepareBlock();
 
-				TestFalse("No transport changes in block", TestClock->HasTransportStateChangeInBlock());
+				TestFalse("No transport changes in block", TestClock->HasTransportStateChangesInBlock());
 
 				constexpr EMusicPlayerTransportState NewState = EMusicPlayerTransportState::Playing;
 				AddStateAtFrame(NewState, 0);
 
-				TestTrue("There is a transport state change in block", TestClock->HasTransportStateChangeInBlock());
+				TestTrue("There is a transport state change in block", TestClock->HasTransportStateChangesInBlock());
 				TestEqual("State at end of block matches the one we added", TestClock->GetTransportStateAtEndOfBlock(), NewState);
 			});
 
 			It("should add NewState if its frame is greater than the last one in the block", [this]()
 			{
 				AddStateAtFrame(EMusicPlayerTransportState::Playing, 0);
-				TestTrue("There is already a transport state change in the block", TestClock->HasTransportStateChangeInBlock());
+				TestTrue("There is already a transport state change in the block", TestClock->HasTransportStateChangesInBlock());
 
-				const int32 NumInitialStates = TestClock->GetTransportTimestampsInBlock().Num();
-				const int32 LastStateFrame = TestClock->GetTransportTimestampsInBlock().Last().BlockSampleFrameIndex;
+				const int32 NumInitialStates = TestClock->GetNumTransportStateChangesInBlock();
+				const int32 LastStateFrame = GetLastTransportStateChangeBlockSample(TestClock.Get());
 
 				constexpr EMusicPlayerTransportState NewState = EMusicPlayerTransportState::Pausing;
 				AddStateAtFrame(NewState, LastStateFrame + 1);
 
-				TestEqual("There is another transport state change in block", TestClock->GetTransportTimestampsInBlock().Num(), NumInitialStates + 1);
+				TestEqual("There is another transport state change in block", TestClock->GetNumTransportStateChangesInBlock(), NumInitialStates + 1);
 				TestEqual("State at end of block matches the one we added", TestClock->GetTransportStateAtEndOfBlock(), NewState);
 			});
 
 			It("should add NewState if it has the same frame as the last one in the block", [this]()
 			{
 				AddStateAtFrame(EMusicPlayerTransportState::Playing, 0);
-				TestTrue("There is already a transport state change in the block", TestClock->HasTransportStateChangeInBlock());
+				TestTrue("There is already a transport state change in the block", TestClock->HasTransportStateChangesInBlock());
 
-				const int32 NumInitialStates = TestClock->GetTransportTimestampsInBlock().Num();
-				const int32 LastStateFrame = TestClock->GetTransportTimestampsInBlock().Last().BlockSampleFrameIndex;
+				const int32 NumInitialStates = TestClock->GetNumTransportStateChangesInBlock();
+				const int32 LastStateFrame = GetLastTransportStateChangeBlockSample(TestClock.Get());
 
 				constexpr EMusicPlayerTransportState NewState = EMusicPlayerTransportState::Paused;
 				AddStateAtFrame(NewState, LastStateFrame );
 
-				TestEqual("There is another transport state change in block", TestClock->GetTransportTimestampsInBlock().Num(), NumInitialStates + 1);
+				TestEqual("There is another transport state change in block", TestClock->GetNumTransportStateChangesInBlock(), NumInitialStates + 1);
 				TestEqual("State at end of block matches the one we added", TestClock->GetTransportStateAtEndOfBlock(), NewState);
 			});
-
-			It("should not add NewState if the frame is less than the last one in the block", [this]()
-			{
-				AddStateAtFrame(EMusicPlayerTransportState::Playing, 0);
-				TestTrue("There is already a transport state change in the block", TestClock->HasTransportStateChangeInBlock());
-
-				const int32 NumInitialStates = TestClock->GetTransportTimestampsInBlock().Num();
-				const EMusicPlayerTransportState LastState = TestClock->GetTransportStateAtEndOfBlock();
-				const int32 LastStateFrame = TestClock->GetTransportTimestampsInBlock().Last().BlockSampleFrameIndex;
-
-				constexpr EMusicPlayerTransportState NewState = EMusicPlayerTransportState::Continuing;
-				AddStateAtFrame(NewState, LastStateFrame - 1);
-
-				TestEqual("There is not another transport state change in block", TestClock->GetTransportTimestampsInBlock().Num(), NumInitialStates);
-				TestEqual("State at end of block matches the one that was already there", TestClock->GetTransportStateAtEndOfBlock(), LastState);
-			});
 		});
 
 		Describe("PrepareBlock()", [this]()
@@ -157,12 +169,12 @@ namespace HarmonixMetasoundTests::MidiClock
 				const float LastTempo = TestClock->GetTempoAtEndOfBlock();
 				const float LastSpeed = TestClock->GetSpeedAtEndOfBlock();
 				TestClock->PrepareBlock();
-				TestEqual("Clock.TransportChangesInBlock.Num()", TestClock->GetTransportTimestampsInBlock().Num(), 0);
-				TestEqual("Clock.SpeedChangesInBlock.Num()", TestClock->GetTempoSpeedTimestampsInBlock().Num(), 1);
+				TestEqual("Clock.TransportChangesInBlock.Num()", TestClock->GetNumTransportStateChangesInBlock(), 0);
+				TestEqual("Clock.SpeedChangesInBlock.Num()", TestClock->GetNumSpeedChangesInBlock(), 0);
 				TestEqual("Clock.SpeedAtBlockSampleFrame(0)", TestClock->GetSpeedAtBlockSampleFrame(0), LastSpeed);
 				TestEqual("Clock.SpeedAtEndOfBlock()", TestClock->GetSpeedAtEndOfBlock(), LastSpeed);
 				TestFalse("Clock.HasSpeedChangesInBlock()", TestClock->HasSpeedChangesInBlock());
-				TestEqual("Clock.TempoChangesInBlock.Num()", TestClock->GetNumTempoChangesInBlock(), 1);
+				TestEqual("Clock.TempoChangesInBlock.Num()", TestClock->GetNumTempoChangesInBlock(), 0);
 				TestEqual("Clock.TempoAtBlockSampleFrame(0)", TestClock->GetTempoAtBlockSampleFrame(0), LastTempo);
 				TestEqual("Clock.GetTempoAtEndOfBlock()", TestClock->GetTempoAtEndOfBlock(), LastTempo);
 				TestFalse("Clock.HasTempoChangesInBlock()", TestClock->HasTempoChangesInBlock());
@@ -174,63 +186,58 @@ namespace HarmonixMetasoundTests::MidiClock
 		{
 			It("should reset speed and tempo changes, and be \"playing\"", [this]()
 			{
-				const float LastTempo = TestClock->GetTempoAtEndOfBlock();
-				const float LastSpeed = TestClock->GetSpeedAtEndOfBlock();
-				TestClock->ResetAndStart(0, false);
+				TestClock->SetTransportState(0, HarmonixMetasound::EMusicPlayerTransportState::Playing);
 				TestEqual("Clock.GetTransportStateAtEndOfBlock()", TestClock->GetTransportStateAtEndOfBlock(), EMusicPlayerTransportState::Playing);
-				TestEqual("Clock.SpeedAtBlockSampleFrame(0)", TestClock->GetSpeedAtBlockSampleFrame(0), LastSpeed);
-				TestEqual("Clock.SpeedAtEndOfBlock()", TestClock->GetSpeedAtEndOfBlock(), LastSpeed);
-				TestFalse("Clock.HasSpeedChangesInBlock()", TestClock->HasSpeedChangesInBlock());
+				TestEqual("Clock.SpeedAtBlockSampleFrame(0)", TestClock->GetSpeedAtBlockSampleFrame(0), 1.0f);
+				TestEqual("Clock.SpeedAtEndOfBlock()", TestClock->GetSpeedAtEndOfBlock(), 1.0f);
+				TestTrue("Clock.HasSpeedChangesInBlock()", TestClock->HasSpeedChangesInBlock());
 				TestEqual("Clock.TempoChangesInBlock.Num()", TestClock->GetNumTempoChangesInBlock(), 1);
-				TestEqual("Clock.TempoAtBlockSampleFrame(0)", TestClock->GetTempoAtBlockSampleFrame(0), LastTempo);
-				TestEqual("Clock.GetTempoAtEndOfBlock()", TestClock->GetTempoAtEndOfBlock(), LastTempo);
-				TestFalse("Clock.HasTempoChangesInBlock()", TestClock->HasTempoChangesInBlock());
+				TestEqual("Clock.TempoAtBlockSampleFrame(0)", TestClock->GetTempoAtBlockSampleFrame(0), 120.0f);
+				TestEqual("Clock.GetTempoAtEndOfBlock()", TestClock->GetTempoAtEndOfBlock(), 120.0f);
+				TestTrue("Clock.HasTempoChangesInBlock()", TestClock->HasTempoChangesInBlock());
 			});
-			
-			It("should be at the correct block frame index", [this]()
-			{
-				const float LastTempo = TestClock->GetTempoAtEndOfBlock();
-				const float LastSpeed = TestClock->GetSpeedAtEndOfBlock();
-				int32 BlockFrame = 100;
-				TestClock->ResetAndStart(BlockFrame, false);
-				TestEqual("Clock.Get", TestClock->GetCurrentBlockFrameIndex(), BlockFrame);
-			});
-			
-			It("should seek to start when requested", [this]()
+
+/*
+			It("should seek to 0 when reset requested", [this]()
 			{
 				int32 BlockFrame = 100;
 				FMusicSeekTarget SeekTarget;
 				SeekTarget.Type = ESeekPointType::BarBeat;
 				SeekTarget.BarBeat = FMusicTimestamp(2, 1.0f);
-				TestClock->SeekTo(SeekTarget, 0);
-				TestFalse("Clock seeked, Clock.GetCurrentMidiTick() == 0", TestClock->GetCurrentMidiTick() == 0);
+				TestClock->SeekTo(0, SeekTarget);
+				TestTrue("Clock seeked, Clock.GetNextMidiTickToProcess() != 0", TestClock->GetNextMidiTickToProcess() != 0);
 				
-				TestClock->ResetAndStart(BlockFrame, true);
-				TestEqual("Clock.GetCurrentBlockFrameIndex()", TestClock->GetCurrentBlockFrameIndex(), BlockFrame);
-				TestEqual("Clock.GetCurrentMidiTick()", TestClock->GetCurrentMidiTick(), -1);
+				TestClock->Seek ResetAndStart(BlockFrame, true);
+				TestEqual("Clock.GetLastProcessedMidiTick()", TestClock->GetLastProcessedMidiTick(), -1);
+				TestEqual("Clock.GetNextMidiTickToProcess()", TestClock->GetNextMidiTickToProcess(), 0);
 			});
+*/
 
-			It("should otherwise not seek to start", [this]()
+/*
+			It("should NOT seek to 0 when reset requested", [this]()
 			{
 				int32 BlockFrame = 100;
 				FMusicSeekTarget SeekTarget;
 				SeekTarget.Type = ESeekPointType::BarBeat;
 				SeekTarget.BarBeat = FMusicTimestamp(2, 1.0f);
-				TestClock->SeekTo(SeekTarget, 0);
-				int32 NewMidiTick = TestClock->GetCurrentMidiTick();
-				TestFalse("Clock seeked, Clock.GetCurrentMidiTick() == 0", TestClock->GetCurrentMidiTick() == 0);
+				TestClock->SeekTo(0, SeekTarget);
+				int32 NewMidiTick = TestClock->GetNextMidiTickToProcess();
+				TestTrue("Clock seeked, Clock.GetNextMidiTickToProcess() != 0", TestClock->GetNextMidiTickToProcess() != 0);
 							
 				TestClock->ResetAndStart(BlockFrame, false);
-				TestEqual("Clock.GetCurrentBlockFrameIndex()", TestClock->GetCurrentBlockFrameIndex(), BlockFrame);
-				TestEqual("Clock.GetCurrentMidiTick()", TestClock->GetCurrentMidiTick(), NewMidiTick);
+				TestEqual("After reset, Clock.GetNextMidiTickToProcess()", TestClock->GetNextMidiTickToProcess(), NewMidiTick);
 			});
-			
+*/
+
 		});
 
 		Describe("SetLoop()", [this]()
 		{
 			It("should correctly set the tempo", [this]()
 			{
+				// set transport to playing so the block gets the initial tempo, time signature, etc...
+				TestClock->SetTransportState(0, HarmonixMetasound::EMusicPlayerTransportState::Playing);
+
 				const float TicksPerQuarterNote = Harmonix::Midi::Constants::GTicksPerQuarterNote;
 				const float TempoBPM = 120.0f;
 				int32 LoopStartTick = 0;
@@ -238,7 +245,7 @@ namespace HarmonixMetasoundTests::MidiClock
 
 				// initial clock tempo
 				TestEqual("Clock.Tempo", TestClock->GetTempoAtEndOfBlock(), TempoBPM);
-				TestEqual("Clock.Tempo", TestClock->GetTempoMap().GetTempoAtTick(0), TempoBPM);
+				TestEqual("Clock.Tempo", TestClock->GetSongMapEvaluator().GetTempoAtTick(0), TempoBPM);
 
 				// microseconds per quarternote
 				int32 MidiTempo = Harmonix::Midi::Constants::BPMToMidiTempo(TempoBPM);
@@ -247,16 +254,16 @@ namespace HarmonixMetasoundTests::MidiClock
 				float LoopStartMs = MsPerTick * LoopStartTick;
 				float LoopEndMs = MsPerTick * LoopEndTick;
 
-				TestFalse("Initial -> Clock.DoesLoop()", TestClock->DoesLoop());
-				TestClock->SetLoop(LoopStartTick, LoopEndTick);
-				TestTrue("SetLoop -> Clock.DoesLoop()", TestClock->DoesLoop());
-				TestEqual("Clock.LoopStartTick()", TestClock->GetLoopStartTick(), LoopStartTick);
-				TestEqual("Clock.LoopEndTick()", TestClock->GetLoopEndTick(), LoopEndTick);
+				TestFalse("Initial -> Clock.HasPersistentLoop()", TestClock->HasPersistentLoop());
+				TestClock->SetupPersistentLoop(LoopStartTick, LoopEndTick - LoopStartTick);
+				TestTrue("SetLoop -> Clock.HasPersistentLoop()", TestClock->HasPersistentLoop());
+				TestEqual("Clock.GetFirstTickInLoop()", TestClock->GetFirstTickInLoop(), LoopStartTick);
+				TestEqual("Clock.GetLoopLengthTicks()", TestClock->GetLoopLengthTicks(), LoopEndTick - LoopStartTick);
 				TestEqual("Clock.LoopStartMs()", TestClock->GetLoopStartMs(), LoopStartMs);
 				TestEqual("Clock.LoopEndMs()", TestClock->GetLoopEndMs(), LoopEndMs);
 
-				TestClock->ClearLoop();
-				TestFalse("Cleared -> Clock.DoesLoop()", TestClock->DoesLoop());
+				TestClock->ClearPersistentLoop();
+				TestFalse("Cleared -> Clock.HasPersistentLoop()", TestClock->HasPersistentLoop());
 			});
 
 		});
@@ -293,118 +300,124 @@ namespace HarmonixMetasoundTests::MidiClock
 		{
 			BeforeEach([&, this]
 			{
-				DrivingClock = MakeUnique<FMidiClock>(OperatorSettings);
-				TestClock->AttachToTimeAuthority(*DrivingClock);
+				DrivingClock = MakeShared<FMidiClock, ESPMode::NotThreadSafe>(OperatorSettings);
+				TestClock->SetDrivingClock(DrivingClock);
 			});
 
 			AfterEach([&, this]
 			{
-				TestClock->ClearLoop();
-				TestClock->DetachFromTimeAuthority();
+				TestClock->ClearPersistentLoop();
+				TestClock->SetDrivingClock(nullptr);
 				DrivingClock.Reset();
 			});
 
-			It("EventType::AdvanceThru.NonLooping", [&, this]
+			It("EventType::AdvanceThru.NonLooping", [&, this]()
 			{
+				DrivingClock->SeekTo(0,0);
+				DrivingClock->SetTransportState(0, HarmonixMetasound::EMusicPlayerTransportState::Playing);
+
 				float DrivingClockSpeed = 1.0f;
 				int32 StartFrame = 0;
 				int32 DeltaFrames = FMidiClock::kMidiGranularity * DrivingClockSpeed;
 				float DeltaMs = DeltaFrames * 1000.0f / OperatorSettings.GetSampleRate();
-				int32 Tick = (int32)(TestClock->GetTempoMap().MsToTick(DeltaMs) + 0.5f);
+				int32 Tick = (int32)(TestClock->GetSongMapEvaluator().MsToTick(DeltaMs) + 0.5f);
 
-				float Speed = 1.0f;
-				int32 PrerollBars = 8;
-				FMidiClockEvent Event = FMidiClockEvent(StartFrame, MidiClockMessageTypes::FAdvanceThru(0, Tick, false));
+				DrivingClock->AdvanceToTick(0, Tick);
 
 				TestClock->PrepareBlock();
 				int32 OldEventsNum = TestClock->GetMidiClockEventsInBlock().Num();
-				TestClock->HandleClockEvent(*DrivingClock, Event, PrerollBars, Speed);
+				TestClock->Advance(*DrivingClock, 0, OperatorSettings.GetNumFramesPerBlock());
 
 				if (!TestTrue("Clock has new clock events", TestClock->GetMidiClockEventsInBlock().Num() > OldEventsNum))
 				{
 					return;
 				}
-				TestTrue("Last Clock Event in block", TestClock->GetMidiClockEventsInBlock().Last().Msg.IsType<MidiClockMessageTypes::FAdvanceThru>());
-				TestEqual("Clock Current Tick", TestClock->GetCurrentMidiTick(), Event.Msg.Get<MidiClockMessageTypes::FAdvanceThru>().ThruTick - 1);
+				TestTrue("Last Clock Event in block", TestClock->GetMidiClockEventsInBlock().Last().Msg.IsType<MidiClockMessageTypes::FAdvance>());
+				const FMidiClockEvent& Event = TestClock->GetMidiClockEventsInBlock().Last();
+				TestEqual("Clock Last Processed Tick", TestClock->GetLastProcessedMidiTick(), Event.Msg.Get<MidiClockMessageTypes::FAdvance>().LastTickToProcess());
+				TestEqual("Clock Next Tick To Process", TestClock->GetNextMidiTickToProcess(), Event.Msg.Get<MidiClockMessageTypes::FAdvance>().LastTickToProcess() + 1);
 			});
 
-			It("EventType::AdvanceThru.Looping", [&, this]
+			It("EventType::AdvanceThru.Looping", [&, this]()
 			{
+				TestClock->SetTransportState(0, HarmonixMetasound::EMusicPlayerTransportState::Playing);
+				DrivingClock->SetTransportState(0, HarmonixMetasound::EMusicPlayerTransportState::Playing);
+
 				FMusicTimestamp LoopEndTimestamp;
 				LoopEndTimestamp.Bar = 2;
 				LoopEndTimestamp.Beat = 1.0f;
-				
 				int32 LoopStartTick = 0;
-				int32 LoopEndTick = TestClock->GetBarMap().MusicTimestampToTick(LoopEndTimestamp);
-				TestClock->SetLoop(LoopStartTick, LoopEndTick);
-				TestClock->SeekTo(0, LoopEndTick - 1, 0);
+				int32 LoopEndTick = TestClock->GetSongMapEvaluator().MusicTimestampToTick(LoopEndTimestamp);
 				
-				float DrivingClockSpeed = 1.0f;
-				int32 StartFrame = 0;
-				int32 DeltaFrames = FMidiClock::kMidiGranularity * DrivingClockSpeed;
-				float DeltaMs = DeltaFrames * 1000.0f / OperatorSettings.GetSampleRate();
-				int32 DeltaTicks = (int32)(TestClock->GetTempoMap().MsToTick(DeltaMs) + 0.5f);
+				TestClock->SetupPersistentLoop(LoopStartTick, LoopEndTick - LoopStartTick);
+
+				TestClock->PrepareBlock();
+				DrivingClock->PrepareBlock();
+
+				int32 TestStartTick = LoopEndTick - 1;;
+
+				DrivingClock->SeekTo(0, TestStartTick);
+
+				float DeltaMs = FMidiClock::kMidiGranularity * 1000.0f / OperatorSettings.GetSampleRate();
+				int32 DeltaTicks = (int32)(DrivingClock->GetSongMapEvaluator().MsToTick(DeltaMs) + 0.5f);
 				int32 Tick = DeltaTicks;
 
-				int32 ExpectedTick = TestClock->GetCurrentMidiTick() + DeltaTicks + 1;
+
+				int32 ExpectedTick = TestStartTick + DeltaTicks;
 				if (ExpectedTick > LoopEndTick)
 				{
 					ExpectedTick = ExpectedTick - LoopEndTick + LoopStartTick;
 				}
 
-				float Speed = 1.0f;
-				int32 PrerollBars = 8;
-				FMidiClockEvent Event = FMidiClockEvent(StartFrame, MidiClockMessageTypes::FAdvanceThru(0, Tick, false));
+				DrivingClock->Advance(0, FMidiClock::kMidiGranularity);
+				TestClock->Advance(*DrivingClock, 0, FMidiClock::kMidiGranularity);
 
-				TestClock->PrepareBlock();
-				int32 OldEventsNum = TestClock->GetMidiClockEventsInBlock().Num();
-				TestClock->HandleClockEvent(*DrivingClock, Event, PrerollBars, Speed);
-
-				if (!TestTrue("Clock has new clock events", TestClock->GetMidiClockEventsInBlock().Num() > OldEventsNum))
+				if (!TestTrue("Clock has new clock events", TestClock->GetMidiClockEventsInBlock().Num() > 0))
 				{
 					return;
 				}
-				TestTrue("Last Clock Event in block", TestClock->GetMidiClockEventsInBlock().Last().Msg.IsType<MidiClockMessageTypes::FAdvanceThru>());
-				TestEqual("Clock Current Tick", TestClock->GetCurrentMidiTick(), ExpectedTick - 1);
+				TestTrue("Last Clock Event in block is advance", TestClock->GetMidiClockEventsInBlock().Last().Msg.IsType<MidiClockMessageTypes::FAdvance>());
+				TestEqual("Clock Current Tick", TestClock->GetNextMidiTickToProcess(), ExpectedTick);
 			});
 
-			It("ProccessClockEvent(SeekTo)", [&, this]
+
+			It("ProccessClockEvent(SeekTo)", [&, this]()
 			{
 				int32 StartFrame = 0;
 				int32 Tick = 1000;
-				float Speed = 1.0f;
-				int32 PrerollBars = 8;
-				FMidiClockEvent Event = FMidiClockEvent(StartFrame, MidiClockMessageTypes::FSeekTo(0, Tick));
 
 				TestClock->PrepareBlock();
 				int32 OldEventsNum = TestClock->GetMidiClockEventsInBlock().Num();
-				TestClock->HandleClockEvent(*DrivingClock, Event, PrerollBars, Speed);
+
+				TestClock->SeekTo(StartFrame, 1000);
 
 				if (!TestTrue("Clock has new clock events", TestClock->GetMidiClockEventsInBlock().Num() > OldEventsNum))
 				{
 					return;
 				}
-				TestTrue("Last Clock Event in block", TestClock->GetMidiClockEventsInBlock()[0].Msg.IsType<MidiClockMessageTypes::FSeekThru>());
-				TestEqual("Clock Current Tick", TestClock->GetCurrentMidiTick(), Tick - 1);
+				TestTrue("Last Clock Event in block", TestClock->GetMidiClockEventsInBlock()[0].Msg.IsType<MidiClockMessageTypes::FSeek>());
+				TestEqual("Clock Current Tick", TestClock->GetNextMidiTickToProcess(), Tick);
 			});
-			
 		});
 
 		Describe("Tempo Changes", [this]()
 		{
-			It("Without driving clock", [this]()
+			It("Without driving clock - One Tempo Change At Span End", [this]()
 			{
 				MidiFileData = FMidiClock::MakeClockConductorMidiData(123, 5, 8);
 				constexpr int32 TempoChangeTick = 234;
 				constexpr float TempoChangeTempo = 89;
 				MidiFileData->AddTempoChange(0, TempoChangeTick, TempoChangeTempo);
-				TestClock->AttachToMidiResource(MidiFileData);
-				TestClock->ResetAndStart(0);
-				
-				while (TestClock->GetCurrentMidiTick() < TempoChangeTick)
+				TestClock->AttachToMidiFile(MidiFileData);
+				TestClock->SetTransportState(0, HarmonixMetasound::EMusicPlayerTransportState::Playing);
+
+				while (TestClock->GetLastProcessedMidiTick() < TempoChangeTick)
 				{
-					TestClock->PrepareBlock();
-					TestClock->WriteAdvance(0, OperatorSettings.GetNumFramesPerBlock());
+					if (TestClock->GetLastProcessedMidiTick() > 0)
+					{
+						TestClock->PrepareBlock();
+					}
+					TestClock->Advance(0, OperatorSettings.GetNumFramesPerBlock());
 				}
 
 				bool HasTempoEvent = false;
@@ -415,7 +428,7 @@ namespace HarmonixMetasoundTests::MidiClock
 					{
 						const MidiClockMessageTypes::FTempoChange& TempoChange = Event.Msg.Get<MidiClockMessageTypes::FTempoChange>();
 
-						if (!TestEqual("Tempo is correct", TempoChange.Tempo, TempoChangeTempo))
+						if (!TestEqual("Tempo is correct", TempoChange.Tempo, TempoChangeTempo, 0.001f))
 						{
 							return;
 						}
@@ -432,23 +445,118 @@ namespace HarmonixMetasoundTests::MidiClock
 
 				TestTrue("Got tempo event", HasTempoEvent);
 			});
+
+			It("Without driving clock - One Tempo Change At Span Start", [this]()
+			{
+				MidiFileData = FMidiClock::MakeClockConductorMidiData(123, 5, 8);
+				constexpr int32 TempoChangeTick = 230;
+				constexpr float TempoChangeTempo = 89;
+				MidiFileData->AddTempoChange(0, TempoChangeTick, TempoChangeTempo);
+				TestClock->AttachToMidiFile(MidiFileData);
+				TestClock->SetTransportState(0, HarmonixMetasound::EMusicPlayerTransportState::Playing);
+
+				while (TestClock->GetLastProcessedMidiTick() < TempoChangeTick)
+				{
+					if (TestClock->GetLastProcessedMidiTick() > 0)
+					{
+						TestClock->PrepareBlock();
+					}
+					TestClock->Advance(0, OperatorSettings.GetNumFramesPerBlock());
+				}
+
+				bool HasTempoEvent = false;
+
+				for (const FMidiClockEvent& Event : TestClock->GetMidiClockEventsInBlock())
+				{
+					if (Event.Msg.IsType<MidiClockMessageTypes::FTempoChange>())
+					{
+						const MidiClockMessageTypes::FTempoChange& TempoChange = Event.Msg.Get<MidiClockMessageTypes::FTempoChange>();
+
+						if (!TestEqual("Tempo is correct", TempoChange.Tempo, TempoChangeTempo, 0.001f))
+						{
+							return;
+						}
+
+						if (!TestEqual("Tick is correct", TempoChange.Tick, TempoChangeTick))
+						{
+							return;
+						}
+
+						HasTempoEvent = true;
+						break;
+					}
+				}
+
+				TestTrue("Got tempo event", HasTempoEvent);
+			});
+
+			It("Without driving clock - Many Tempo Changes In Span", [this]()
+			{
+				MidiFileData = FMidiClock::MakeClockConductorMidiData(123, 5, 8);
+				constexpr int32 NumChanges = 4;
+				constexpr int32 TempoChangeTicks[NumChanges] = { 230, 231, 232, 233 };
+				constexpr float TempoChangeTempos [NumChanges] = { 89.0, 89.2, 89.4, 89.6 };
+				for (int32 i = 0; i < NumChanges; ++i)
+				{
+					MidiFileData->AddTempoChange(0, TempoChangeTicks[i], TempoChangeTempos[i]);
+				}
+				TestClock->AttachToMidiFile(MidiFileData);
+				TestClock->SetTransportState(0, HarmonixMetasound::EMusicPlayerTransportState::Playing);
+
+				while (TestClock->GetLastProcessedMidiTick() < TempoChangeTicks[NumChanges - 1])
+				{
+					if (TestClock->GetLastProcessedMidiTick() > 0)
+					{
+						TestClock->PrepareBlock();
+					}
+					TestClock->Advance(0, OperatorSettings.GetNumFramesPerBlock());
+				}
+
+				int NumTempoEvents = 0;
+
+				for (const FMidiClockEvent& Event : TestClock->GetMidiClockEventsInBlock())
+				{
+					if (Event.Msg.IsType<MidiClockMessageTypes::FTempoChange>())
+					{
+						const MidiClockMessageTypes::FTempoChange& TempoChange = Event.Msg.Get<MidiClockMessageTypes::FTempoChange>();
+
+						if (!TestEqual("Tempo is correct", TempoChange.Tempo, TempoChangeTempos[NumTempoEvents], 0.001f))
+						{
+							return;
+						}
+
+						if (!TestEqual("Tick is correct", TempoChange.Tick, TempoChangeTicks[NumTempoEvents]))
+						{
+							return;
+						}
+
+						NumTempoEvents++;
+					}
+				}
+
+				TestEqual("Got correct number of tempo events", NumTempoEvents, NumChanges);
+			});
+
 		});
 
 		Describe("Time signature Changes", [this]()
 		{
-			It("Without driving clock", [this]()
+			It("Without driving clock - One Change", [this]()
 			{
 				MidiFileData = FMidiClock::MakeClockConductorMidiData(123, 5, 8);
-				const int32 TimeSigChangeTick = MidiFileData->SongMaps.GetBarMap().BarBeatTickIncludingCountInToTick(2, 1, 0);
+				const int32 TimeSigChangeTick = MidiFileData->SongMaps.BarBeatTickIncludingCountInToTick(2, 1, 0);
 				const FTimeSignature NewTimeSig{ 3, 4 };
 				MidiFileData->AddTimeSigChange(0, TimeSigChangeTick, NewTimeSig.Numerator, NewTimeSig.Denominator);
-				TestClock->AttachToMidiResource(MidiFileData);
-				TestClock->ResetAndStart(0);
+				TestClock->AttachToMidiFile(MidiFileData);
+				TestClock->SetTransportState(0, HarmonixMetasound::EMusicPlayerTransportState::Playing);
 						
-				while (TestClock->GetCurrentMidiTick() < TimeSigChangeTick)
+				while (TestClock->GetLastProcessedMidiTick() < TimeSigChangeTick)
 				{
-					TestClock->PrepareBlock();
-					TestClock->WriteAdvance(0, OperatorSettings.GetNumFramesPerBlock());
+					if (TestClock->GetLastProcessedMidiTick() > 0)
+					{
+						TestClock->PrepareBlock();
+					}
+					TestClock->Advance(0, OperatorSettings.GetNumFramesPerBlock());
 				}
 
 				bool HasTimeSignatureEvent = false;
@@ -476,6 +584,73 @@ namespace HarmonixMetasoundTests::MidiClock
 
 				TestTrue("Got time sig event", HasTimeSignatureEvent);
 			});
+
+			It("Without driving clock - One Change with tempos", [this]()
+			{
+				MidiFileData = FMidiClock::MakeClockConductorMidiData(123, 5, 8);
+				const int32 TimeSigChangeTick = MidiFileData->SongMaps.BarBeatTickIncludingCountInToTick(2, 1, 0);
+				constexpr int32 NumTempoChanges = 3;
+				constexpr int32 TempoChangeTicks[NumTempoChanges] = { 4799, 4800, 4801 };
+				constexpr float TempoChangeTempos[NumTempoChanges] = { 155.0f, 157.2f, 158.4f };
+				MidiFileData->AddTempoChange(0, TempoChangeTicks[0], TempoChangeTempos[0]);
+				const FTimeSignature NewTimeSig{ 3, 4 };
+				MidiFileData->AddTimeSigChange(0, TimeSigChangeTick, NewTimeSig.Numerator, NewTimeSig.Denominator);
+				MidiFileData->AddTempoChange(0, TempoChangeTicks[1], TempoChangeTempos[1]);
+				MidiFileData->AddTempoChange(0, TempoChangeTicks[2], TempoChangeTempos[2]);
+				TestClock->AttachToMidiFile(MidiFileData);
+				TestClock->SetTransportState(0, HarmonixMetasound::EMusicPlayerTransportState::Playing);
+
+				while (TestClock->GetLastProcessedMidiTick() < TimeSigChangeTick)
+				{
+					if (TestClock->GetLastProcessedMidiTick() > 0)
+					{
+						TestClock->PrepareBlock();
+					}
+					TestClock->Advance(0, OperatorSettings.GetNumFramesPerBlock());
+				}
+
+				bool HasTimeSignatureEvent = false;
+				int32 NumFoundTempoChanges = 0;
+
+				for (const FMidiClockEvent& Event : TestClock->GetMidiClockEventsInBlock())
+				{
+					if (Event.Msg.IsType<MidiClockMessageTypes::FTimeSignatureChange>())
+					{
+						const MidiClockMessageTypes::FTimeSignatureChange& TimeSigChange = Event.Msg.Get<MidiClockMessageTypes::FTimeSignatureChange>();
+
+						if (!TestEqual("Time signature is correct", TimeSigChange.TimeSignature, NewTimeSig))
+						{
+							return;
+						}
+
+						if (!TestEqual("Tick is correct", TimeSigChange.Tick, TimeSigChangeTick))
+						{
+							return;
+						}
+
+						TestFalse("Already found time signature", HasTimeSignatureEvent);
+
+						HasTimeSignatureEvent = true;
+					}
+					else if (const MidiClockMessageTypes::FTempoChange* AsTempoChange = Event.TryGet<MidiClockMessageTypes::FTempoChange>())
+					{
+						if (!TestEqual("Tempo is correct", AsTempoChange->Tempo, TempoChangeTempos[NumFoundTempoChanges], 0.001f))
+						{
+							return;
+						}
+
+						if (!TestEqual("Tick is correct", AsTempoChange->Tick, TempoChangeTicks[NumFoundTempoChanges]))
+						{
+							return;
+						}
+
+						NumFoundTempoChanges++;
+					}
+				}
+
+				TestTrue("Got time sig event", HasTimeSignatureEvent);
+				TestEqual("Found correct number of tempo changes", NumFoundTempoChanges, NumTempoChanges);
+			});
 		});
 	}
 
diff --git a/Engine/Plugins/Runtime/Harmonix/Source/HarmonixMetasoundTests/Private/MidiOps/PulseGeneratorTests.cpp b/Engine/Plugins/Runtime/Harmonix/Source/HarmonixMetasoundTests/Private/MidiOps/PulseGeneratorTests.cpp
--- a/Engine/Plugins/Runtime/Harmonix/Source/HarmonixMetasoundTests/Private/MidiOps/PulseGeneratorTests.cpp
+++ b/Engine/Plugins/Runtime/Harmonix/Source/HarmonixMetasoundTests/Private/MidiOps/PulseGeneratorTests.cpp
@@ -21,7 +21,7 @@ namespace Harmonix::Midi::Ops::Tests
 		constexpr uint8 TimeSigNumerator = 4;
 		constexpr uint8 TimeSigDenominator = 4;
 		const auto Clock = MakeShared<HarmonixMetasound::FMidiClock, ESPMode::NotThreadSafe>(OperatorSettings);
-		Clock->AttachToMidiResource(HarmonixMetasound::FMidiClock::MakeClockConductorMidiData(Tempo, TimeSigNumerator, TimeSigDenominator));
+		Clock->AttachToMidiFile(HarmonixMetasound::FMidiClock::MakeClockConductorMidiData(Tempo, TimeSigNumerator, TimeSigDenominator));
 		
 		HarmonixMetasound::FMidiStream OutputStream;
 
@@ -32,9 +32,15 @@ namespace Harmonix::Midi::Ops::Tests
 			constexpr int32 NotesUntilWeAreSatisfiedThisWorks = 23;
 			int32 NumNotesReceived = 0;
 
-			Clock->ResetAndStart(0);
+			Clock->SeekTo(0,0);
+			Clock->SetTransportState(0, HarmonixMetasound::EMusicPlayerTransportState::Playing);
 
-			FTimeSignature TimeSignature = Clock->GetBarMap().GetTimeSignatureAtTick(0);
+			FTimeSignature TimeSignature(4, 4);
+			const FTimeSignature* TimeSigPtr = Clock->GetSongMapEvaluator().GetTimeSignatureAtTick(0);
+			if (TimeSigPtr)
+			{
+				TimeSignature = *TimeSigPtr;
+			}
 			FMusicTimeInterval Interval = PulseGenerator.GetInterval();
 			FMusicTimestamp NextPulse{ 1, 1 };
 			IncrementTimestampByOffset(NextPulse, Interval, TimeSignature);
@@ -46,21 +52,20 @@ namespace Harmonix::Midi::Ops::Tests
 				{
 					IncrementTimestampByInterval(EndTimestamp, Interval, TimeSignature);
 				}
-				EndTick = Clock->GetBarMap().MusicTimestampToTick(EndTimestamp);
+				EndTick = Clock->GetSongMapEvaluator().MusicTimestampToTick(EndTimestamp);
 			}
 			
-			while (Clock->GetCurrentMidiTick() < EndTick)
+			while (Clock->GetLastProcessedMidiTick() < EndTick)
 			{
 				// Advance the clock, which will advance the play cursor in the pulse generator
-				Clock->PrepareBlock();
-				Clock->WriteAdvance(0, OperatorSettings.GetNumFramesPerBlock());
+				Clock->Advance(0, OperatorSettings.GetNumFramesPerBlock());
 
 				// Process, which will pop the next notes
 				OutputStream.PrepareBlock();
 				PulseGenerator.Process(OutputStream);
 				
 				// If this is a block where we should get a pulse, check that we got it
-				if (Clock->GetCurrentMidiTick() >= Clock->GetBarMap().MusicTimestampToTick(NextPulse))
+				if (Clock->GetLastProcessedMidiTick() >= Clock->GetSongMapEvaluator().MusicTimestampToTick(NextPulse))
 				{
 					const bool ShouldGetNoteOff = NumNotesReceived > 0;
 
@@ -116,6 +121,9 @@ namespace Harmonix::Midi::Ops::Tests
 
 					IncrementTimestampByInterval(NextPulse, Interval, TimeSignature);
 				}
+
+				// we can prepare the clock for the next block...
+				Clock->PrepareBlock();
 			}
 			
 			UTEST_TRUE("Got all the notes at the right time", NumNotesReceived >= NotesUntilWeAreSatisfiedThisWorks);
diff --git a/Engine/Plugins/Runtime/Harmonix/Source/HarmonixMetasoundTests/Private/Nodes/DelayNodeTests.cpp b/Engine/Plugins/Runtime/Harmonix/Source/HarmonixMetasoundTests/Private/Nodes/DelayNodeTests.cpp
--- a/Engine/Plugins/Runtime/Harmonix/Source/HarmonixMetasoundTests/Private/Nodes/DelayNodeTests.cpp
+++ b/Engine/Plugins/Runtime/Harmonix/Source/HarmonixMetasoundTests/Private/Nodes/DelayNodeTests.cpp
@@ -243,25 +243,22 @@ namespace HarmonixMetasoundTests::DelayNode
 			const TSharedPtr<FMidiFileData> MidiData = MakeShared<FMidiFileData>();
 			check(MidiData);
 
-			FTempoMap& TempoMap = MidiData->SongMaps.GetTempoMap();
-			TempoMap.Empty();
-			FBarMap& BarMap = MidiData->SongMaps.GetBarMap();
-			BarMap.Empty();
+			MidiData->SongMaps.EmptyAllMaps();
 			MidiData->Tracks.Empty();
 
 			MidiData->Tracks.Add(FMidiTrack(TEXT("conductor")));
 			MidiData->Tracks[0].AddEvent(FMidiEvent(0, FMidiMsg(static_cast<uint8>(TimeSigNum), static_cast<uint8>(TimeSigDenom))));
-			BarMap.AddTimeSignatureAtBarIncludingCountIn(0, TimeSigNum, TimeSigDenom);
+			MidiData->SongMaps.AddTimeSignatureAtBarIncludingCountIn(0, TimeSigNum, TimeSigDenom);
 			const int32 MidiTempo = Harmonix::Midi::Constants::BPMToMidiTempo(Tempo);
 			MidiData->Tracks[0].AddEvent(FMidiEvent(0, FMidiMsg(MidiTempo)));
-			TempoMap.AddTempoInfoPoint(MidiTempo, 0);
+			MidiData->SongMaps.AddTempoInfoPoint(MidiTempo, 0);
 			MidiData->Tracks[0].Sort();
 			MidiData->ConformToLength(std::numeric_limits<int32>::max());
 			MidiData->SongMaps.GetSongLengthData().LengthTicks = std::numeric_limits<int32>::max();
 			MidiData->SongMaps.GetSongLengthData().LengthFractionalBars = std::numeric_limits<float>::max();
 
-			(*ClockInput)->AttachToMidiResource(MidiData);
-			(*ClockInput)->ResetAndStart(0);
+			(*ClockInput)->AttachToMidiFile(MidiData);
+			(*ClockInput)->SetTransportState(0, HarmonixMetasound::EMusicPlayerTransportState::Playing);
 
 			SampleRemainder = 0;
 			SampleCount = 0;
@@ -280,7 +277,9 @@ namespace HarmonixMetasoundTests::DelayNode
 			{
 				return false;
 			}
+			(*ClockInput)->PrepareBlock();
 			const int32 NumSamples = Audio::GetMultichannelBufferNumFrames(ComparisonBuffer);
+			int32 BlockFrameIndex = SampleRemainder;
 			SampleRemainder += NumSamples;
 			constexpr int32 MidiGranularity = 128;
 			while (SampleRemainder >= MidiGranularity)
@@ -288,7 +287,8 @@ namespace HarmonixMetasoundTests::DelayNode
 				SampleCount += MidiGranularity;
 				SampleRemainder -= MidiGranularity;
 				const float AdvanceToMs = static_cast<float>(SampleCount) * 1000.0f / SampleRate;
-				(*ClockInput)->AdvanceHiResToMs(0, AdvanceToMs, true);
+				(*ClockInput)->AdvanceToMs(BlockFrameIndex, AdvanceToMs);
+				BlockFrameIndex += MidiGranularity;
 			}
 
 			return true;
diff --git a/Engine/Plugins/Runtime/Harmonix/Source/HarmonixMetasoundTests/Private/Nodes/MetronomeNodeTests.cpp b/Engine/Plugins/Runtime/Harmonix/Source/HarmonixMetasoundTests/Private/Nodes/MetronomeNodeTests.cpp
--- a/Engine/Plugins/Runtime/Harmonix/Source/HarmonixMetasoundTests/Private/Nodes/MetronomeNodeTests.cpp
+++ b/Engine/Plugins/Runtime/Harmonix/Source/HarmonixMetasoundTests/Private/Nodes/MetronomeNodeTests.cpp
@@ -107,7 +107,7 @@ namespace HarmonixMetasoundTests::MetronomeNode
 				return false;
 			}
 
-			if (!InTest.TestEqual(Testf("Midi Clock Looping"), (*OutputMidiClock)->DoesLoop(), Params.Loop))
+			if (!InTest.TestEqual(Testf("Midi Clock Looping"), (*OutputMidiClock)->HasPersistentLoop(), Params.Loop))
 			{
 				return false;
 			}
@@ -122,7 +122,7 @@ namespace HarmonixMetasoundTests::MetronomeNode
 			float TicksPerBlock = TicksPerSec * SecsPerBlock;
 
 			TSharedPtr<FMidiFileData> MidiData = FMidiClock::MakeClockConductorMidiData(Params.Tempo, Params.TimeSigNumerator, Params.TimeSigDenominator);
-			int32 LoopLengthTicks = MidiData->SongMaps.GetBarMap().BarIncludingCountInToTick(Params.LoopLengthBars);
+			int32 LoopLengthTicks = MidiData->SongMaps.BarIncludingCountInToTick(Params.LoopLengthBars);
 			
 			//test for tempo consistency by stopping and restarting the transport (clock output and tempo map)
 			if (TestCaseString.Equals(TempoChangeTestString))
@@ -145,15 +145,15 @@ namespace HarmonixMetasoundTests::MetronomeNode
 				}
 
 				//expect 1 tempo change point at the beginning (original tempo)
-				int32 NumTempoChange = (*OutputMidiClock)->GetTempoMap().GetNumTempoChangePoints();
+				int32 NumTempoChange = (*OutputMidiClock)->GetSongMapEvaluator().GetNumTempoChanges();
 				if (!InTest.TestEqual("Expect 1 tempo change", NumTempoChange, ExpectedNumTempoChangeAtStart))
 				{
 					return false;
 				}
 
 				//check tempo in tempo map
-				int32 TempoChangeTick = (*OutputMidiClock)->GetTempoMap().GetTempoChangePointTick(TempoPointsChangeIndexAtStart);
-				float CurrentTempoBPM = (*OutputMidiClock)->GetTempoMap().GetTempoAtTick(TempoChangeTick);
+				int32 TempoChangeTick = (*OutputMidiClock)->GetSongMapEvaluator().GetTempoChangePointTick(TempoPointsChangeIndexAtStart);
+				float CurrentTempoBPM = (*OutputMidiClock)->GetSongMapEvaluator().GetTempoAtTick(TempoChangeTick);
 				if (!InTest.TestEqual("Expect original tempo at tick 0", CurrentTempoBPM, Params.Tempo, 0.001f))
 				{
 					return false;
@@ -174,15 +174,15 @@ namespace HarmonixMetasoundTests::MetronomeNode
 				}
 
 				//expect 2 tempo change points in tempo map: 1 at the beginning (original), 1 at the current tick (different)
-				NumTempoChange = (*OutputMidiClock)->GetTempoMap().GetNumTempoChangePoints();
+				NumTempoChange = (*OutputMidiClock)->GetSongMapEvaluator().GetNumTempoChanges();
 				if (!InTest.TestEqual("Expect 2 tempo change", NumTempoChange, ExpectedNumTempoChangeBeforeStop))
 				{
 					return false;
 				}
 
 				//check tempo in tempo map
-				TempoChangeTick = (*OutputMidiClock)->GetTempoMap().GetTempoChangePointTick(TempoPointsChangeIndexAfterStop);
-				CurrentTempoBPM = (*OutputMidiClock)->GetTempoMap().GetTempoAtTick(TempoChangeTick);
+				TempoChangeTick = (*OutputMidiClock)->GetSongMapEvaluator().GetTempoChangePointTick(TempoPointsChangeIndexAfterStop);
+				CurrentTempoBPM = (*OutputMidiClock)->GetSongMapEvaluator().GetTempoAtTick(TempoChangeTick);
 				if (!InTest.TestEqual("Expect a different tempo from the original at the current tick", CurrentTempoBPM, DifferentTempo,0.001f))
 				{
 					return false;
@@ -219,15 +219,15 @@ namespace HarmonixMetasoundTests::MetronomeNode
 				}
 				
 				//expect 1 tempo change at the start 
-				NumTempoChange = (*OutputMidiClock)->GetTempoMap().GetNumTempoChangePoints();
-				if (!InTest.TestEqual("Expect 1 tempo change", (*OutputMidiClock)->GetTempoMap().GetNumTempoChangePoints(), ExpectedNumTempoChangeStopAndRestart))
+				NumTempoChange = (*OutputMidiClock)->GetSongMapEvaluator().GetNumTempoChanges();
+				if (!InTest.TestEqual("Expect 1 tempo change", NumTempoChange, ExpectedNumTempoChangeStopAndRestart))
 				{
 					return false;
 				}
 
 				//check tempo in tempo map (original tempo) 
-				TempoChangeTick = (*OutputMidiClock)->GetTempoMap().GetTempoChangePointTick(TempoPointsChangeIndexAtStart);
-				CurrentTempoBPM = (*OutputMidiClock)->GetTempoMap().GetTempoAtTick(TempoChangeTick);
+				TempoChangeTick = (*OutputMidiClock)->GetSongMapEvaluator().GetTempoChangePointTick(TempoPointsChangeIndexAtStart);
+				CurrentTempoBPM = (*OutputMidiClock)->GetSongMapEvaluator().GetTempoAtTick(TempoChangeTick);
 				if (!InTest.TestEqual("Expect original tempo at tick 0", CurrentTempoBPM, Params.Tempo,0.001f))
 				{
 					return false;
@@ -245,7 +245,7 @@ namespace HarmonixMetasoundTests::MetronomeNode
 
 				int32 ExpectedTick = FMath::RoundToInt32(TicksPerBlock * (BlockIndex + 1));
 
-				if (!InTest.TestEqual(Testf("Midi Clock Looping"), (*OutputMidiClock)->DoesLoop(), Params.Loop))
+				if (!InTest.TestEqual(Testf("Midi Clock Looping"), (*OutputMidiClock)->HasPersistentLoop(), Params.Loop))
 				{
 					return false;
 				}
@@ -262,7 +262,7 @@ namespace HarmonixMetasoundTests::MetronomeNode
 					ExpectedTick %= LoopLengthTicks;
 				}
 
-				int32 ActualTick = (*OutputMidiClock)->GetCurrentMidiTick();
+				int32 ActualTick = (*OutputMidiClock)->GetNextMidiTickToProcess();
 				// allow for single tick tolerance?
 				ExpectedTick = FMath::Abs(ExpectedTick - ActualTick) <= 1 ? ActualTick : ExpectedTick;
 				if (AllTicksEqual && (ActualTick != ExpectedTick))
@@ -273,21 +273,20 @@ namespace HarmonixMetasoundTests::MetronomeNode
 				}
 
 				// test looping here since it the values may not be updated until the first execution
-				if ((*OutputMidiClock)->DoesLoop())
+				if ((*OutputMidiClock)->HasPersistentLoop())
 				{
 					int32 LoopStartTick = 0;
 					int32 LoopEndTick = LoopLengthTicks;
 
-					// use default values since the looping clock uses 120 bpm for its own midi data
-					float LoopStartMs = LoopStartTick / DefaultTicksPerMs;
-					float LoopEndMs = LoopEndTick / DefaultTicksPerMs;
+					float LoopStartMs = LoopStartTick / TicksPerMs;
+					float LoopEndMs = LoopEndTick / TicksPerMs;
 
-					if (!InTest.TestEqual(Testf("Midi Clock Loop Start Tick"), (*OutputMidiClock)->GetLoopStartTick(), LoopStartTick))
+					if (!InTest.TestEqual(Testf("Midi Clock Loop Start Tick"), (*OutputMidiClock)->GetFirstTickInLoop(), LoopStartTick))
 					{
 						return false;
 					}
 
-					if (!InTest.TestEqual(Testf("Midi Clock Loop End Tick"), (*OutputMidiClock)->GetLoopEndTick(), LoopEndTick))
+					if (!InTest.TestEqual(Testf("Midi Clock Loop Length"), (*OutputMidiClock)->GetLoopLengthTicks(), LoopEndTick - LoopStartTick))
 					{
 						return false;
 					}
diff --git a/Engine/Plugins/Runtime/Harmonix/Source/HarmonixMetasoundTests/Private/Nodes/MidiPlayerNodeTests.cpp b/Engine/Plugins/Runtime/Harmonix/Source/HarmonixMetasoundTests/Private/Nodes/MidiPlayerNodeTests.cpp
--- a/Engine/Plugins/Runtime/Harmonix/Source/HarmonixMetasoundTests/Private/Nodes/MidiPlayerNodeTests.cpp
+++ b/Engine/Plugins/Runtime/Harmonix/Source/HarmonixMetasoundTests/Private/Nodes/MidiPlayerNodeTests.cpp
@@ -160,7 +160,7 @@ namespace HarmonixMetasoundTests::MidiPlayerNode
 				MetronomeClockOut = Generator->GetOutputReadReference<FMidiClock>(NodeNames::Metronome::ClockOut);
 
 				TSharedPtr<FMidiFileData> MidiData = FMidiClock::MakeClockConductorMidiData(Params.Metronome.Tempo, Params.Metronome.TimeSig.Numerator, Params.Metronome.TimeSig.Denominator);
-				MetronomeLoopLengthTicks = MidiData->SongMaps.GetBarMap().BarIncludingCountInToTick(Params.Metronome.LoopLengthBars);
+				MetronomeLoopLengthTicks = MidiData->SongMaps.BarIncludingCountInToTick(Params.Metronome.LoopLengthBars);
 			}
 
 			if (Params.MidiPlayer.MidiFile.IsValid())
@@ -210,19 +210,19 @@ namespace HarmonixMetasoundTests::MidiPlayerNode
 					MidiPlayerExpectedTick %= MidiPlayerLoopLengthTicks;
 				}
 
-				if (!InTest.TestEqual("Midi Player Looping", (*MidiPlayerClockOut)->DoesLoop(), Params.MidiPlayer.Loop))
+				if (!InTest.TestEqual("Midi Player Looping", (*MidiPlayerClockOut)->HasPersistentLoop(), Params.MidiPlayer.Loop))
 				{
 					return false;
 				}
 
 				if (MetronomeClockOut)
 				{
-					if (!InTest.TestEqual("Metronome Looping", (*MetronomeClockOut)->DoesLoop(), Params.Metronome.Loop))
+					if (!InTest.TestEqual("Metronome Looping", (*MetronomeClockOut)->HasPersistentLoop(), Params.Metronome.Loop))
 					{
 						return false;
 					}
 
-					int32 MetronomeActualTick = (*MetronomeClockOut)->GetCurrentMidiTick();
+					int32 MetronomeActualTick = (*MetronomeClockOut)->GetNextMidiTickToProcess();
 					MetronomeExpectedTick = FMath::Abs(MidiPlayerExpectedTick - MetronomeActualTick) <= 1 ? MetronomeActualTick : MetronomeExpectedTick;
 					if (MetronomeAllTicksEqual && (MetronomeActualTick != MetronomeExpectedTick))
 					{
@@ -232,7 +232,7 @@ namespace HarmonixMetasoundTests::MidiPlayerNode
 					}
 				}
 
-				int32 MidiPlayerActualTick = (*MidiPlayerClockOut)->GetCurrentMidiTick();
+				int32 MidiPlayerActualTick = (*MidiPlayerClockOut)->GetNextMidiTickToProcess();
 				MidiPlayerExpectedTick = FMath::Abs(MidiPlayerExpectedTick - MidiPlayerActualTick) <= 1 ? MidiPlayerActualTick : MidiPlayerExpectedTick;
 
 				if (MidiPlayerAllTicksEqual && (MidiPlayerActualTick != MidiPlayerExpectedTick))
@@ -257,10 +257,7 @@ namespace HarmonixMetasoundTests::MidiPlayerNode
 			TSharedPtr<FMidiFileData> OutMidiData = MakeShared<FMidiFileData>();
 
 			// clear it all out for good measure...
-			FTempoMap& TempoMap = OutMidiData->SongMaps.GetTempoMap();
-			TempoMap.Empty();
-			FBarMap& BarMap = OutMidiData->SongMaps.GetBarMap();
-			BarMap.Empty();
+			OutMidiData->SongMaps.EmptyAllMaps();
 			OutMidiData->Tracks.Empty();
 
 			// create conductor track
@@ -270,13 +267,13 @@ namespace HarmonixMetasoundTests::MidiPlayerNode
 			int32 TimeSigNum = FMath::Clamp(InTimeSigNum, 1, 64);
 			int32 TimeSigDen = FMath::Clamp(InTimeSigDen, 1, 64);
 			Track.AddEvent(FMidiEvent(0, FMidiMsg((uint8)TimeSigNum, (uint8)TimeSigDen)));
-			BarMap.AddTimeSignatureAtBarIncludingCountIn(0, TimeSigNum, TimeSigDen);
+			OutMidiData->SongMaps.AddTimeSignatureAtBarIncludingCountIn(0, TimeSigNum, TimeSigDen);
 
 			// add tempo info
 			float TempoBpm = FMath::Max(1.0f, InTempoBpm);
 			int32 MidiTempo = Harmonix::Midi::Constants::BPMToMidiTempo(TempoBpm);
 			Track.AddEvent(FMidiEvent(0, FMidiMsg(MidiTempo)));
-			TempoMap.AddTempoInfoPoint(MidiTempo, 0);
+			OutMidiData->SongMaps.AddTempoInfoPoint(MidiTempo, 0);
 
 			Track.Sort();
 			OutMidiData->SongMaps.SetLengthTotalBars(InLengthBars);
@@ -312,7 +309,8 @@ namespace HarmonixMetasoundTests::MidiPlayerNode
 		// Validate output.
 		TOptional<FMidiClockReadRef> OutputMidiClock = Generator->GetOutputReadReference<FMidiClock>(CommonPinNames::Outputs::MidiClockName);
 		UTEST_TRUE("Output exists", OutputMidiClock.IsSet());
-		UTEST_EQUAL("Current Midi Tick Test", (*OutputMidiClock)->GetCurrentMidiTick(), -1);
+		UTEST_EQUAL("Current Midi Tick Test", (*OutputMidiClock)->GetLastProcessedMidiTick(), -1);
+		UTEST_EQUAL("Next Midi Tick Test", (*OutputMidiClock)->GetNextMidiTickToProcess(), 0);
 
 		return true;
 	}
diff --git a/Engine/Plugins/Runtime/Harmonix/Source/HarmonixMetasoundTests/Private/Nodes/MidiPulseGeneratorNodeTests.cpp b/Engine/Plugins/Runtime/Harmonix/Source/HarmonixMetasoundTests/Private/Nodes/MidiPulseGeneratorNodeTests.cpp
--- a/Engine/Plugins/Runtime/Harmonix/Source/HarmonixMetasoundTests/Private/Nodes/MidiPulseGeneratorNodeTests.cpp
+++ b/Engine/Plugins/Runtime/Harmonix/Source/HarmonixMetasoundTests/Private/Nodes/MidiPulseGeneratorNodeTests.cpp
@@ -22,6 +22,7 @@ namespace HarmonixMetasound::Nodes::MidiPulseGeneratorNode::Tests
 
 		TOptional<FMidiClockWriteRef> Clock = Generator->GetInputWriteReference<FMidiClock>(Inputs::MidiClockName);
 		UTEST_TRUE("Got clock", Clock.IsSet());
+		(*Clock)->SetTransportState(0, EMusicPlayerTransportState::Playing);
 
 		Harmonix::Midi::Ops::FPulseGenerator PulseGenerator;
 		FMidiStream PulseGeneratorMidiOutput;
@@ -35,7 +36,7 @@ namespace HarmonixMetasound::Nodes::MidiPulseGeneratorNode::Tests
 		{
 			// Advance the clock, which will advance the play cursor in the pulse generators
 			(*Clock)->PrepareBlock();
-			(*Clock)->WriteAdvance(0, Generator->OperatorSettings.GetNumFramesPerBlock());
+			(*Clock)->Advance(0, Generator->OperatorSettings.GetNumFramesPerBlock());
 
 			// Process
 			Generator->OnGenerateAudio(Buffer.GetData(), Buffer.Num());
diff --git a/Engine/Plugins/Runtime/Harmonix/Source/HarmonixMetasoundTests/Private/Nodes/MidiTextTriggerNodeTests.cpp b/Engine/Plugins/Runtime/Harmonix/Source/HarmonixMetasoundTests/Private/Nodes/MidiTextTriggerNodeTests.cpp
--- a/Engine/Plugins/Runtime/Harmonix/Source/HarmonixMetasoundTests/Private/Nodes/MidiTextTriggerNodeTests.cpp
+++ b/Engine/Plugins/Runtime/Harmonix/Source/HarmonixMetasoundTests/Private/Nodes/MidiTextTriggerNodeTests.cpp
@@ -20,8 +20,10 @@ namespace HarmonixMetasoundTests::MidiTextTriggerNode
 	{
 		// Make a midi file...
 		UMidiFile* TheMidi = NewObject<UMidiFile>();
+		// Set the initial time signature
+		TheMidi->GetSongMaps()->AddTimeSignatureAtBarIncludingCountIn(0, 4, 4);
 		// Set the initial tempo...
-		TheMidi->GetSongMaps()->GetTempoMap().AddTempoInfoPoint(Harmonix::Midi::Constants::BPMToMidiTempo(120),0);
+		TheMidi->GetSongMaps()->AddTempoInfoPoint(Harmonix::Midi::Constants::BPMToMidiTempo(120),0);
 		// Make a track to hold some text events. This will be track 1 (track 0 is always the conductor track)...
 		FMidiTrack* TextTrack = TheMidi->AddTrack("TextEventTest");
 		// Make a text event...
diff --git a/Engine/Plugins/Runtime/Harmonix/Source/HarmonixMetasoundTests/Private/Nodes/MorphingLfoNode.spec.cpp b/Engine/Plugins/Runtime/Harmonix/Source/HarmonixMetasoundTests/Private/Nodes/MorphingLfoNode.spec.cpp
--- a/Engine/Plugins/Runtime/Harmonix/Source/HarmonixMetasoundTests/Private/Nodes/MorphingLfoNode.spec.cpp
+++ b/Engine/Plugins/Runtime/Harmonix/Source/HarmonixMetasoundTests/Private/Nodes/MorphingLfoNode.spec.cpp
@@ -65,24 +65,21 @@ namespace HarmonixMetasoundTests::MorphingLFONode
 		const TSharedPtr<FMidiFileData> MidiData = MakeShared<FMidiFileData>();
 		check(MidiData);
 
-		FTempoMap& TempoMap = MidiData->SongMaps.GetTempoMap();
-		TempoMap.Empty();
-		FBarMap& BarMap = MidiData->SongMaps.GetBarMap();
-		BarMap.Empty();
+		MidiData->SongMaps.EmptyAllMaps();
 		MidiData->Tracks.Empty();
 
 		MidiData->Tracks.Add(FMidiTrack(TEXT("conductor")));
 		MidiData->Tracks[0].AddEvent(FMidiEvent(0, FMidiMsg(static_cast<uint8>(TimeSigNumerator), static_cast<uint8>(TimeSigDenominator))));
-		BarMap.AddTimeSignatureAtBarIncludingCountIn(0, TimeSigNumerator, TimeSigNumerator);
+		MidiData->SongMaps.AddTimeSignatureAtBarIncludingCountIn(0, TimeSigNumerator, TimeSigNumerator);
 		const int32 MidiTempo = Harmonix::Midi::Constants::BPMToMidiTempo(Tempo);
 		MidiData->Tracks[0].AddEvent(FMidiEvent(0, FMidiMsg(MidiTempo)));
-		TempoMap.AddTempoInfoPoint(MidiTempo, 0);
+		MidiData->SongMaps.AddTempoInfoPoint(MidiTempo, 0);
 		MidiData->Tracks[0].Sort();
 		MidiData->ConformToLength(std::numeric_limits<int32>::max());
 
-		(*ClockInput)->AttachToMidiResource(MidiData);
-		(*ClockInput)->ResetAndStart(0);
-		(*ClockInput)->AddSpeedChangeToBlock(HarmonixMetasound::FMidiTimestampSpeed{ 0, 0, Speed });
+		(*ClockInput)->AttachToMidiFile(MidiData);
+		(*ClockInput)->SetTransportState(0, HarmonixMetasound::EMusicPlayerTransportState::Playing);
+		(*ClockInput)->SetSpeed(0, Speed);
 
 		SampleCount = 0;
 		SampleRemainder = 0;
@@ -101,7 +98,10 @@ namespace HarmonixMetasoundTests::MorphingLFONode
 			return false;
 		}
 		
+		(*ClockInput)->PrepareBlock();
+
 		const int32 NumSamples = Generator->OperatorSettings.GetNumFramesPerBlock();
+		int32 BlockFrameIndex = SampleRemainder;
 		SampleRemainder += NumSamples;
 		constexpr int32 MidiGranularity = 128;
 		while (SampleRemainder >= MidiGranularity)
@@ -109,7 +109,8 @@ namespace HarmonixMetasoundTests::MorphingLFONode
 			SampleCount += MidiGranularity;
 			SampleRemainder -= MidiGranularity;
 			const float AdvanceToMs = static_cast<float>(SampleCount) * 1000.0f / Generator->OperatorSettings.GetSampleRate();
-			(*ClockInput)->AdvanceHiResToMs(0, AdvanceToMs, true);
+			(*ClockInput)->AdvanceToMs(BlockFrameIndex, AdvanceToMs);
+			BlockFrameIndex += MidiGranularity;
 		}
 
 		return true;
diff --git a/Engine/Plugins/Runtime/Harmonix/Source/HarmonixMetasoundTests/Private/Subsystems/MidiClockUpdateSubsystemTests.cpp b/Engine/Plugins/Runtime/Harmonix/Source/HarmonixMetasoundTests/Private/Subsystems/MidiClockUpdateSubsystemTests.cpp
--- a/Engine/Plugins/Runtime/Harmonix/Source/HarmonixMetasoundTests/Private/Subsystems/MidiClockUpdateSubsystemTests.cpp
+++ b/Engine/Plugins/Runtime/Harmonix/Source/HarmonixMetasoundTests/Private/Subsystems/MidiClockUpdateSubsystemTests.cpp
@@ -22,18 +22,16 @@ namespace HarmonixMetasoundTests::MidiClockUpdateSubsystem
 			
 			MidiData->Tracks.Add(FMidiTrack(TEXT("conductor")));
 			MidiData->Tracks[0].AddEvent(FMidiEvent(0, FMidiMsg(static_cast<uint8>(TimeSigNum), static_cast<uint8>(TimeSigDenom))));
-			FBarMap& BarMap = MidiData->SongMaps.GetBarMap();
-			BarMap.AddTimeSignatureAtBarIncludingCountIn(0, TimeSigNum, TimeSigDenom);
+			MidiData->SongMaps.AddTimeSignatureAtBarIncludingCountIn(0, TimeSigNum, TimeSigDenom);
 			const int32 MidiTempo = Harmonix::Midi::Constants::BPMToMidiTempo(Tempo);
 			MidiData->Tracks[0].AddEvent(FMidiEvent(0, FMidiMsg(MidiTempo)));
-			FTempoMap& TempoMap = MidiData->SongMaps.GetTempoMap();
-			TempoMap.AddTempoInfoPoint(MidiTempo, 0);
+			MidiData->SongMaps.AddTempoInfoPoint(MidiTempo, 0);
 			MidiData->Tracks[0].Sort();
 			MidiData->ConformToLength(std::numeric_limits<int32>::max());
 
 			HarmonixMetasound::FMidiClock Clock{ OperatorSettings };
-			Clock.AttachToMidiResource(MidiData);
-			Clock.ResetAndStart(0);
+			Clock.AttachToMidiFile(MidiData);
+			Clock.SetTransportState(0, HarmonixMetasound::EMusicPlayerTransportState::Playing);
 
 			return Clock;
 		}	
@@ -69,22 +67,23 @@ namespace HarmonixMetasoundTests::MidiClockUpdateSubsystem
 			// Advance the high-resolution clock
 			SampleRemainder += NumSamples;
 			constexpr int32 MidiGranularity = 128;
+			Clock.PrepareBlock();
+			int32 BlockOffset = 0;
 			while (SampleRemainder >= MidiGranularity)
 			{
 				SampleCount += MidiGranularity;
 				SampleRemainder -= MidiGranularity;
 				const float AdvanceToMs = static_cast<float>(SampleCount) * 1000.0f / OperatorSettings.GetSampleRate();
-				Clock.AdvanceHiResToMs(0, AdvanceToMs, true);
+				Clock.Advance(BlockOffset, MidiGranularity);
+				BlockOffset += MidiGranularity;
 			}
 
 			// Tick the subsystem (low-resolution clocks)
 			Subsystem->TickForTesting();
 
 			// Check that the high- and low-resolution clocks are at the same place
-			UTEST_EQUAL(
-				FString::Printf(TEXT("High- and low-res are at the same tick: iteration %i"), i),
-				Clock.GetCurrentHiResTick(),
-				Clock.GetCurrentLowResTick());
+			// TO DO!
+
 		}
 
 		return true;
diff --git a/Engine/Plugins/Runtime/Harmonix/Source/HarmonixMidi/Private/HarmonixMidi/BarMap.cpp b/Engine/Plugins/Runtime/Harmonix/Source/HarmonixMidi/Private/HarmonixMidi/BarMap.cpp
--- a/Engine/Plugins/Runtime/Harmonix/Source/HarmonixMidi/Private/HarmonixMidi/BarMap.cpp
+++ b/Engine/Plugins/Runtime/Harmonix/Source/HarmonixMidi/Private/HarmonixMidi/BarMap.cpp
@@ -1,12 +1,12 @@
 // Copyright Epic Games, Inc. All Rights Reserved.
 #include "HarmonixMidi/BarMap.h"
 #include "HarmonixMidi/MidiConstants.h"
-#include "HarmonixMidi/TempoMap.h"
-#include "HarmonixMidi/BeatMap.h" 
+#include "HarmonixMidi/SongMaps.h"
 
 #include "Misc/ScopeLock.h"
 #include "Serialization/Archive.h"
 
+const FTimeSignaturePoint FBarMap::sDefaultTimeSignature(0, 0, FTimeSignature(4, 4), 0, std::numeric_limits<int32>::max());
 
 bool FBarMap::operator==(const FBarMap& Other) const
 {
@@ -40,6 +40,44 @@ bool FBarMap::IsEmpty() const
 	return Points.IsEmpty();
 }
 
+int32 FBarMap::CalculateMidiTick(const FMusicTimestamp& Timestamp, const EMidiClockSubdivisionQuantization Quantize) const
+{
+	int32 TriggerTick = 0;
+	if (Quantize == EMidiClockSubdivisionQuantization::None)
+	{
+		TriggerTick = FMath::RoundToInt32(MusicTimestampToTick(Timestamp));
+	}
+	else
+	{
+		int32 RawTick = FMath::RoundToInt32(MusicTimestampToTick(Timestamp));
+		int32 BarTick = MusicTimestampBarToTick(Timestamp.Bar);
+		int32 TicksPerQuantizationUnit = SubdivisionToMidiTicks(Quantize, RawTick);
+		if (ensure(TicksPerQuantizationUnit > 0))
+		{
+			float NumUnits = ((float)(RawTick - BarTick)) / (float)TicksPerQuantizationUnit;
+			int32 NumWholeUnits = FMath::RoundToInt32(NumUnits);
+			TriggerTick = BarTick + (NumWholeUnits * TicksPerQuantizationUnit);
+		}
+		else
+		{
+			TriggerTick = FMath::RoundToInt32(MusicTimestampToTick(Timestamp));
+		}
+	}
+	return TriggerTick;
+}
+
+int32 FBarMap::SubdivisionToMidiTicks(const EMidiClockSubdivisionQuantization Division, const int32 AtTick) const
+{
+	FTimeSignature TimeSignature(4, 4);
+	int32 BarMapPointIndex = GetPointIndexForTick(AtTick);
+	if (BarMapPointIndex >= 0)
+	{
+		const FTimeSignaturePoint& TimeSignaturePoint = GetTimeSignaturePoint(BarMapPointIndex);
+		TimeSignature = TimeSignaturePoint.TimeSignature;
+	}
+	return Harmonix::Midi::Constants::SubdivisionToMidiTicks(Division, TimeSignature, TicksPerQuarterNote);
+}
+
 float FBarMap::MusicTimestampToTick(const FMusicTimestamp& Timestamp) const
 {
 	float Beat = Timestamp.Beat;
@@ -407,7 +445,10 @@ int32 FBarMap::GetPointIndexForTick(int32 Tick) const
 
 const FTimeSignaturePoint& FBarMap::GetTimeSignaturePoint(int32 Index) const
 {
-	check(!Points.IsEmpty());
+	if (Points.IsEmpty())
+	{
+		return sDefaultTimeSignature;
+	}
 
 	if (!Points.IsValidIndex(Index))
 	{
@@ -431,7 +472,10 @@ FTimeSignaturePoint& FBarMap::GetTimeSignaturePoint(int32 Index)
 
 const FTimeSignature& FBarMap::GetTimeSignatureAtTick(int32 Tick) const
 {
-	check(!Points.IsEmpty());
+	if (Points.IsEmpty())
+	{
+		return sDefaultTimeSignature.TimeSignature;
+	}
 	
 	int32 Index = FMusicMapUtl::GetPointIndexForTick(Points, Tick);
 
@@ -443,9 +487,27 @@ const FTimeSignature& FBarMap::GetTimeSignatureAtTick(int32 Tick) const
 	return Points[Index].TimeSignature;
 }
 
+const FTimeSignaturePoint* FBarMap::GetTimeSignaturePointForTick(int32 Tick) const
+{
+	if (Points.IsEmpty())
+	{
+		return nullptr;
+	}
+
+	int32 Index = FMusicMapUtl::GetPointIndexForTick(Points, Tick);
+	if (!Points.IsValidIndex(Index))
+	{
+		Index = 0;
+	}
+	return &Points[Index];
+}
+
 const FTimeSignature& FBarMap::GetTimeSignatureAtBar(int32 Bar) const
 {
-	check(!Points.IsEmpty());
+	if (Points.IsEmpty())
+	{
+		return sDefaultTimeSignature.TimeSignature;
+	}
 
 	Bar -= StartBar;
 
@@ -472,3 +534,12 @@ void FBarMap::Finalize(int32 InLastTick)
 	}
 	FMusicMapUtl::Finalize(Points, InLastTick);
 }
+
+int32 FBarMap::GetTimeSignatureChangePointTick(int32 PointIndex) const
+{
+	if (!Points.IsValidIndex(PointIndex))
+	{
+		return 0;
+	}
+	return Points[PointIndex].StartTick;
+}
diff --git a/Engine/Plugins/Runtime/Harmonix/Source/HarmonixMidi/Private/HarmonixMidi/BeatMap.cpp b/Engine/Plugins/Runtime/Harmonix/Source/HarmonixMidi/Private/HarmonixMidi/BeatMap.cpp
--- a/Engine/Plugins/Runtime/Harmonix/Source/HarmonixMidi/Private/HarmonixMidi/BeatMap.cpp
+++ b/Engine/Plugins/Runtime/Harmonix/Source/HarmonixMidi/Private/HarmonixMidi/BeatMap.cpp
@@ -229,3 +229,14 @@ void FBeatMap::Finalize(int32 LastTick)
 	}
 	Bars.Emplace(PulseBarStartTick, PulseBarLength, FirstBeatInBar, LastBeatInBar);
 }
+
+FString HARMONIXMIDI_API MusicalBeatTypeToString(EMusicalBeatType BeatType)
+{
+	switch (BeatType)
+	{
+	case EMusicalBeatType::Downbeat: return TEXT("Downbeat");
+	case EMusicalBeatType::Strong:	 return TEXT("Strong");
+	case EMusicalBeatType::Normal:	 return TEXT("Normal");
+	}
+	return TEXT("<Unrecognized>");
+}
diff --git a/Engine/Plugins/Runtime/Harmonix/Source/HarmonixMidi/Private/HarmonixMidi/MidiConstants.cpp b/Engine/Plugins/Runtime/Harmonix/Source/HarmonixMidi/Private/HarmonixMidi/MidiConstants.cpp
--- a/Engine/Plugins/Runtime/Harmonix/Source/HarmonixMidi/Private/HarmonixMidi/MidiConstants.cpp
+++ b/Engine/Plugins/Runtime/Harmonix/Source/HarmonixMidi/Private/HarmonixMidi/MidiConstants.cpp
@@ -1,11 +1,91 @@
 // Copyright Epic Games, Inc. All Rights Reserved.
 
 #include "HarmonixMidi/MidiConstants.h"
+#include "HarmonixMidi/SongMaps.h"
 
 DEFINE_LOG_CATEGORY(LogMIDI);
 
 namespace Harmonix::Midi::Constants
 {
+	int32 SubdivisionToMidiTicks(const EMidiClockSubdivisionQuantization Division, const FTimeSignature& TimeSignature, const int32 TicksPerQuarterNote)
+	{
+		switch (Division)
+		{
+		case EMidiClockSubdivisionQuantization::None:                   return 1;
+		case EMidiClockSubdivisionQuantization::Bar: 					return TimeSignature.Numerator * ((TicksPerQuarterNote * 4) / TimeSignature.Denominator);
+		case EMidiClockSubdivisionQuantization::Beat:					return (TicksPerQuarterNote * 4) / TimeSignature.Denominator;
+		case EMidiClockSubdivisionQuantization::ThirtySecondNote:		return TicksPerQuarterNote / 8;
+		case EMidiClockSubdivisionQuantization::SixteenthNote:			return TicksPerQuarterNote / 4;
+		case EMidiClockSubdivisionQuantization::EighthNote:				return TicksPerQuarterNote / 2;
+		case EMidiClockSubdivisionQuantization::QuarterNote:			return TicksPerQuarterNote;
+		case EMidiClockSubdivisionQuantization::HalfNote:				return TicksPerQuarterNote * 2;
+		case EMidiClockSubdivisionQuantization::WholeNote:				return TicksPerQuarterNote * 4;
+		case EMidiClockSubdivisionQuantization::DottedSixteenthNote:	return (TicksPerQuarterNote / 4) + (TicksPerQuarterNote / 8);
+		case EMidiClockSubdivisionQuantization::DottedEighthNote:		return (TicksPerQuarterNote / 2) + (TicksPerQuarterNote / 4);
+		case EMidiClockSubdivisionQuantization::DottedQuarterNote:		return (TicksPerQuarterNote)+(TicksPerQuarterNote / 2);
+		case EMidiClockSubdivisionQuantization::DottedHalfNote:			return (TicksPerQuarterNote * 2) + (TicksPerQuarterNote);
+		case EMidiClockSubdivisionQuantization::DottedWholeNote:		return (TicksPerQuarterNote * 4) + (TicksPerQuarterNote * 2);
+		case EMidiClockSubdivisionQuantization::SixteenthNoteTriplet:   return (TicksPerQuarterNote / 2) / 3;
+		case EMidiClockSubdivisionQuantization::EighthNoteTriplet:		return TicksPerQuarterNote / 3;
+		case EMidiClockSubdivisionQuantization::QuarterNoteTriplet:		return (TicksPerQuarterNote * 2) / 3;
+		case EMidiClockSubdivisionQuantization::HalfNoteTriplet:        return (TicksPerQuarterNote * 4) / 3;
+		default:	                                             		checkNoEntry();	return 1;
+		}
+	}
+
+	float SubdivisionToBeats(EMidiClockSubdivisionQuantization Subdivision, const FTimeSignature& TimeSignature)
+	{
+		// Easy cases first
+		if (Subdivision == EMidiClockSubdivisionQuantization::Bar)
+		{
+			return TimeSignature.Numerator;
+		}
+
+		if (Subdivision == EMidiClockSubdivisionQuantization::Beat)
+		{
+			return 1;
+		}
+
+		const float BeatsPerQuarter = TimeSignature.Denominator / 4.0f;
+
+		switch (Subdivision)
+		{
+		case EMidiClockSubdivisionQuantization::ThirtySecondNote:
+			return BeatsPerQuarter / 8;
+		case EMidiClockSubdivisionQuantization::SixteenthNote:
+			return BeatsPerQuarter / 4;
+		case EMidiClockSubdivisionQuantization::EighthNote:
+			return BeatsPerQuarter / 2;
+		case EMidiClockSubdivisionQuantization::QuarterNote:
+			return BeatsPerQuarter;
+		case EMidiClockSubdivisionQuantization::HalfNote:
+			return BeatsPerQuarter * 2;
+		case EMidiClockSubdivisionQuantization::WholeNote:
+			return BeatsPerQuarter * 4;
+		case EMidiClockSubdivisionQuantization::DottedSixteenthNote:
+			return BeatsPerQuarter / 4 + BeatsPerQuarter / 8;
+		case EMidiClockSubdivisionQuantization::DottedEighthNote:
+			return BeatsPerQuarter / 2 + BeatsPerQuarter / 4;
+		case EMidiClockSubdivisionQuantization::DottedQuarterNote:
+			return BeatsPerQuarter + BeatsPerQuarter / 2;
+		case EMidiClockSubdivisionQuantization::DottedHalfNote:
+			return BeatsPerQuarter * 3;
+		case EMidiClockSubdivisionQuantization::DottedWholeNote:
+			return BeatsPerQuarter * 6;
+		case EMidiClockSubdivisionQuantization::SixteenthNoteTriplet:
+			return (BeatsPerQuarter / 4) * 2 / 3;
+		case EMidiClockSubdivisionQuantization::EighthNoteTriplet:
+			return (BeatsPerQuarter / 2) * 2 / 3;
+		case EMidiClockSubdivisionQuantization::QuarterNoteTriplet:
+			return BeatsPerQuarter * 2 / 3;
+		case EMidiClockSubdivisionQuantization::HalfNoteTriplet:
+			return BeatsPerQuarter * 4 / 3;
+		default:
+			checkNoEntry();
+			return 0;
+		}
+	}
+
 	FString GetTextTypeName(uint8 TextType)
 	{
 		switch (TextType)
diff --git a/Engine/Plugins/Runtime/Harmonix/Source/HarmonixMidi/Private/HarmonixMidi/MidiFile.cpp b/Engine/Plugins/Runtime/Harmonix/Source/HarmonixMidi/Private/HarmonixMidi/MidiFile.cpp
--- a/Engine/Plugins/Runtime/Harmonix/Source/HarmonixMidi/Private/HarmonixMidi/MidiFile.cpp
+++ b/Engine/Plugins/Runtime/Harmonix/Source/HarmonixMidi/Private/HarmonixMidi/MidiFile.cpp
@@ -96,13 +96,19 @@ FString UMidiFile::GetImportedSrcFilePath() const
 
 #endif
 
-TSharedPtr<Audio::IProxyData> UMidiFile::CreateProxyData(const Audio::FProxyDataInitParams& InitParams)
+TSharedPtr<FMidiFileData> UMidiFile::GetOrCreateRenderableCopy()
 {
 	if (!RenderableCopyOfMidiFileData)
 	{
 		RenderableCopyOfMidiFileData = MakeShared<FMidiFileData>(TheMidiData);
 	}
-	TSharedPtr<FMidiFileProxy> Proxy = MakeShared<FMidiFileProxy>(RenderableCopyOfMidiFileData);
+	return RenderableCopyOfMidiFileData;
+}
+
+TSharedPtr<Audio::IProxyData> UMidiFile::CreateProxyData(const Audio::FProxyDataInitParams& InitParams)
+{
+	TSharedPtr<FMidiFileData> Renderable = GetOrCreateRenderableCopy();
+	TSharedPtr<FMidiFileProxy> Proxy = MakeShared<FMidiFileProxy>(Renderable);
 	return Proxy;
 }
 
@@ -249,27 +255,26 @@ void UMidiFile::BuildConductorTrack()
 		TheMidiData.Tracks[0] = FMidiTrack("Conductor");
 	}
 
-	const FTempoMap& TempoMap = TheMidiData.SongMaps.GetTempoMap();
-
-	int32 numTempoChanges = TempoMap.GetNumTempoChangePoints();
+	int32 numTempoChanges = TheMidiData.SongMaps.GetNumTempoChanges();
 
 	int32 i;
 	for (i = 0; i < numTempoChanges; ++i)
 	{
-		int32 tick = TempoMap.GetTempoChangePointTick(i);
-		float msPerQuarterNote = TempoMap.GetMsPerQuarterNoteAtTick(tick);
+		const FTempoInfoPoint* Point = TheMidiData.SongMaps.GetTempoInfoPoint(i);
+		check(Point);
+		int32 tick = Point->StartTick;
+		float msPerQuarterNote = Point->GetMsPerQuarterNote();
 		int32 usecPerQuarterNote = int32(msPerQuarterNote * 1000);
 		TheMidiData.Tracks[0].AddEvent(FMidiEvent(tick, FMidiMsg(usecPerQuarterNote)));
 	}
 
-	const FBarMap& BarMap = TheMidiData.SongMaps.GetBarMap();
-
-	int32 numTimeSigChanges = BarMap.GetNumTimeSignaturePoints();
+	int32 numTimeSigChanges = TheMidiData.SongMaps.GetNumTimeSignatureChanges();
 
 	for (i = 0; i < numTimeSigChanges; ++i)
 	{
-		const FTimeSignaturePoint& sig = BarMap.GetTimeSignaturePoint(i);
-		TheMidiData.Tracks[0].AddEvent(FMidiEvent(sig.StartTick, FMidiMsg(sig.TimeSignature.Numerator, sig.TimeSignature.Denominator)));
+		const FTimeSignaturePoint* sig = TheMidiData.SongMaps.GetTimeSignaturePoint(i);
+		check(sig);
+		TheMidiData.Tracks[0].AddEvent(FMidiEvent(sig->StartTick, FMidiMsg(sig->TimeSignature.Numerator, sig->TimeSignature.Denominator)));
 	}
 
 	TheMidiData.Tracks[0].Sort();
@@ -520,7 +525,7 @@ void UMidiFile::PostEditChangeProperty(FPropertyChangedEvent& PropertyChangedEve
 	
 	if (PropertyChangedEvent.Property->GetName() == StartBarPropertyName)
 	{
-		TheMidiData.SongMaps.GetBarMap().SetStartBar(StartBar);
+		TheMidiData.SongMaps.SetStartBar(StartBar);
 	}
 }
 #endif
@@ -671,27 +676,26 @@ void FMidiFileData::PostSerialize(const FArchive& Ar)
 	{
 		UE_LOG(LogMIDI, Warning, TEXT("Empty bar map. Rebuilding."));
 		const FMidiEventList& Events = Tracks[0].GetEvents();
-		FBarMap& BarMap = SongMaps.GetBarMap();
-		BarMap.SetStartBar(1);
+		SongMaps.SetStartBar(1);
 		for (auto& Event : Events)
 		{
 			if (Event.GetMsg().MsgType() == FMidiMsg::EType::TimeSig)
 			{
 				int32 Tick = Event.GetTick();
 				int32 BarIndex;
-				check(Tick == 0 || BarMap.GetNumTimeSignaturePoints() > 0);
+				check(Tick == 0 || SongMaps.GetNumTimeSignatureChanges() > 0);
 				if (Tick == 0)
 				{
 					BarIndex = 0;
 				}
 				else
 				{
-					BarIndex = BarMap.TickToBarIncludingCountIn(Tick);
+					BarIndex = SongMaps.TickToBarIncludingCountIn(Tick);
 				}
-				BarMap.AddTimeSignatureAtBarIncludingCountIn(BarIndex, Event.GetMsg().GetTimeSigNumerator(), Event.GetMsg().GetTimeSigDenominator());
+				SongMaps.AddTimeSignatureAtBarIncludingCountIn(BarIndex, Event.GetMsg().GetTimeSigNumerator(), Event.GetMsg().GetTimeSigDenominator());
 			}
 		}
-		BarMap.Finalize(GetLastEventTick());
+		SongMaps.FinalizeBarMap(GetLastEventTick());
 	}
 }
 
@@ -891,28 +895,25 @@ void FMidiFileData::AddTempoChange(int32 TrackIdx, int32 Tick, float TempoBPM)
 {
 	check(Tracks.IsValidIndex(TrackIdx));
 
-	FTempoMap& TempoMap = SongMaps.GetTempoMap();
 	int32 MidiTempo = Harmonix::Midi::Constants::BPMToMidiTempo(TempoBPM);
 	Tracks[TrackIdx].AddEvent(FMidiEvent(Tick, FMidiMsg(MidiTempo)));
 	Tracks[TrackIdx].Sort();
-	TempoMap.AddTempoInfoPoint(MidiTempo, Tick);
+	SongMaps.AddTempoInfoPoint(MidiTempo, Tick);
 }
 
 void FMidiFileData::AddTimeSigChange(int32 TrackIdx, int32 Tick, int32 InTimeSigNum, int32 InTimeSigDenom)
 {
 	check(Tracks.IsValidIndex(TrackIdx));
 
-	FBarMap& BarMap = SongMaps.GetBarMap();
-
 	// Time signature changes can only happen at the beginning of a bar,
 	// so round up to the next bar boundary...
 	int32 AbsoluteBar = FMath::CeilToInt32(SongMaps.GetBarIncludingCountInAtTick(Tick));
-	Tick = BarMap.BarBeatTickIncludingCountInToTick(AbsoluteBar, 1, 0);
+	Tick = SongMaps.BarBeatTickIncludingCountInToTick(AbsoluteBar, 1, 0);
 	int32 TimeSigNum = FMath::Clamp(InTimeSigNum, 1, 64);
 	int32 TimeSigDenom = FMath::Clamp(InTimeSigDenom, 1, 64);
 	Tracks[TrackIdx].AddEvent(FMidiEvent(Tick, FMidiMsg((uint8)TimeSigNum, (uint8)TimeSigDenom)));
 	Tracks[TrackIdx].Sort();
-	BarMap.AddTimeSignatureAtBarIncludingCountIn(AbsoluteBar, TimeSigNum, TimeSigDenom);
+	SongMaps.AddTimeSignatureAtBarIncludingCountIn(AbsoluteBar, TimeSigNum, TimeSigDenom);
 }
 
 void FMidiFileData::ScanTracksForSongLengthChange()
@@ -938,7 +939,7 @@ void FMidiFileData::ScanTracksForSongLengthChange()
 	//update length data in song length data if needed 
 	LengthData.LengthTicks = NewLastEventTick + 1;
 	LengthData.LastTick = NewLastEventTick;
-	LengthData.LengthFractionalBars = SongMaps.GetBarMap().TickToFractionalBarIncludingCountIn(LengthData.LengthTicks);
+	LengthData.LengthFractionalBars = SongMaps.TickToFractionalBarIncludingCountIn(LengthData.LengthTicks);
 }
 
 bool FMidiFileData::LengthIsAPerfectSubdivision() const
diff --git a/Engine/Plugins/Runtime/Harmonix/Source/HarmonixMidi/Private/HarmonixMidi/MidiPlayCursor.cpp b/Engine/Plugins/Runtime/Harmonix/Source/HarmonixMidi/Private/HarmonixMidi/MidiPlayCursor.cpp
--- a/Engine/Plugins/Runtime/Harmonix/Source/HarmonixMidi/Private/HarmonixMidi/MidiPlayCursor.cpp
+++ b/Engine/Plugins/Runtime/Harmonix/Source/HarmonixMidi/Private/HarmonixMidi/MidiPlayCursor.cpp
@@ -50,8 +50,8 @@ void FMidiPlayCursor::SetOwner(FMidiPlayCursorMgr* NewOwner, FMidiPlayCursorTrac
 			{
 				PreRollMs = Tracker->CurrentMs;
 			}
-			int32 Tick = (int32)Owner->GetTempoMap().MsToTick(Tracker->CurrentMs - PreRollMs);
-			float Ms = Owner->GetTempoMap().TickToMs(Tick);
+			int32 Tick = (int32)Owner->GetSongMaps().MsToTick(Tracker->CurrentMs - PreRollMs);
+			float Ms = Owner->GetSongMaps().TickToMs(Tick);
 			Tracker->ResetNewCursor(this, Tick, Ms, false);
 		}
 	}
@@ -63,7 +63,7 @@ void FMidiPlayCursor::Reset(bool ForceNoBroadcast)
 	LoopCount = 0;
 	if (GetOwner())
 	{
-		CurrentMs = Owner->GetTempoMap().TickToMs(CurrentTick);
+		CurrentMs = Owner->GetSongMaps().TickToMs(CurrentTick);
 		TrackNextEventIndexs.SetNumUninitialized(Owner->Tracks().Num());
 
 		for (int32 i = 0; i < TrackNextEventIndexs.Num(); i++)
@@ -122,11 +122,11 @@ void FMidiPlayCursor::RecalcNextEventsDueToMidiChanges(FMidiPlayCursorMgr::EMidi
 
 	if (PositionMode == FMidiPlayCursorMgr::EMidiChangePositionCorrectMode::MaintainTick)
 	{
-		CurrentMs = Owner->GetTempoMap().TickToMs(CurrentTick);
+		CurrentMs = Owner->GetSongMaps().TickToMs(CurrentTick);
 	}
 	else // if (positionMode == MidiChangePositionCorrectMode::MaintainTime)
 	{
-		CurrentTick = Owner->GetTempoMap().MsToTick(CurrentMs);
+		CurrentTick = Owner->GetSongMaps().MsToTick(CurrentMs);
 	}
 
 	// changes to tracks might actually be more complicated...
@@ -236,11 +236,11 @@ void FMidiPlayCursor::PrepareLookAheadMs(bool ForceNoBroadcast)
 		float NewMs = Tracker->CurrentMs + LookaheadMs;
 		if (NewMs > 0.0f)
 		{
-			CurrentTick = (int32)(GetOwner()->GetTempoMap().MsToTick(NewMs) + 0.5f);
+			CurrentTick = (int32)(GetOwner()->GetSongMaps().MsToTick(NewMs) + 0.5f);
 		}
 		else
 		{
-			CurrentTick = (int32)(GetOwner()->GetTempoMap().MsToTick(NewMs) - 0.5f);
+			CurrentTick = (int32)(GetOwner()->GetSongMaps().MsToTick(NewMs) - 0.5f);
 		}
 		SeekThruTick(CurrentTick);
 		CurrentMs = NewMs;
@@ -543,7 +543,7 @@ void FMidiPlayCursor::AdvanceByMs(bool ProcessLoops, bool Broadcast, bool IsPreR
 void FMidiPlayCursor::DoAdvanceForLaggingMsCursor(bool Broadcast, bool IsPreRoll)
 {
 	float NewMs = Tracker->CurrentMs + LookaheadMs;
-	int32 NewTick = (int32)(GetOwner()->GetTempoMap().MsToTick(NewMs) + 0.5f);
+	int32 NewTick = (int32)(GetOwner()->GetSongMaps().MsToTick(NewMs) + 0.5f);
 	if (Owner->DoesLoop(Tracker->IsLowRes))
 	{
 		float LoopStartMs = Owner->GetLoopStartMs(Tracker->IsLowRes);
@@ -572,7 +572,7 @@ void FMidiPlayCursor::DoAdvanceForLaggingMsCursor(bool Broadcast, bool IsPreRoll
 			{
 				// wrap new tick back to loop end...
 				NewMs = LoopEndMs - (LoopStartMs - NewMs);
-				NewTick = (int32)(Owner->GetTempoMap().MsToTick(NewMs) + 0.5f);
+				NewTick = (int32)(Owner->GetSongMaps().MsToTick(NewMs) + 0.5f);
 				// now advance from current position to new position...
 				if (Broadcast)
 				{
@@ -626,7 +626,7 @@ void FMidiPlayCursor::DoAdvanceForLaggingMsCursor(bool Broadcast, bool IsPreRoll
 				{
 					// wrap new tick back from loopEnd...
 					NewMs = LoopEndMs - (LoopStartMs - NewMs);
-					NewTick = (int32)(Owner->GetTempoMap().MsToTick(NewMs) + 0.5f);
+					NewTick = (int32)(Owner->GetSongMaps().MsToTick(NewMs) + 0.5f);
 					if (Broadcast)
 					{
 						AdvanceThruTick(NewTick, IsPreRoll);
@@ -677,7 +677,7 @@ void FMidiPlayCursor::DoAdvanceForLaggingMsCursor(bool Broadcast, bool IsPreRoll
 void FMidiPlayCursor::DoAdvanceForLeadingMsCursor(bool Broadcast, bool ProcessLoops, bool IsPreRoll)
 {
 	float NewMs = Tracker->CurrentMs + LookaheadMs;
-	int32 NewTick = (int32)(GetOwner()->GetTempoMap().MsToTick(NewMs) + 0.5f);
+	int32 NewTick = (int32)(GetOwner()->GetSongMaps().MsToTick(NewMs) + 0.5f);
 	if (Owner->DoesLoop(Tracker->IsLowRes) && ProcessLoops) 
 	{
 		float LoopStartMs = Owner->GetLoopStartMs(Tracker->IsLowRes);
@@ -694,7 +694,7 @@ void FMidiPlayCursor::DoAdvanceForLeadingMsCursor(bool Broadcast, bool ProcessLo
 			{
 				// wrap NewTick...
 				NewMs = LoopStartMs + (NewMs - LoopEndMs);
-				NewTick = (int32)(Owner->GetTempoMap().MsToTick(NewMs) + 0.5f);
+				NewTick = (int32)(Owner->GetSongMaps().MsToTick(NewMs) + 0.5f);
 			}
 			if (NewTick < CurrentTick)
 			{
@@ -749,7 +749,7 @@ void FMidiPlayCursor::DoAdvanceForLeadingMsCursor(bool Broadcast, bool ProcessLo
 				{
 					// jump NewTick...
 					NewMs = LoopStartMs + (NewMs - LoopEndMs);
-					NewTick = (int32)(Owner->GetTempoMap().MsToTick(NewMs) + 0.5f);
+					NewTick = (int32)(Owner->GetSongMaps().MsToTick(NewMs) + 0.5f);
 				}
 
 				if (CurrentTick < LoopEndTick)
@@ -818,7 +818,7 @@ void FMidiPlayCursor::SeekToTick(int32 Tick)
 		}
 	}
 	CurrentTick = Tick - 1;
-	CurrentMs = Owner->GetTempoMap().TickToMs(CurrentTick);
+	CurrentMs = Owner->GetSongMaps().TickToMs(CurrentTick);
 }
 
 void FMidiPlayCursor::SeekThruTick(int32 Tick)
@@ -845,7 +845,7 @@ void FMidiPlayCursor::SeekThruTick(int32 Tick)
 		}
 	}
 	CurrentTick = Tick;
-	CurrentMs = Owner->GetTempoMap().TickToMs(Tick);
+	CurrentMs = Owner->GetSongMaps().TickToMs(Tick);
 }
 
 bool FMidiPlayCursor::IsDone() const
@@ -1159,8 +1159,8 @@ void FMidiPlayCursor::AdvanceThruTick(int32 Tick, bool IsPreRoll)
 	}
 
 	CurrentTick = Tick;
-	const FTempoMap& TempoMap = Owner->GetTempoMap();
-	CurrentMs = TempoMap.TickToMs(Tick);
+	const ISongMapEvaluator& Maps = Owner->GetSongMaps();
+	CurrentMs = Maps.TickToMs(Tick);
 
 	// now we MAY have note on events which still need to be broadcast...
 	if ((FilterPassFlags & EFilterPassFlags::PreRollNoteOn) != EFilterPassFlags::None)
@@ -1170,7 +1170,7 @@ void FMidiPlayCursor::AdvanceThruTick(int32 Tick, bool IsPreRoll)
 			{
 				const FMidiMsg& Msg = Event->GetMsg();
 				int32 EventTick = Event->GetTick();
-				float EventMs = TempoMap.TickToMs(EventTick);
+				float EventMs = Maps.TickToMs(EventTick);
 				OnPreRollNoteOn(TrackIndex, EventTick, CurrentTick, CurrentMs - EventMs, Msg.GetStdStatus(), Msg.GetStdData1(), Msg.GetStdData2());
 			});
 	}
diff --git a/Engine/Plugins/Runtime/Harmonix/Source/HarmonixMidi/Private/HarmonixMidi/MidiPlayCursorMgr.cpp b/Engine/Plugins/Runtime/Harmonix/Source/HarmonixMidi/Private/HarmonixMidi/MidiPlayCursorMgr.cpp
--- a/Engine/Plugins/Runtime/Harmonix/Source/HarmonixMidi/Private/HarmonixMidi/MidiPlayCursorMgr.cpp
+++ b/Engine/Plugins/Runtime/Harmonix/Source/HarmonixMidi/Private/HarmonixMidi/MidiPlayCursorMgr.cpp
@@ -18,8 +18,8 @@ FMidiPlayCursorMgr::FMidiPlayCursorMgr()
 	, InMidiChangeLock(false)
 {
 	// setup the default tempo map to have one entry...
-	DefaultMaps.GetTempoMap().AddTempoInfoPoint(Harmonix::Midi::Constants::BPMToMidiTempo(120.0f), 0);
-	DefaultMaps.GetBarMap().AddTimeSignatureAtBarIncludingCountIn(0, 4, 4);
+	DefaultMaps.AddTempoInfoPoint(Harmonix::Midi::Constants::BPMToMidiTempo(120.0f), 0);
+	DefaultMaps.AddTimeSignatureAtBarIncludingCountIn(0, 4, 4);
 }
 
 void FMidiPlayCursorMgr::Reset()
@@ -145,9 +145,9 @@ void FMidiPlayCursorMgr::DetermineLength()
 			LengthTicks = FMath::Max(LengthTicks, Tracks[i].GetEvents().Last().GetTick());
 		}
 		//Round file length up to the nearest bar
-		int32 Bar = FMath::CeilToInt32(SongMaps->GetBarMap().TickToFractionalBarIncludingCountIn(LengthTicks));
-		LengthTicks = SongMaps->GetBarMap().BarIncludingCountInToTick(Bar);
-		LengthMs = SongMaps->GetTempoMap().TickToMs(LengthTicks);
+		int32 Bar = FMath::CeilToInt32(SongMaps->TickToFractionalBarIncludingCountIn(LengthTicks));
+		LengthTicks = SongMaps->BarIncludingCountInToTick(Bar);
+		LengthMs = SongMaps->TickToMs(LengthTicks);
 
 		for (FMidiPlayCursorTracker& Tracker : Trackers)
 		{
@@ -193,7 +193,7 @@ void FMidiPlayCursorMgr::RecalculatePreRollDueToCursorPosition(FMidiPlayCursor*
 	float PreRollMs = 0.0f;
 	if (PlayCursor->GetLookaheadType() == FMidiPlayCursor::ELookaheadType::Ticks)
 	{
-		PreRollMs = -GetTempoMap().TickToMs(-PlayCursor->GetLookaheadTicks());
+		PreRollMs = -GetSongMaps().TickToMs(-PlayCursor->GetLookaheadTicks());
 	}
 	else
 	{
@@ -202,11 +202,11 @@ void FMidiPlayCursorMgr::RecalculatePreRollDueToCursorPosition(FMidiPlayCursor*
 	if (-PreRollMs < GetHiResTracker().CurrentMs)
 	{
 		// yup... earliest look ahead!
-		int32 NewTick = GetTempoMap().MsToTick(-PreRollMs);
+		int32 NewTick = GetSongMaps().MsToTick(-PreRollMs);
 
 		// back up one tick...
 		NewTick--;
-		PreRollMs = GetTempoMap().TickToMs(NewTick);
+		PreRollMs = GetSongMaps().TickToMs(NewTick);
 		for (FMidiPlayCursorTracker& Tracker : Trackers)
 		{
 			Tracker.Reset(NewTick, PreRollMs, false);
@@ -280,7 +280,7 @@ void FMidiPlayCursorMgr::ResetTrackers()
 {
 	FMidiPlayCursorListLock LowResCursorListLock(GetLowResTracker().CursorListCS);
 	FMidiPlayCursorListLock HiResCursorListLock(GetHiResTracker().CursorListCS);
-	float Ms = GetTempoMap().TickToMs(-1);
+	float Ms = GetSongMaps().TickToMs(-1);
 	for (FMidiPlayCursorTracker& Tracker : Trackers)
 	{
 		Tracker.Reset(-1, Ms, false);
@@ -314,24 +314,12 @@ const UMidiFile::FMidiTrackList& FMidiPlayCursorMgr::Tracks() const
 	return MidiFileData.IsValid() ? MidiFileData->Tracks : DefaultTracks;
 }
 
-const FSongMaps& FMidiPlayCursorMgr::GetSongMaps() const
+const ISongMapEvaluator& FMidiPlayCursorMgr::GetSongMaps() const
 {
 	check(SongMaps);
 	return *SongMaps;
 }
 
-const FTempoMap& FMidiPlayCursorMgr::GetTempoMap() const
-{
-	check(SongMaps);
-	return SongMaps->GetTempoMap();
-}
-
-const FBarMap& FMidiPlayCursorMgr::GetBarMap() const
-{
-	check(SongMaps);
-	return SongMaps->GetBarMap();
-}
-
 void FMidiPlayCursorMgr::GetCursorExtentsMs(float& Earliest, float& Latest) const
 {
 	Earliest = (GetHiResTracker().EarliestCursorMs > GetLowResTracker().EarliestCursorMs) ?
@@ -428,15 +416,14 @@ void FMidiPlayCursorMgr::SetLoopImpl(int32 StartTick, int32 EndTick, bool Ignori
 	// might fall in the middle of the current span of leading and lagging play cursors,
 	// which would result in ugly behavior!
 
-	const FTempoMap& TempoMap = GetTempoMap();
 	Trackers[IsLowRes].LoopStartTick = StartTick;
-	Trackers[IsLowRes].LoopStartMs = TempoMap.TickToMs(StartTick);
+	Trackers[IsLowRes].LoopStartMs = GetSongMaps().TickToMs(StartTick);
 	if (EndTick == kEndTick)
 	{
 		EndTick = LengthTicks;
 	}
 	Trackers[IsLowRes].LoopEndTick = EndTick;
-	Trackers[IsLowRes].LoopEndMs = TempoMap.TickToMs(EndTick);
+	Trackers[IsLowRes].LoopEndMs = GetSongMaps().TickToMs(EndTick);
 	Trackers[IsLowRes].LoopIgnoringLookAhead = IgnoringLookAhead;
 	Trackers[IsLowRes].Loop = true;
 	bool CursorsInPhase = CursorsAllInPhaseImpl(IsLowRes);
@@ -479,13 +466,12 @@ void FMidiPlayCursorMgr::SeekTo(int32 Tick, int32 PreRollBars, bool IsRenderThre
 	// Back up one tick, as we are going to be setting the cursors'
 	// 'played through' position.
 	Tick--;
-	float NewPosMs = GetTempoMap().TickToMs(Tick);
+	float NewPosMs = GetSongMaps().TickToMs(Tick);
 	float PreRollStartMs = NewPosMs;
 	int32 PreRollStartTick = Tick;
 	if (PreRollBars > 0)
 	{
-		const FBarMap& Map = GetBarMap();
-		float Bar = Map.TickToFractionalBarIncludingCountIn(Tick);
+		float Bar = GetSongMaps().TickToFractionalBarIncludingCountIn(Tick);
 		Bar -= (float)PreRollBars;
 
 		if (Bar < 0.0f)
@@ -494,10 +480,10 @@ void FMidiPlayCursorMgr::SeekTo(int32 Tick, int32 PreRollBars, bool IsRenderThre
 		}
 		else
 		{
-			PreRollStartTick = Map.BarIncludingCountInToTick(Bar);
+			PreRollStartTick = GetSongMaps().BarIncludingCountInToTick(Bar);
 		}
 
-		PreRollStartMs = GetTempoMap().TickToMs(PreRollStartTick);
+		PreRollStartMs = GetSongMaps().TickToMs(PreRollStartTick);
 
 		if (PreRollStartTick > Tick)
 		{
@@ -531,10 +517,8 @@ void FMidiPlayCursorMgr::MoveToLoopStart()
 	{
 		if (ensureMsgf(GetHiResTracker().Loop, TEXT("That's odd. Asked to move to the beginning of the loop... but there is no loop!")))
 		{
-			const FTempoMap& TempoMap = GetTempoMap();
-
 			int32 NewThruTick = GetHiResTracker().LoopStartTick - 1;
-			float NewThruMs = TempoMap.TickToMs(NewThruTick);
+			float NewThruMs = GetSongMaps().TickToMs(NewThruTick);
 
 			// Move the hi-res cursor to the loop start
 			GetHiResTracker().MoveToLoopStart(NewThruTick, NewThruMs);
@@ -609,8 +593,8 @@ void FMidiPlayCursorMgr::MidiDataChangeComplete(EMidiChangePositionCorrectMode P
 		// Now fix up looping...
 		int32 NewLoopStartTick = OriginalLoopStartTick;
 		int32 NewLoopEndTick = FMath::Min(GetHiResTracker().LoopEndTick, OriginalLoopEndTick);
-		float NewLoopStartMs = GetTempoMap().TickToMs(NewLoopStartTick);
-		float NewLoopEndMs = GetTempoMap().TickToMs(NewLoopEndTick);
+		float NewLoopStartMs = GetSongMaps().TickToMs(NewLoopStartTick);
+		float NewLoopEndMs = GetSongMaps().TickToMs(NewLoopEndTick);
 		for (FMidiPlayCursorTracker& Tracker : Trackers)
 		{
 			Tracker.LoopStartTick = NewLoopStartTick;
@@ -629,11 +613,11 @@ void FMidiPlayCursorMgr::MidiDataChangeComplete(EMidiChangePositionCorrectMode P
 			}
 			if (PositionMode == EMidiChangePositionCorrectMode::MaintainTick)
 			{
-				GetHiResTracker().CurrentMs = GetTempoMap().TickToMs(GetHiResTracker().CurrentTick);
+				GetHiResTracker().CurrentMs = GetSongMaps().TickToMs(GetHiResTracker().CurrentTick);
 			}
 			else
 			{
-				GetHiResTracker().CurrentTick = GetTempoMap().MsToTick(GetHiResTracker().CurrentMs);
+				GetHiResTracker().CurrentTick = GetSongMaps().MsToTick(GetHiResTracker().CurrentMs);
 			}
 			GetHiResTracker().TraversingCursors = false;
 			// Now the low res cursors...
@@ -644,11 +628,11 @@ void FMidiPlayCursorMgr::MidiDataChangeComplete(EMidiChangePositionCorrectMode P
 			}
 			if (PositionMode == EMidiChangePositionCorrectMode::MaintainTick)
 			{
-				GetLowResTracker().CurrentMs = GetTempoMap().TickToMs(GetLowResTracker().CurrentTick);
+				GetLowResTracker().CurrentMs = GetSongMaps().TickToMs(GetLowResTracker().CurrentTick);
 			}
 			else
 			{
-				GetLowResTracker().CurrentTick = GetTempoMap().MsToTick(GetLowResTracker().CurrentMs);
+				GetLowResTracker().CurrentTick = GetSongMaps().MsToTick(GetLowResTracker().CurrentMs);
 			}
 			// Now deal with the possibility that the low res cursors might be "out of phase" with the hi res
 			// cursors. This would happen if we are in the middle of a loop...
@@ -670,7 +654,7 @@ void FMidiPlayCursorMgr::MidiDataChangeComplete(EMidiChangePositionCorrectMode P
 		}
 		else
 		{
-			int32 SeekTick = PositionMode == EMidiChangePositionCorrectMode::MaintainTick ? GetHiResTracker().CurrentTick : GetTempoMap().MsToTick(GetHiResTracker().CurrentMs);
+			int32 SeekTick = PositionMode == EMidiChangePositionCorrectMode::MaintainTick ? GetHiResTracker().CurrentTick : GetSongMaps().MsToTick(GetHiResTracker().CurrentMs);
 			SeekTo(SeekTick, PreRollBars, true, false);
 		}
 		GetHiResTracker().CursorListCS.Unlock();
@@ -682,7 +666,7 @@ void FMidiPlayCursorMgr::MidiDataChangeComplete(EMidiChangePositionCorrectMode P
 FMusicTimestamp FMidiPlayCursorMgr::GetMusicTimestampAtMs(float Ms) const
 {
 	const int32 TickAtOffset = GetSongMaps().MsToTick(Ms);
-	return GetBarMap().TickToMusicTimestamp(TickAtOffset);
+	return GetSongMaps().TickToMusicTimestamp(TickAtOffset);
 }
 
 void FMidiPlayCursorMgr::AdvanceHiResToMs(float Ms, bool Broadcast)
@@ -824,16 +808,15 @@ void FMidiPlayCursorMgr::UpdateLowResCursors(FMidiPlayCursorTracker& Tracker)
 bool FMidiPlayCursorMgr::AdvanceTrackerByDeltaMs(float Ms, FMidiPlayCursorTracker& Tracker, bool IsLowRes, bool Broadcast)
 {
 	Tracker.ElapsedMs += Ms;
-	const FTempoMap& TempoMap = GetTempoMap();
 	float NewMs = Tracker.CurrentMs + Ms;
-	int32 NewTick = (int32)(TempoMap.MsToTick(NewMs) + 0.5f);
+	int32 NewTick = (int32)(GetSongMaps().MsToTick(NewMs) + 0.5f);
 	bool  Looped = false;
 
 	if (NewTick >= Tracker.LoopEndTick && Tracker.CurrentTick < Tracker.LoopEndTick && Tracker.Loop)
 	{
 		// loop the tick around...
 		NewTick = (NewTick - Tracker.LoopEndTick) + Tracker.LoopStartTick;
-		NewMs = TempoMap.TickToMs(NewTick);
+		NewMs = GetSongMaps().TickToMs(NewTick);
 		Tracker.CurrentMs = NewMs;
 		Tracker.CurrentTick = NewTick;
 		Looped = true;
@@ -866,9 +849,7 @@ bool FMidiPlayCursorMgr::AdvanceTrackerByDeltaMs(float Ms, FMidiPlayCursorTracke
 
 void FMidiPlayCursorMgr::AdvanceTrackerThruTick(int32 ToTick, FMidiPlayCursorTracker& Tracker, bool IsLowRes, bool Broadcast)
 {
-	const FTempoMap& TempoMap = GetTempoMap();
-
-	float NewMs = TempoMap.TickToMs(ToTick);
+	float NewMs = GetSongMaps().TickToMs(ToTick);
 	Tracker.ElapsedMs += NewMs - Tracker.CurrentMs;
 	int32   NewTick = ToTick;
 
diff --git a/Engine/Plugins/Runtime/Harmonix/Source/HarmonixMidi/Private/HarmonixMidi/MidiSongPos.cpp b/Engine/Plugins/Runtime/Harmonix/Source/HarmonixMidi/Private/HarmonixMidi/MidiSongPos.cpp
--- a/Engine/Plugins/Runtime/Harmonix/Source/HarmonixMidi/Private/HarmonixMidi/MidiSongPos.cpp
+++ b/Engine/Plugins/Runtime/Harmonix/Source/HarmonixMidi/Private/HarmonixMidi/MidiSongPos.cpp
@@ -57,6 +57,34 @@ bool FMidiSongPos::operator>=(const FMidiSongPos& rhs) const
 	return cmp == 1 || cmp == 0;
 }
 
+FMidiSongPos FMidiSongPos::Lerp(const FMidiSongPos& A, const FMidiSongPos& B, float Alpha)
+{
+	FMidiSongPos Result;
+	Result.SecondsIncludingCountIn = FMath::Lerp(A.SecondsIncludingCountIn, B.SecondsIncludingCountIn, Alpha);
+	Result.SecondsFromBarOne = FMath::Lerp(A.SecondsFromBarOne, B.SecondsFromBarOne, Alpha);
+	Result.TimeSigNumerator = A.TimeSigNumerator;
+	Result.TimeSigDenominator = A.TimeSigDenominator;
+	Result.Tempo = A.Tempo;
+	Result.CurrentSongSection = (Alpha < 0.5f) ? A.CurrentSongSection : B.CurrentSongSection;
+	Result.BarsIncludingCountIn = FMath::Lerp(A.BarsIncludingCountIn, B.BarsIncludingCountIn, Alpha);
+	Result.BeatsIncludingCountIn = FMath::Lerp(A.BeatsIncludingCountIn, B.BeatsIncludingCountIn, Alpha);
+	
+	// infer start bar...
+	int32 StartBar = A.Timestamp.Bar - FMath::FloorToInt32(A.BarsIncludingCountIn);
+	
+	Result.Timestamp.Bar = FMath::FloorToInt32(Result.BarsIncludingCountIn) + StartBar;
+	Result.Timestamp.Beat = FMath::Fractional(Result.BarsIncludingCountIn) * Result.TimeSigNumerator + 1.0f;
+	if (FMath::FloorToInt32(Result.BeatsIncludingCountIn) == FMath::FloorToInt32(A.BeatsIncludingCountIn))
+	{
+		Result.BeatType = A.BeatType;
+	}
+	else
+	{
+		Result.BeatType = B.BeatType;
+	}
+	return Result;
+}
+
 bool FMidiSongPos::operator==(const FMidiSongPos& rhs) const
 {
 	return SongPosCmp(*this, rhs) == 0;
@@ -83,56 +111,48 @@ void FMidiSongPos::SetByTime(float InElapsedMs, float InBpm, int32 InTimeSigNume
 	Timestamp.Beat = FMath::Fractional(BarsIncludingCountIn) * InTimeSigNumerator + 1.0f;
 }
 
-void FMidiSongPos::SetByTime(float InMs, const FSongMaps& Maps)
+void FMidiSongPos::SetByTime(float InMs, const ISongMapEvaluator& Map)
 {
-	float Tick = Maps.GetTempoMap().MsToTick(InMs);
-	SetByTimeAndTick(InMs, Tick, Maps);
+	float Tick = Map.MsToTick(InMs);
+	SetByTimeAndTick(InMs, Tick, Map);
 }
 
-void FMidiSongPos::SetByTick(float InTick, const FSongMaps& Maps)
+void FMidiSongPos::SetByTick(float InTick, const ISongMapEvaluator& Map)
 {
-	float Ms = Maps.GetTempoMap().TickToMs(InTick);
-	SetByTimeAndTick(Ms, InTick, Maps);
+	float Ms = Map.TickToMs(InTick);
+	SetByTimeAndTick(Ms, InTick, Map);
 }
 
-void FMidiSongPos::SetByTimeAndTick(float InMs, float InTick, const FSongMaps& Maps)
+void FMidiSongPos::SetByTimeAndTick(float InMs, float InTick, const ISongMapEvaluator& Map)
 {
-	const FTimeSignature* TimeSig = Maps.GetTimeSignatureAtTick(int32(InTick));
+	const FTimeSignature* TimeSig = Map.GetTimeSignatureAtTick(int32(InTick));
 	SecondsIncludingCountIn = InMs / 1000.0f;
-	SecondsFromBarOne       = SecondsIncludingCountIn - Maps.GetCountInSeconds();
+	SecondsFromBarOne       = SecondsIncludingCountIn - Map.GetCountInSeconds();
 	TimeSigNumerator        = TimeSig ? TimeSig->Numerator : 4;
 	TimeSigDenominator      = TimeSig ? TimeSig->Denominator : 4;
-	Tempo                   = Maps.GetTempoAtTick(int32(InTick));
-	BarsIncludingCountIn    = Maps.GetBarIncludingCountInAtTick(InTick);
+	Tempo                   = Map.GetTempoAtTick(int32(InTick));
+	BarsIncludingCountIn    = Map.GetBarIncludingCountInAtTick(InTick);
 
-	const FSongSection* SongSection = !Maps.GetSectionMap().IsEmpty() ? Maps.GetSectionAtTick(int32(InTick)) : nullptr;
-	CurrentSongSection		= SongSection ? FSongSection(SongSection->Name, SongSection->StartTick, SongSection->LengthTicks) : FSongSection();
+	const FSongSection* SongSection = Map.GetSectionAtTick(int32(InTick));
+	CurrentSongSection = SongSection ? FSongSection(SongSection->Name, SongSection->StartTick, SongSection->LengthTicks) : FSongSection();
 	
-	const FBeatMap& BeatMap = Maps.GetBeatMap();
-	int32 BeatPointIndex     = BeatMap.GetPointIndexForTick(int32(InTick));
-	if (BeatPointIndex >= 0)
+	int32 BeatPointIndex = 0;
+	const FBeatMapPoint* BeatPoint = Map.GetBeatPointInfoAtTick(int32(InTick), &BeatPointIndex);
+	if (BeatPoint != nullptr)
 	{
-		const FBeatMapPoint& BeatPoint = BeatMap.GetBeatPointInfo(BeatPointIndex);
-		float TickInBeat     = InTick - BeatPoint.StartTick;
-		float FractionalPart = TickInBeat / BeatPoint.LengthTicks;
-		BeatsIncludingCountIn     = BeatPointIndex + FractionalPart;
-		BeatType             = BeatPoint.Type;
+		float TickInBeat = InTick - BeatPoint->StartTick;
+		float FractionalPart = TickInBeat / BeatPoint->LengthTicks;
+		BeatsIncludingCountIn = BeatPointIndex + FractionalPart;
+		BeatType = BeatPoint->Type;
 	}
 	else
 	{
-		// use basic time signature based info from bar map...
-		const FBarMap& BarMap = Maps.GetBarMap();
-		int32 BarPointIndex = BarMap.GetPointIndexForTick(int32(InTick));
-		if (BarPointIndex == INDEX_NONE)
-		{
-			ensure(InTick < 0.0f);
-			BarPointIndex = 0;
-		}
-		const FTimeSignaturePoint& TimeSigPoint = BarMap.GetTimeSignaturePoint(BarPointIndex);
-		float BarAtTimeSig = BarsIncludingCountIn - TimeSigPoint.BarIndex;
+		const FTimeSignaturePoint* TimeSigPoint = Map.GetTimeSignaturePointAtTick(int32(InTick));
+		check (TimeSigPoint);
+		float BarAtTimeSig = BarsIncludingCountIn - TimeSigPoint->BarIndex;
 		float FractPart = FMath::Fractional(BarAtTimeSig);
-		BeatsIncludingCountIn = TimeSigPoint.BeatIndex + (BarAtTimeSig * TimeSigPoint.TimeSignature.Numerator);
+		BeatsIncludingCountIn = TimeSigPoint->BeatIndex + (BarAtTimeSig * TimeSigPoint->TimeSignature.Numerator);
 		BeatType = FMath::IsNearlyEqual(Timestamp.Beat, 1.0f) ? EMusicalBeatType::Downbeat : EMusicalBeatType::Normal;
 	}
-	Timestamp = Maps.GetBarMap().TickToMusicTimestamp(InTick);
+	Timestamp = Map.TickToMusicTimestamp(InTick);
 }
