Massive Midi Clock Refactor

diff --git a/Engine/Plugins/Runtime/Harmonix/Source/HarmonixMidi/Private/HarmonixMidi/MusicTimeSpan.cpp b/Engine/Plugins/Runtime/Harmonix/Source/HarmonixMidi/Private/HarmonixMidi/MusicTimeSpan.cpp
--- a/Engine/Plugins/Runtime/Harmonix/Source/HarmonixMidi/Private/HarmonixMidi/MusicTimeSpan.cpp
+++ b/Engine/Plugins/Runtime/Harmonix/Source/HarmonixMidi/Private/HarmonixMidi/MusicTimeSpan.cpp
@@ -5,7 +5,7 @@
 #include "HarmonixMidi/MidiSongPos.h"
 #include "HarmonixMidi/MidiConstants.h"
 
-float FMusicalTimeSpan::CalcPositionInSpan(const FMidiSongPos& Position, const FSongMaps& Maps) const
+float FMusicalTimeSpan::CalcPositionInSpan(const FMidiSongPos& Position, const ISongMapEvaluator& Maps) const
 {
 	if (Position.TimeSigDenominator == 0)
 	{
@@ -20,7 +20,7 @@ float FMusicalTimeSpan::CalcPositionInSpan(const FMidiSongPos& Position, const F
 	return CalcPositionInSpanNoOffset(Position, Maps);
 }
 
-float FMusicalTimeSpan::CalcPositionInSpanWithOffset(const FMidiSongPos& Position, const FSongMaps& Maps) const
+float FMusicalTimeSpan::CalcPositionInSpanWithOffset(const FMidiSongPos& Position, const ISongMapEvaluator& Maps) const
 {
 	using namespace Harmonix::Midi::Constants;
 	FMidiSongPos OffsetPosition;
@@ -36,8 +36,8 @@ float FMusicalTimeSpan::CalcPositionInSpanWithOffset(const FMidiSongPos& Positio
 		break;
 	case EMusicTimeSpanOffsetUnits::Bars:
 		{
-			float FractionalBar = Position.BarsIncludingCountIn + Maps.GetBarMap().GetStartBar() - Offset;
-			int32 Tick = Maps.GetBarMap().FractionalBarIncludingCountInToTick(FractionalBar);
+			float FractionalBar = Position.BarsIncludingCountIn + Maps.GetStartBar() - Offset;
+			int32 Tick = Maps.FractionalBarIncludingCountInToTick(FractionalBar);
 			OffsetPosition.SetByTick((float)Tick, Maps);
 		}
 		break;
@@ -105,7 +105,7 @@ float FMusicalTimeSpan::CalcPositionInSpanWithOffset(const FMidiSongPos& Positio
 	return CalcPositionInSpanNoOffset(OffsetPosition, Maps);
 }
 
-float FMusicalTimeSpan::CalcPositionInSpanNoOffset(const FMidiSongPos& Position, const FSongMaps& Maps) const
+float FMusicalTimeSpan::CalcPositionInSpanNoOffset(const FMidiSongPos& Position, const ISongMapEvaluator& Maps) const
 {
 	if (LengthUnits == EMusicTimeSpanLengthUnits::Bars || LengthUnits == EMusicTimeSpanLengthUnits::Beats)
 	{
@@ -115,14 +115,14 @@ float FMusicalTimeSpan::CalcPositionInSpanNoOffset(const FMidiSongPos& Position,
 }
 
 
-float FMusicalTimeSpan::CalcPositionInSpan(float Ms, const FSongMaps& Maps) const
+float FMusicalTimeSpan::CalcPositionInSpan(float Ms, const ISongMapEvaluator& Maps) const
 {
 	FMidiSongPos Position;
 	Position.SetByTime(Ms, Maps);
 	return CalcPositionInSpan(Position, Maps);
 }
 
-float FMusicalTimeSpan::CalculateEnclosingVariableSizeSpanExtents(const FMidiSongPos& Position, const FSongMaps& Maps) const
+float FMusicalTimeSpan::CalculateEnclosingVariableSizeSpanExtents(const FMidiSongPos& Position, const ISongMapEvaluator& Maps) const
 {
 	switch (LengthUnits)
 	{
@@ -153,56 +153,57 @@ float FMusicalTimeSpan::CalculateEnclosingVariableSizeSpanExtents(const FMidiSon
 	return 0.0f;
 }
 
-float FMusicalTimeSpan::CalculateEnclosingFixedSizeSpanExtents(const FMidiSongPos& Position, const FSongMaps& Maps) const
+float FMusicalTimeSpan::CalculateEnclosingFixedSizeSpanExtents(const FMidiSongPos& Position, const ISongMapEvaluator& Maps) const
 {
 	int32 Tick = Maps.MsToTick(Position.SecondsIncludingCountIn * 1000.0f);
+	int32 MidiTTQ = Maps.GetTicksPerQuarterNote();
 	int32 GridUnitTicks = 1;
 	switch (LengthUnits)
 	{
 	case EMusicTimeSpanLengthUnits::ThirtySecondNotes:
-		GridUnitTicks = Maps.GetTicksPerQuarterNote() / 8;
+		GridUnitTicks = MidiTTQ / 8;
 		break;
 	case EMusicTimeSpanLengthUnits::SixteenthNotes:
-		GridUnitTicks = Maps.GetTicksPerQuarterNote() / 4;
+		GridUnitTicks = MidiTTQ / 4;
 		break;
 	case EMusicTimeSpanLengthUnits::EighthNotes:
-		GridUnitTicks = Maps.GetTicksPerQuarterNote() / 2;
+		GridUnitTicks = MidiTTQ / 2;
 		break;
 	case EMusicTimeSpanLengthUnits::QuarterNotes:
-		GridUnitTicks = Maps.GetTicksPerQuarterNote();
+		GridUnitTicks = MidiTTQ;
 		break;
 	case EMusicTimeSpanLengthUnits::HalfNotes:
-		GridUnitTicks = Maps.GetTicksPerQuarterNote() * 2;
+		GridUnitTicks = MidiTTQ * 2;
 		break;
 	case EMusicTimeSpanLengthUnits::WholeNotes:
-		GridUnitTicks = Maps.GetTicksPerQuarterNote() * 4;
+		GridUnitTicks = MidiTTQ * 4;
 		break;
 	case EMusicTimeSpanLengthUnits::DottedSixteenthNotes:
-		GridUnitTicks = Maps.GetTicksPerQuarterNote() * 3 / 8;
+		GridUnitTicks = MidiTTQ * 3 / 8;
 		break;
 	case EMusicTimeSpanLengthUnits::DottedEighthNotes:
-		GridUnitTicks = Maps.GetTicksPerQuarterNote() * 3 / 4;
+		GridUnitTicks = MidiTTQ * 3 / 4;
 		break;
 	case EMusicTimeSpanLengthUnits::DottedQuarterNotes:
-		GridUnitTicks = Maps.GetTicksPerQuarterNote() * 3 / 2;
+		GridUnitTicks = MidiTTQ * 3 / 2;
 		break;
 	case EMusicTimeSpanLengthUnits::DottedHalfNotes:
-		GridUnitTicks = Maps.GetTicksPerQuarterNote() * 3;
+		GridUnitTicks = MidiTTQ * 3;
 		break;
 	case EMusicTimeSpanLengthUnits::DottedWholeNotes:
-		GridUnitTicks = Maps.GetTicksPerQuarterNote() * 6;
+		GridUnitTicks = MidiTTQ * 6;
 		break;
 	case EMusicTimeSpanLengthUnits::SixteenthNoteTriplets:
-		GridUnitTicks = Maps.GetTicksPerQuarterNote() / 6;
+		GridUnitTicks = MidiTTQ / 6;
 		break;
 	case EMusicTimeSpanLengthUnits::EighthNoteTriplets:
-		GridUnitTicks = Maps.GetTicksPerQuarterNote() / 3;
+		GridUnitTicks = MidiTTQ / 3;
 		break;
 	case EMusicTimeSpanLengthUnits::QuarterNoteTriplets:
-		GridUnitTicks = Maps.GetTicksPerQuarterNote() * 2 / 3;
+		GridUnitTicks = MidiTTQ * 2 / 3;
 		break;
 	case EMusicTimeSpanLengthUnits::HalfNoteTriplets:
-		GridUnitTicks = Maps.GetTicksPerQuarterNote() * 4 / 3;
+		GridUnitTicks = MidiTTQ * 4 / 3;
 		break;
 	default:
 		checkNoEntry();
diff --git a/Engine/Plugins/Runtime/Harmonix/Source/HarmonixMidi/Private/HarmonixMidi/SmoothedMidiPlayCursor.cpp b/Engine/Plugins/Runtime/Harmonix/Source/HarmonixMidi/Private/HarmonixMidi/SmoothedMidiPlayCursor.cpp
--- a/Engine/Plugins/Runtime/Harmonix/Source/HarmonixMidi/Private/HarmonixMidi/SmoothedMidiPlayCursor.cpp
+++ b/Engine/Plugins/Runtime/Harmonix/Source/HarmonixMidi/Private/HarmonixMidi/SmoothedMidiPlayCursor.cpp
@@ -109,7 +109,7 @@ void FSmoothedMidiPlayCursor::SyncSmoothingTimer(bool bEnableErrorCorrection)
 			SmoothingTimer.Reset(RawMs);
 			SmoothedMs = RawMs;
 			SmoothedMsDelta = 0.0f;
-			SmoothedTick = CursorOwner->GetTempoMap().MsToTick(RawMs);
+			SmoothedTick = CursorOwner->GetSongMaps().MsToTick(RawMs);
 			ErrorTracker.Reset();
 			return;
 		}
@@ -168,7 +168,7 @@ void FSmoothedMidiPlayCursor::SyncSmoothingTimer(bool bEnableErrorCorrection)
 		}
 	}
 	SmoothedMs = float(SmoothMs);
-	SmoothedTick = CursorOwner->GetTempoMap().MsToTick(SmoothedMs);
+	SmoothedTick = CursorOwner->GetSongMaps().MsToTick(SmoothedMs);
 	LoopedThisPass = false;
 }
 
diff --git a/Engine/Plugins/Runtime/Harmonix/Source/HarmonixMidi/Private/HarmonixMidi/SongMaps.cpp b/Engine/Plugins/Runtime/Harmonix/Source/HarmonixMidi/Private/HarmonixMidi/SongMaps.cpp
--- a/Engine/Plugins/Runtime/Harmonix/Source/HarmonixMidi/Private/HarmonixMidi/SongMaps.cpp
+++ b/Engine/Plugins/Runtime/Harmonix/Source/HarmonixMidi/Private/HarmonixMidi/SongMaps.cpp
@@ -9,6 +9,23 @@ FSongMaps::FSongMaps()
 {
 }
 
+FSongMaps::FSongMaps(float Bpm, int32 TimeSigNumerator, int32 TimeSigDenominator)
+{
+	TempoMap.AddTempoInfoPoint(Harmonix::Midi::Constants::BPMToMidiTempo(Bpm), 0);
+	BarMap.AddTimeSignatureAtBarIncludingCountIn(0, TimeSigNumerator, TimeSigDenominator);
+}
+
+FSongMaps::FSongMaps(const ISongMapEvaluator& Other)
+{
+	TicksPerQuarterNote = Other.GetTicksPerQuarterNote();
+	TempoMap = Other.GetTempoMap();
+	BarMap = Other.GetBarMap();
+	BeatMap = Other.GetBeatMap();
+	SectionMap = Other.GetSectionMap();
+	ChordMap = Other.GetChordMap();
+	LengthData = Other.GetSongLengthData();
+}
+
 bool FSongMaps::operator==(const FSongMaps& Other) const
 {
 	if (TrackNames.Num() != Other.TrackNames.Num())
@@ -46,11 +63,11 @@ void FSongMaps::Init(int32 InTicksPerQuarterNote)
 	ChordMap.Empty();
 }
 
-void FSongMaps::Copy(const FSongMaps& Other, int32 StartTick, int32 EndTick)
+void FSongMaps::Copy(const ISongMapEvaluator& Other, int32 StartTick, int32 EndTick)
 {
 	EmptyAllMaps();
 
-	TicksPerQuarterNote = Other.TicksPerQuarterNote;
+	TicksPerQuarterNote = Other.GetTicksPerQuarterNote();
 
 	TempoMap.Copy(Other.GetTempoMap(), StartTick, EndTick);
 	BarMap.Copy(Other.GetBarMap(), StartTick, EndTick);
@@ -58,18 +75,25 @@ void FSongMaps::Copy(const FSongMaps& Other, int32 StartTick, int32 EndTick)
 	SectionMap.Copy(Other.GetSectionMap(), StartTick, EndTick);
 	ChordMap.Copy(Other.GetChordMap(), StartTick, EndTick);
 
-	int32 LastTick = EndTick == -1 ? Other.LengthData.LastTick : EndTick;
+	int32 LastTick = EndTick == -1 ? Other.GetSongLengthData().LastTick : EndTick;
 	TempoMap.Finalize(LastTick);
 	BarMap.Finalize(LastTick);
 	BeatMap.Finalize(LastTick);
 	SectionMap.Finalize(LastTick);
 	ChordMap.Finalize(LastTick);
 
-	memset(&LengthData, 0, sizeof(LengthData));
-	EMidiClockSubdivisionQuantization Division = EMidiClockSubdivisionQuantization::None;
-	LengthData.LastTick = QuantizeTickToAnyNearestSubdivision(EndTick, EMidiFileQuantizeDirection::Nearest, Division) - 1;
-	LengthData.LengthTicks = LengthData.LastTick + 1;
-	LengthData.LengthFractionalBars = BarMap.TickToFractionalBarIncludingCountIn(LengthData.LengthTicks);
+	if (EndTick != std::numeric_limits<int32>::max())
+	{
+		memset(&LengthData, 0, sizeof(LengthData));
+		EMidiClockSubdivisionQuantization Division = EMidiClockSubdivisionQuantization::None;
+		LengthData.LastTick = QuantizeTickToAnyNearestSubdivision(EndTick, EMidiFileQuantizeDirection::Nearest, Division) - 1;
+		LengthData.LengthTicks = LengthData.LastTick + 1;
+		LengthData.LengthFractionalBars = BarMap.TickToFractionalBarIncludingCountIn(LengthData.LengthTicks);
+	}
+	else
+	{
+		LengthData = Other.GetSongLengthData();
+	}
 }
 
 bool FSongMaps::LoadFromStdMidiFile(const FString& FilePath)
@@ -93,32 +117,6 @@ bool FSongMaps::LoadFromStdMidiFile(TSharedPtr<FArchive> Archive, const FString&
 	return ReadWithReader(Reader);
 }
 
-float FSongMaps::TickToMs(float Tick) const
-{
-	if (TempoMap.IsEmpty())
-	{
-		UE_LOG(LogMIDI, Log, TEXT("SongMaps does not contain a Tempo Map."));
-		return 0.0f;
-	}
-	return TempoMap.TickToMs(Tick);
-}
-
-float FSongMaps::MsToTick(float Ms) const
-{
-	if (TempoMap.IsEmpty())
-	{
-		UE_LOG(LogMIDI, Log, TEXT("SongMaps does not contain a Tempo Map."));
-		return 0;
-	}
-	return TempoMap.MsToTick(Ms);
-}
-
-float FSongMaps::GetCountInSeconds() const
-{
-	float BarOneBeatOneTick = BarMap.MusicTimestampToTick({ 1, 1.0f });
-	return TempoMap.TickToMs(BarOneBeatOneTick) / 1000.f;
-}
-
 bool FSongMaps::FinalizeRead(IMidiReader* Reader)
 {
 	int32 LastTick = Reader->GetLastTick();
@@ -174,35 +172,177 @@ bool FSongMaps::IsEmpty() const
 			LengthData.LengthTicks == 0;
 }
 
-float FSongMaps::GetSongLengthMs() const
+void FSongMaps::SetStartBar(int32 StartBar)
 {
-	return TickToMs(float(GetSongLengthData().LengthTicks));
+	BarMap.SetStartBar(StartBar);
 }
 
-int32 FSongMaps::GetSongLengthBeats() const
+void FSongMaps::SetSongLengthTicks(int32 NewLengthTicks)
 {
-	return BeatMap.GetNumMapPoints();
+	if (NewLengthTicks < 1)
+	{
+		UE_LOG(LogMIDI, Warning, TEXT("SetSongLengthTicks: Asked to set length less than 1. That is not possible. Setting to length 1!"));
+		NewLengthTicks = 1;
+	}
+	LengthData.LengthTicks = NewLengthTicks;
+	LengthData.LastTick = NewLengthTicks - 1;
+	LengthData.LengthFractionalBars = BarMap.TickToFractionalBarIncludingCountIn(LengthData.LengthTicks);
 }
 
-float FSongMaps::GetSongLengthFractionalBars() const
+void FSongMaps::FinalizeBarMap(int32 InLastTick)
 {
-	return LengthData.LengthFractionalBars;
+	BarMap.Finalize(InLastTick);
 }
 
-void FSongMaps::SetSongLengthTicks(int32 NewLengthTicks)
+bool FSongMaps::AddTempoInfoPoint(int32 MicrosecondsPerQuarterNote, int32 Tick, bool SortNow)
 {
-	if (NewLengthTicks < 1)
+	return TempoMap.AddTempoInfoPoint(MicrosecondsPerQuarterNote, Tick, SortNow);
+}
+
+bool FSongMaps::AddTimeSignatureAtBarIncludingCountIn(int32 BarIndex, int32 InNumerator, int32 InDenominator, bool SortNow, bool FailOnError)
+{
+	return BarMap.AddTimeSignatureAtBarIncludingCountIn(BarIndex, InNumerator, InDenominator, SortNow, FailOnError);
+}
+
+FTimeSignaturePoint* FSongMaps::GetMutableTimeSignaturePoint(int32 PointIndex)
+{
+	if (!BarMap.GetTimeSignaturePoints().IsValidIndex(PointIndex))
 	{
-		UE_LOG(LogMIDI, Warning, TEXT("SetSongLengthTicks: Asked to set length less than 1. That is not possible. Setting to length 1!"));
-		NewLengthTicks = 1;
+		return nullptr;
 	}
-	LengthData.LengthTicks = NewLengthTicks;
-	LengthData.LastTick    = NewLengthTicks - 1;
-	LengthData.LengthFractionalBars  = BarMap.TickToFractionalBarIncludingCountIn(LengthData.LengthTicks);
+	return &BarMap.GetTimeSignaturePoint(PointIndex);
+}
+
+void FSongMaps::StringLengthToMT(const FString& LengthString, int32& OutBars, int32& OutTicks)
+{
+	if (BarMap.IsEmpty())
+	{
+		UE_LOG(LogMIDI, Log, TEXT("SongMaps does not contain a Bar Map."));
+		return;
+	}
+
+	// ASSUMES WE DON'T HAVE A COMPLICATED TIME SIGNATURE.
+	// IF WE DID, THEN WE WOULD NEVER GET HERE, AS WE WOULD
+	// HAVE DETERMINED THE LENGTH FROM THE MIDI FILE!
+
+	// ALSO: LengthString, being a 'length', is specifying 0 based bars and beats!
+	//       And, the return values are lengths, so also 0 based.
+	const TCHAR* Walk = *LengthString;
+	OutBars = 0;
+	int32 Beats = 0;
+	OutTicks = 0;
+
+	// get bars...
+	int32 Count = 0;
+	while (*Walk != ':' && *Walk != 0 && Count < 3)
+	{
+		OutBars = (OutBars * 10) + ((int32)(*Walk) - (int32)'0');
+		Walk++;
+		Count++;
+	}
+	if (*Walk == ':')
+	{
+		// get beats...
+		Count = 0;
+		Walk++;
+		while (*Walk != ':' && *Walk != 0 && Count < 3)
+		{
+			Beats = (Beats * 10) + ((int32)(*Walk) - (int32)'0');
+			Walk++;
+			Count++;
+		}
+		if (*Walk == ':')
+		{
+			// get ticks...
+			Count = 0;
+			Walk++;
+			while (*Walk != ':' && *Walk != 0 && Count < 3)
+			{
+				OutTicks = (OutTicks * 10) + ((int32)(*Walk) - (int32)'0');
+				Walk++;
+				Count++;
+			}
+		}
+	}
+	const FTimeSignaturePoint* TimeSignaturePoint = &BarMap.GetTimeSignaturePoint(0);
+	if (!TimeSignaturePoint)
+	{
+		UE_LOG(LogMIDI, Log, TEXT("No Time Signature found in SongMaps."));
+		return;
+	}
+	OutBars += (Beats == 0 && OutTicks == 0) ? 0 : 1;
+	int32 TicksPerBeat = TicksPerQuarterNote / (TimeSignaturePoint->TimeSignature.Denominator / 4);
+	int32 TicksPerBar = TicksPerBeat * TimeSignaturePoint->TimeSignature.Numerator;
+	OutTicks = (TicksPerBar * OutBars)
+		+ (TicksPerBeat * Beats)
+		+ OutTicks;
+}
+
+FString FSongMaps::GetTrackName(int32 Index) const
+{
+	if (Index < 0 || Index >= TrackNames.Num())
+		return FString();
+
+	return TrackNames[Index];
+}
+
+int32 ISongMapEvaluator::GetTicksPerQuarterNote() const
+{
+	return GetTempoMap().GetTicksPerQuarterNote();
+}
+
+float ISongMapEvaluator::TickToMs(float Tick) const
+{
+	const FTempoMap& TempoMap = GetTempoMap();
+	if (TempoMap.IsEmpty())
+	{
+		UE_LOG(LogMIDI, Log, TEXT("SongMaps does not contain a Tempo Map."));
+		return 0.0f;
+	}
+	return TempoMap.TickToMs(Tick);
+}
+
+float ISongMapEvaluator::MsToTick(float Ms) const
+{
+	const FTempoMap& TempoMap = GetTempoMap();
+	if (TempoMap.IsEmpty())
+	{
+		UE_LOG(LogMIDI, Log, TEXT("SongMaps does not contain a Tempo Map."));
+		return 0;
+	}
+	return TempoMap.MsToTick(Ms);
 }
 
-bool FSongMaps::LengthIsAPerfectSubdivision() const
+float ISongMapEvaluator::GetCountInSeconds() const
 {
+	const FBarMap& BarMap = GetBarMap();
+	const FTempoMap& TempoMap = GetTempoMap();
+	float BarOneBeatOneTick = BarMap.MusicTimestampToTick({ 1, 1.0f });
+	return TempoMap.TickToMs(BarOneBeatOneTick) / 1000.f;
+}
+
+float ISongMapEvaluator::GetSongLengthMs() const
+{
+	return TickToMs(float(GetSongLengthData().LengthTicks));
+}
+
+int32 ISongMapEvaluator::GetSongLengthBeats() const
+{
+	const FBeatMap& BeatMap = GetBeatMap();
+	return BeatMap.GetNumMapPoints();
+}
+
+float ISongMapEvaluator::GetSongLengthFractionalBars() const
+{
+	const FSongLengthData& LengthData = GetSongLengthData();
+	return LengthData.LengthFractionalBars;
+}
+
+bool ISongMapEvaluator::LengthIsAPerfectSubdivision() const
+{
+	const FBarMap& BarMap = GetBarMap();
+	const FSongLengthData& LengthData = GetSongLengthData();
+	const int32 TicksPerQuarterNote = BarMap.GetTicksPerQuarterNote();
 	int32 BarIndex = 0;
 	int32 BeatInBar = 0;
 	int32 TickIndexInBeat = 0;
@@ -214,8 +354,10 @@ bool FSongMaps::LengthIsAPerfectSubdivision() const
 	return (TickIndexInBeat % TicksPer64thTriplet) == 0 || (TickIndexInBeat % TicksPer64th) == 0;
 }
 
-int32 FSongMaps::QuantizeTickToAnyNearestSubdivision(int32 InTick, EMidiFileQuantizeDirection Direction, EMidiClockSubdivisionQuantization& Division) const
+int32 ISongMapEvaluator::QuantizeTickToAnyNearestSubdivision(int32 InTick, EMidiFileQuantizeDirection Direction, EMidiClockSubdivisionQuantization& Division) const
 {
+	const FBarMap& BarMap = GetBarMap();
+	const int32 TicksPerQuarterNote = BarMap.GetTicksPerQuarterNote();
 	int32 BarIndex = 0;
 	int32 BeatInBar = 0;
 	int32 TickIndexInBeat = 0;
@@ -251,9 +393,9 @@ int32 FSongMaps::QuantizeTickToAnyNearestSubdivision(int32 InTick, EMidiFileQuan
 	EMidiClockSubdivisionQuantization BestDivision = EMidiClockSubdivisionQuantization::None;
 	int32 BestDistanceFromDivision = std::numeric_limits<int32>::max();
 	
-	auto TryDivision = [this, &TickAtTimeSignature, &BestDivision, &BestDistanceFromDivision, &TimeSignature, &Division, &Direction](EMidiClockSubdivisionQuantization TryDivision)
+	auto TryDivision = [this, &TickAtTimeSignature, &BestDivision, &BestDistanceFromDivision, &TimeSignature, &Division, &Direction, &TicksPerQuarterNote](EMidiClockSubdivisionQuantization TryDivision)
 		{
-			int32 TicksPerDivision = SubdivisionToMidiTicks(TryDivision, TimeSignature);
+			int32 TicksPerDivision = Harmonix::Midi::Constants::SubdivisionToMidiTicks(TryDivision, TimeSignature, TicksPerQuarterNote);
 			int32 DistanceFromDivision = TickAtTimeSignature % TicksPerDivision;
 			switch (Direction)
 			{
@@ -305,7 +447,7 @@ int32 FSongMaps::QuantizeTickToAnyNearestSubdivision(int32 InTick, EMidiFileQuan
 	return ZeroPoint + (TickAtTimeSignature - BestDistanceFromDivision);
 }
 
-int32 FSongMaps::QuantizeTickToNearestSubdivision(int32 InTick, EMidiFileQuantizeDirection Direction, EMidiClockSubdivisionQuantization Division) const
+int32 ISongMapEvaluator::QuantizeTickToNearestSubdivision(int32 InTick, EMidiFileQuantizeDirection Direction, EMidiClockSubdivisionQuantization Division) const
 {
 	int32 LowerTick = 0;
 	int32 UpperTick = 0;
@@ -319,14 +461,17 @@ int32 FSongMaps::QuantizeTickToNearestSubdivision(int32 InTick, EMidiFileQuantiz
 	return UpperTick;
 }
 
-void FSongMaps::GetTicksForNearestSubdivision(int32 InTick, EMidiClockSubdivisionQuantization Division, int32& LowerTick, int32& UpperTick) const
+void ISongMapEvaluator::GetTicksForNearestSubdivision(int32 InTick, EMidiClockSubdivisionQuantization Division, int32& LowerTick, int32& UpperTick) const
 {
+	const FBarMap& BarMap = GetBarMap();
+	const int32 TicksPerQuarterNote = BarMap.GetTicksPerQuarterNote();
+
 	int32 TicksInSubdivision = 0;
 	int32 TickError = 0;
 	if (BarMap.IsEmpty())
 	{
 		FTimeSignature TimeSignature(4, 4);
-		TicksInSubdivision = SubdivisionToMidiTicks(Division, TimeSignature);
+		TicksInSubdivision = Harmonix::Midi::Constants::SubdivisionToMidiTicks(Division, TimeSignature, TicksPerQuarterNote);
 		TickError = InTick % TicksInSubdivision;
 		LowerTick = InTick - TickError;
 		UpperTick = LowerTick + TicksInSubdivision;
@@ -375,8 +520,10 @@ void FSongMaps::GetTicksForNearestSubdivision(int32 InTick, EMidiClockSubdivisio
 	UpperTick = LowerTick + TicksInSubdivision;
 }
 
-FString FSongMaps::GetSongLengthString() const
+FString ISongMapEvaluator::GetSongLengthString() const
 {
+	const FBarMap& BarMap = GetBarMap();
+	const FSongLengthData& LengthData = GetSongLengthData();
 	int32 BarIndex;
 	int32 BeatInBar;
 	int32 TickIndexInBeat;
@@ -387,104 +534,70 @@ FString FSongMaps::GetSongLengthString() const
 	return FString::Printf(TEXT("%d | %.3f"), BarIndex, (float)BeatIndex + (float)TickIndexInBeat / (float)TicksPerBeat );
 }
 
-void FSongMaps::StringLengthToMT(const FString& LengthString, int32& OutBars, int32& OutTicks)
+///////////////////////////////////////////////////////////////////////////////////
+// TEMPO
+const FTempoInfoPoint* ISongMapEvaluator::GetTempoInfoForMs(float Ms) const
 {
-	if (BarMap.IsEmpty())
-	{
-		UE_LOG(LogMIDI, Log, TEXT("SongMaps does not contain a Bar Map."));
-		return;
-	}
-
-	// ASSUMES WE DON'T HAVE A COMPLICATED TIME SIGNATURE.
-	// IF WE DID, THEN WE WOULD NEVER GET HERE, AS WE WOULD
-	// HAVE DETERMINED THE LENGTH FROM THE MIDI FILE!
-	
-	// ALSO: LengthString, being a 'length', is specifying 0 based bars and beats!
-	//       And, the return values are lengths, so also 0 based.
-	const TCHAR* Walk = *LengthString;
-	OutBars = 0;
-	int32 Beats = 0;
-	OutTicks = 0;
-
-	// get bars...
-	int32 Count = 0;
-	while (*Walk != ':' && *Walk != 0 && Count < 3)
-	{
-		OutBars = (OutBars * 10) + ((int32)(*Walk) - (int32)'0');
-		Walk++;
-		Count++;
-	}
-	if (*Walk == ':')
-	{
-		// get beats...
-		Count = 0;
-		Walk++;
-		while (*Walk != ':' && *Walk != 0 && Count < 3)
-		{
-			Beats = (Beats * 10) + ((int32)(*Walk) - (int32)'0');
-			Walk++;
-			Count++;
-		}
-		if (*Walk == ':')
-		{
-			// get ticks...
-			Count = 0;
-			Walk++;
-			while (*Walk != ':' && *Walk != 0 && Count < 3)
-			{
-				OutTicks = (OutTicks * 10) + ((int32)(*Walk) - (int32)'0');
-				Walk++;
-				Count++;
-			}
-		}
-	}
-	const FTimeSignaturePoint* TimeSignaturePoint = &BarMap.GetTimeSignaturePoint(0);
-	if (!TimeSignaturePoint)
+	const FTempoMap& TempoMap = GetTempoMap();
+	if (TempoMap.IsEmpty())
 	{
-		UE_LOG(LogMIDI, Log, TEXT("No Time Signature found in SongMaps."));
-		return;
+		UE_LOG(LogMIDI, Log, TEXT("SongMaps does not contain a Tempo Map."));
+		return nullptr;
 	}
-	OutBars += (Beats == 0 && OutTicks == 0) ? 0 : 1;
-	int32 TicksPerBeat   = TicksPerQuarterNote / (TimeSignaturePoint->TimeSignature.Denominator / 4);
-	int32 TicksPerBar = TicksPerBeat * TimeSignaturePoint->TimeSignature.Numerator;
-	OutTicks = (TicksPerBar * OutBars)
-		+ (TicksPerBeat * Beats)
-		+ OutTicks;
-}
-
-FString FSongMaps::GetTrackName(int32 Index) const
-{
-	if (Index < 0 || Index >= TrackNames.Num())
-		return FString();
-
-	return TrackNames[Index];
+	int32 Tick = int32(MsToTick(Ms));
+	return TempoMap.GetTempoPointAtTick(Tick);
 }
 
-///////////////////////////////////////////////////////////////////////////////////
-// TEMPO
-const FTempoInfoPoint* FSongMaps::GetTempoInfoForMs(float Ms) const
+const FTempoInfoPoint* ISongMapEvaluator::GetTempoInfoForTick(int32 Tick) const
 {
+	const FTempoMap& TempoMap = GetTempoMap();
 	if (TempoMap.IsEmpty())
 	{
 		UE_LOG(LogMIDI, Log, TEXT("SongMaps does not contain a Tempo Map."));
 		return nullptr;
 	}
-	int32 Tick = int32(MsToTick(Ms));
 	return TempoMap.GetTempoPointAtTick(Tick);
 }
 
-const FTempoInfoPoint* FSongMaps::GetTempoInfoForTick(int32 Tick) const
+int32 ISongMapEvaluator::GetTempoPointIndexForTick(int32 Tick) const
 {
+	const FTempoMap& TempoMap = GetTempoMap();
 	if (TempoMap.IsEmpty())
 	{
 		UE_LOG(LogMIDI, Log, TEXT("SongMaps does not contain a Tempo Map."));
+		return -1;
+	}
+	return TempoMap.GetTempoPointIndexAtTick(Tick);
+}
+
+const FTempoInfoPoint* ISongMapEvaluator::GetTempoInfoPoint(int32 PointIndex) const
+{
+	const FTempoMap& TempoMap = GetTempoMap();
+	if (!TempoMap.GetTempoPoints().IsValidIndex(PointIndex))
+	{
 		return nullptr;
 	}
-	return TempoMap.GetTempoPointAtTick(Tick);
+	return &TempoMap.GetTempoPoints()[PointIndex];
+}
+
+int32 ISongMapEvaluator::GetNumTempoChanges() const
+{
+	return GetTempoMap().GetNumTempoChangePoints();
 }
 
-float FSongMaps::GetTempoAtMs(float Ms) const
+int32 ISongMapEvaluator::GetTempoChangePointTick(int32 PointIndex) const
 {
+	const FTempoMap& TempoMap = GetTempoMap();
+	if (!TempoMap.GetTempoPoints().IsValidIndex(PointIndex))
+	{
+		return 0;
+	}
+	return TempoMap.GetTempoChangePointTick(PointIndex);
+}
+
+float ISongMapEvaluator::GetTempoAtMs(float Ms) const
+{
+	const FTempoMap& TempoMap = GetTempoMap();
 	if (TempoMap.IsEmpty())
 	{
 		UE_LOG(LogMIDI, Log, TEXT("SongMaps does not contain a Tempo Map."));
@@ -494,8 +607,9 @@ float FSongMaps::GetTempoAtMs(float Ms) const
 	return TempoMap.GetTempoAtTick(Tick);
 }
 
-float FSongMaps::GetTempoAtTick(int32 Tick) const
+float ISongMapEvaluator::GetTempoAtTick(int32 Tick) const
 {
+	const FTempoMap& TempoMap = GetTempoMap();
 	if (TempoMap.IsEmpty())
 	{
 		UE_LOG(LogMIDI, Log, TEXT("SongMaps does not contain a Tempo Map."));
@@ -504,10 +618,16 @@ float FSongMaps::GetTempoAtTick(int32 Tick) const
 	return TempoMap.GetTempoAtTick(Tick);
 }
 
+bool ISongMapEvaluator::TempoMapIsEmpty() const
+{
+	return GetTempoMap().GetNumTempoChangePoints() == 0; 
+}
+
 ///////////////////////////////////////////////////////////////////////////////////
 // BEAT
-const FBeatMapPoint* FSongMaps::GetBeatAtMs(float Ms) const
+const FBeatMapPoint* ISongMapEvaluator::GetBeatAtMs(float Ms) const
 {
+	const FBeatMap& BeatMap = GetBeatMap();
 	if (BeatMap.IsEmpty())
 	{
 		UE_LOG(LogMIDI, Log, TEXT("SongMaps does not contain a Beat Map."));
@@ -517,8 +637,12 @@ const FBeatMapPoint* FSongMaps::GetBeatAtMs(float Ms) const
 	return BeatMap.GetPointInfoForTick(Tick);
 }
 
-float FSongMaps::GetMsAtBeat(float Beat) const
+float ISongMapEvaluator::GetMsAtBeat(float Beat) const
 {
+	const FBeatMap& BeatMap = GetBeatMap();
+	const FBarMap& BarMap = GetBarMap();
+	const int32 TicksPerQuarterNote = BarMap.GetTicksPerQuarterNote();
+
 	if (!BeatMap.IsEmpty() && Beat < BeatMap.GetNumMapPoints())
 	{
 		float Tick = BeatMap.GetFractionalTickAtBeat(Beat);
@@ -553,8 +677,9 @@ float FSongMaps::GetMsAtBeat(float Beat) const
 	return TickToMs(BarMap.FractionalBarIncludingCountInToTick(Bar));
 }
 
-const FBeatMapPoint* FSongMaps::GetBeatAtTick(int32 Tick) const
+const FBeatMapPoint* ISongMapEvaluator::GetBeatAtTick(int32 Tick) const
 {
+	const FBeatMap& BeatMap = GetBeatMap();
 	if (BeatMap.IsEmpty())
 	{
 		UE_LOG(LogMIDI, Log, TEXT("SongMaps does not contain a Beat Map."));
@@ -563,13 +688,13 @@ const FBeatMapPoint* FSongMaps::GetBeatAtTick(int32 Tick) const
 	return BeatMap.GetPointInfoForTick(Tick);
 }
 
-float FSongMaps::GetMsPerBeatAtMs(float Ms) const
+float ISongMapEvaluator::GetMsPerBeatAtMs(float Ms) const
 {
 	int32 Tick = int32(MsToTick(Ms));
 	return GetMsPerBeatAtTick(Tick);
 }
 
-float FSongMaps::GetMsPerBeatAtTick(int32 Tick) const
+float ISongMapEvaluator::GetMsPerBeatAtTick(int32 Tick) const
 {
 	const FBeatMapPoint* BeatInfo = GetBeatAtTick(Tick);
 	if (!BeatInfo)
@@ -585,14 +710,15 @@ float FSongMaps::GetMsPerBeatAtTick(int32 Tick) const
 	return TickToMs(float(BeatInfo->StartTick + BeatInfo->LengthTicks)) - TickToMs(float(BeatInfo->StartTick));
 }
 
-float FSongMaps::GetFractionalBeatAtMs(float Ms) const
+float ISongMapEvaluator::GetFractionalBeatAtMs(float Ms) const
 {
 	float Tick = MsToTick(Ms);
 	return GetFractionalBeatAtTick(Tick);
 }
 
-float FSongMaps::GetFractionalBeatAtTick(float Tick) const
+float ISongMapEvaluator::GetFractionalBeatAtTick(float Tick) const
 {
+	const FBeatMap& BeatMap = GetBeatMap();
 	int32 BeatIndex = GetBeatIndexAtTick(int32(Tick));
 	if (BeatIndex < 0 || BeatMap.IsEmpty())
 	{
@@ -606,30 +732,31 @@ float FSongMaps::GetFractionalBeatAtTick(float Tick) const
 	return BeatIndex + FractionalPart + 1.0f; // +1 for musical position
 }
 
-int32 FSongMaps::GetBeatIndexAtMs(float Ms) const
+int32 ISongMapEvaluator::GetBeatIndexAtMs(float Ms) const
 {
 	int32 Tick = int32(MsToTick(Ms));
 	return GetBeatIndexAtTick(Tick);
 }
 
-int32 FSongMaps::GetBeatIndexAtTick(int32 Tick) const
+int32 ISongMapEvaluator::GetBeatIndexAtTick(int32 Tick) const
 {
+	const FBeatMap& BeatMap = GetBeatMap();
 	if (BeatMap.IsEmpty())
 	{
-		UE_LOG(LogMIDI, Log, TEXT("SongMaps does not contain a Beat Map."));
 		return -1;
 	}
 	return BeatMap.GetPointIndexForTick(Tick);
 }
 
-EMusicalBeatType FSongMaps::GetBeatTypeAtMs(float Ms) const
+EMusicalBeatType ISongMapEvaluator::GetBeatTypeAtMs(float Ms) const
 {
 	int32 Tick = int32(MsToTick(Ms));
 	return GetBeatTypeAtTick(Tick);
 }
 
-EMusicalBeatType FSongMaps::GetBeatTypeAtTick(int32 Tick) const
+EMusicalBeatType ISongMapEvaluator::GetBeatTypeAtTick(int32 Tick) const
 {
+	const FBeatMap& BeatMap = GetBeatMap();
 	if (BeatMap.IsEmpty())
 	{
 		UE_LOG(LogMIDI, Log, TEXT("SongMaps does not contain a Beat Map."));
@@ -638,14 +765,26 @@ EMusicalBeatType FSongMaps::GetBeatTypeAtTick(int32 Tick) const
 	return BeatMap.GetBeatTypeAtTick(Tick);
 }
 
-float FSongMaps::GetBeatInPulseBarAtMs(float Ms) const
+const FBeatMapPoint* ISongMapEvaluator::GetBeatPointInfoAtTick(int32 Tick, int32* PointIndex) const
+{
+	const FBeatMap& BeatMap = GetBeatMap();
+	if (PointIndex)
+	{
+		*PointIndex = BeatMap.GetPointIndexForTick(Tick);
+		return *PointIndex >= 0 ? &BeatMap.GetBeatPointInfo(*PointIndex) : nullptr;		
+	}
+	return BeatMap.GetPointInfoForTick(Tick);
+}
+
+float ISongMapEvaluator::GetBeatInPulseBarAtMs(float Ms) const
 {
 	float Tick = MsToTick(Ms);
 	return GetBeatInPulseBarAtTick(Tick);
 }
 
-float FSongMaps::GetBeatInPulseBarAtTick(float Tick) const
+float ISongMapEvaluator::GetBeatInPulseBarAtTick(float Tick) const
 {
+	const FBeatMap& BeatMap = GetBeatMap();
 	if (BeatMap.IsEmpty())
 	{
 		UE_LOG(LogMIDI, Log, TEXT("SongMaps does not contain a Beat Map."));
@@ -654,14 +793,15 @@ float FSongMaps::GetBeatInPulseBarAtTick(float Tick) const
 	return BeatMap.GetBeatInPulseBarAtTick(Tick);
 }
 
-int32 FSongMaps::GetNumBeatsInPulseBarAtMs(float Ms) const
+int32 ISongMapEvaluator::GetNumBeatsInPulseBarAtMs(float Ms) const
 {
 	int32 Tick = int32(MsToTick(Ms));
 	return GetNumBeatsInPulseBarAtTick(Tick);
 }
 
-int32 FSongMaps::GetNumBeatsInPulseBarAtTick(int32 Tick) const
+int32 ISongMapEvaluator::GetNumBeatsInPulseBarAtTick(int32 Tick) const
 {
+	const FBeatMap& BeatMap = GetBeatMap();
 	if (BeatMap.IsEmpty())
 	{
 		UE_LOG(LogMIDI, Log, TEXT("SongMaps does not contain a Beat Map."));
@@ -670,17 +810,34 @@ int32 FSongMaps::GetNumBeatsInPulseBarAtTick(int32 Tick) const
 	return BeatMap.GetNumBeatsInPulseBarAt(Tick);
 }
 
+bool ISongMapEvaluator::BeatMapIsEmpty() const
+{
+	return GetBeatMap().GetNumMapPoints() == 0; 
+}
+
 ///////////////////////////////////////////////////////////////////////////////////
 // Time Signature
 
-const FTimeSignature* FSongMaps::GetTimeSignatureAtMs(float Ms) const
+int32 ISongMapEvaluator::GetStartBar() const
+{
+	const FBarMap& BarMap = GetBarMap();
+	return BarMap.GetStartBar();
+}
+
+int32 ISongMapEvaluator::GetNumTimeSignatureChanges() const
+{
+	return GetBarMap().GetNumTimeSignaturePoints();
+}
+
+const FTimeSignature* ISongMapEvaluator::GetTimeSignatureAtMs(float Ms) const
 {
 	int32 Tick = int32(MsToTick(Ms));
 	return GetTimeSignatureAtTick(Tick);
 }
 
-const FTimeSignature* FSongMaps::GetTimeSignatureAtTick(int32 Tick) const
+const FTimeSignature* ISongMapEvaluator::GetTimeSignatureAtTick(int32 Tick) const
 {
+	const FBarMap& BarMap = GetBarMap();
 	if (BarMap.IsEmpty())
 	{
 		UE_LOG(LogMIDI, Log, TEXT("SongMaps does not contain a Bar Map."));
@@ -689,8 +846,20 @@ const FTimeSignature* FSongMaps::GetTimeSignatureAtTick(int32 Tick) const
 	return &BarMap.GetTimeSignatureAtTick(Tick);
 }
 
-const FTimeSignature* FSongMaps::GetTimeSignatureAtBar(int32 Bar) const
+int32 ISongMapEvaluator::GetTimeSignaturePointIndexForTick(int32 Tick) const
+{
+	const FBarMap& BarMap = GetBarMap();
+	if (BarMap.IsEmpty())
+	{
+		UE_LOG(LogMIDI, Log, TEXT("SongMaps does not contain a Bar Map."));
+		return -1;
+	}
+	return BarMap.GetPointIndexForTick(Tick);
+}
+
+const FTimeSignature* ISongMapEvaluator::GetTimeSignatureAtBar(int32 Bar) const
 {
+	const FBarMap& BarMap = GetBarMap();
 	if (Bar < 1)
 	{
 		UE_LOG(LogMIDI, Warning, TEXT("Bar < 1 (%d) specified as a musical position! Bars are '1' based in musical positions. Using bar 1!"), Bar);
@@ -704,30 +873,47 @@ const FTimeSignature* FSongMaps::GetTimeSignatureAtBar(int32 Bar) const
 	return &BarMap.GetTimeSignatureAtBar(Bar);
 }
 
-float FSongMaps::GetBarIncludingCountInAtMs(float Ms) const
+const FTimeSignaturePoint* ISongMapEvaluator::GetTimeSignaturePointAtTick(int32 Tick) const
 {
-	float Tick = MsToTick(Ms);
-	return GetBarIncludingCountInAtTick(Tick);
+	return GetBarMap().GetTimeSignaturePointForTick(Tick);
 }
 
-float FSongMaps::GetBarIncludingCountInAtTick(float Tick) const
+const FTimeSignaturePoint* ISongMapEvaluator::GetTimeSignaturePoint(int32 PointIndex) const
 {
-	if (BarMap.IsEmpty())
+	return &GetBarMap().GetTimeSignaturePoint(PointIndex);
+}
+
+int32 ISongMapEvaluator::GetTimeSignatureChangePointTick(int32 PointIndex)
+{
+	const FBarMap& BarMap = GetBarMap();
+	if (!BarMap.GetTimeSignaturePoints().IsValidIndex(PointIndex))
 	{
-		UE_LOG(LogMIDI, Log, TEXT("SongMaps does not contain a Bar Map."));
-		return 1.0f; // 1 based music positions!
+		return 0;
 	}
-	return BarMap.TickToFractionalBarIncludingCountIn(Tick);
+	return BarMap.GetTimeSignatureChangePointTick(PointIndex);
 }
 
-float FSongMaps::GetMsPerBarAtMs(float Ms) const
+float ISongMapEvaluator::GetBarIncludingCountInAtMs(float Ms) const
+{
+	float Tick = MsToTick(Ms);
+	return GetBarIncludingCountInAtTick(Tick);
+}
+
+float ISongMapEvaluator::GetBarIncludingCountInAtTick(float Tick) const
+{
+	return GetBarMap().TickToFractionalBarIncludingCountIn(Tick);
+}
+
+float ISongMapEvaluator::GetMsPerBarAtMs(float Ms) const
 {
 	float Tick = MsToTick(Ms);
 	return GetMsPerBarAtTick(Tick);
 }
 
-float FSongMaps::GetMsPerBarAtTick(float Tick) const
+float ISongMapEvaluator::GetMsPerBarAtTick(float Tick) const
 {
+	const FTempoMap& TempoMap = GetTempoMap();
+	const FBarMap& BarMap = GetBarMap();
 	if (TempoMap.IsEmpty())
 	{
 		UE_LOG(LogMIDI, Log, TEXT("SongMaps does not contain a Tempo Map."));
@@ -753,6 +939,46 @@ float FSongMaps::GetMsPerBarAtTick(float Tick) const
 	return QuarterNotesInBar * MsPerQuarterNote;
 }
 
+bool ISongMapEvaluator::BarMapIsEmpty() const
+{
+	return GetBarMap().GetNumTimeSignaturePoints() == 0; 
+}
+
+FMusicTimestamp ISongMapEvaluator::TickToMusicTimestamp(float Tick, int32* OutBeatsPerBar) const
+{
+	return GetBarMap().TickToMusicTimestamp(Tick, OutBeatsPerBar);
+}
+
+int32 ISongMapEvaluator::BarIncludingCountInToTick(int32 BarIndex, int32* OutBeatsPerBar, int32* OutTicksPerBeat) const
+{
+	return GetBarMap().BarIncludingCountInToTick(BarIndex, OutBeatsPerBar, OutTicksPerBeat);
+}
+
+int32 ISongMapEvaluator::BarBeatTickIncludingCountInToTick(int32 BarIndex, int32 BeatInBar, int32 TickInBeat) const
+{
+	return GetBarMap().BarBeatTickIncludingCountInToTick(BarIndex, BeatInBar, TickInBeat);
+}
+
+float ISongMapEvaluator::FractionalBarIncludingCountInToTick(float FractionalBarIndex) const
+{
+	return GetBarMap().FractionalBarIncludingCountInToTick(FractionalBarIndex);
+}
+
+int32 ISongMapEvaluator::TickToBarIncludingCountIn(int32 Tick) const
+{
+	return GetBarMap().TickToBarIncludingCountIn(Tick);
+}
+
+float ISongMapEvaluator::TickToFractionalBarIncludingCountIn(float Tick) const
+{
+	return GetBarMap().TickToFractionalBarIncludingCountIn(Tick);
+}
+
+void ISongMapEvaluator::TickToBarBeatTickIncludingCountIn(int32 RawTick, int32& OutBarIndex, int32& OutBeatInBarIndex, int32& OutTickIndexInBeat, int32* OutBeatsPerBar, int32* OutTicksPerBeat) const
+{
+	GetBarMap().TickToBarBeatTickIncludingCountIn(RawTick, OutBarIndex, OutBeatInBarIndex, OutTickIndexInBeat, OutBeatsPerBar, OutTicksPerBeat);
+}
+
 void FSongMaps::SetLengthTotalBars(int32 Bars)
 {
 	check(Bars >= 0);
@@ -762,144 +988,75 @@ void FSongMaps::SetLengthTotalBars(int32 Bars)
 	LengthData.LengthFractionalBars = Bars;
 }
 
-int32 FSongMaps::CalculateMidiTick(const FMusicTimestamp& Timestamp, const EMidiClockSubdivisionQuantization Quantize) const
+int32 ISongMapEvaluator::CalculateMidiTick(const FMusicTimestamp& Timestamp, const EMidiClockSubdivisionQuantization Quantize) const
 {
-	int32 TriggerTick = 0;
-	if (Quantize == EMidiClockSubdivisionQuantization::None)
-	{
-		TriggerTick = FMath::RoundToInt32(BarMap.MusicTimestampToTick(Timestamp));
-	}
-	else
-	{
-		int32 RawTick = FMath::RoundToInt32(BarMap.MusicTimestampToTick(Timestamp));
-		int32 BarTick = BarMap.MusicTimestampBarToTick(Timestamp.Bar);
-		int32 TicksPerQuantizationUnit = SubdivisionToMidiTicks(Quantize, RawTick);
-		if (ensure(TicksPerQuantizationUnit > 0))
-		{
-			float NumUnits = ((float)(RawTick - BarTick)) / (float)TicksPerQuantizationUnit;
-			int32 NumWholeUnits = FMath::RoundToInt32(NumUnits);
-			TriggerTick = BarTick + (NumWholeUnits * TicksPerQuantizationUnit);
-		}
-		else
-		{
-			TriggerTick = FMath::RoundToInt32(BarMap.MusicTimestampToTick(Timestamp));
-		}
-	}
-	return TriggerTick;
+	const FBarMap& BarMap = GetBarMap();
+	return BarMap.CalculateMidiTick(Timestamp, Quantize);
 }
 
-int32 FSongMaps::SubdivisionToMidiTicks(const EMidiClockSubdivisionQuantization Division, const FTimeSignature& TimeSignature) const
+int32 ISongMapEvaluator::SubdivisionToMidiTicks(const EMidiClockSubdivisionQuantization Division, const int32 AtTick) const
 {
-	switch (Division)
-	{
-	case EMidiClockSubdivisionQuantization::None:                   return 1;
-	case EMidiClockSubdivisionQuantization::Bar: 					return TimeSignature.Numerator * ((TicksPerQuarterNote * 4)/TimeSignature.Denominator);
-	case EMidiClockSubdivisionQuantization::Beat:					return (TicksPerQuarterNote * 4) / TimeSignature.Denominator;
-	case EMidiClockSubdivisionQuantization::ThirtySecondNote:		return TicksPerQuarterNote / 8;
-	case EMidiClockSubdivisionQuantization::SixteenthNote:			return TicksPerQuarterNote / 4;
-	case EMidiClockSubdivisionQuantization::EighthNote:				return TicksPerQuarterNote / 2;
-	case EMidiClockSubdivisionQuantization::QuarterNote:			return TicksPerQuarterNote;
-	case EMidiClockSubdivisionQuantization::HalfNote:				return TicksPerQuarterNote * 2;
-	case EMidiClockSubdivisionQuantization::WholeNote:				return TicksPerQuarterNote * 4;
-	case EMidiClockSubdivisionQuantization::DottedSixteenthNote:	return (TicksPerQuarterNote / 4) + (TicksPerQuarterNote / 8);
-	case EMidiClockSubdivisionQuantization::DottedEighthNote:		return (TicksPerQuarterNote / 2) + (TicksPerQuarterNote / 4);
-	case EMidiClockSubdivisionQuantization::DottedQuarterNote:		return (TicksPerQuarterNote)+(TicksPerQuarterNote / 2);
-	case EMidiClockSubdivisionQuantization::DottedHalfNote:			return (TicksPerQuarterNote * 2) + (TicksPerQuarterNote);
-	case EMidiClockSubdivisionQuantization::DottedWholeNote:		return (TicksPerQuarterNote * 4) + (TicksPerQuarterNote * 2);
-	case EMidiClockSubdivisionQuantization::SixteenthNoteTriplet:   return (TicksPerQuarterNote / 2) / 3;
-	case EMidiClockSubdivisionQuantization::EighthNoteTriplet:		return TicksPerQuarterNote / 3;
-	case EMidiClockSubdivisionQuantization::QuarterNoteTriplet:		return (TicksPerQuarterNote * 2) / 3;
-	case EMidiClockSubdivisionQuantization::HalfNoteTriplet:        return (TicksPerQuarterNote * 4) / 3;
-	default:	                                             		checkNoEntry();	return 1;
-	}
+	const FBarMap& BarMap = GetBarMap();
+	return BarMap.SubdivisionToMidiTicks(Division, AtTick);
 }
 
-int32 FSongMaps::SubdivisionToMidiTicks(const EMidiClockSubdivisionQuantization Division, const int32 AtTick) const
+float ISongMapEvaluator::MusicTimestampToTick(const FMusicTimestamp& Timestamp) const
 {
-	FTimeSignature TimeSignature(4,4);
-	int32 BarMapPointIndex = BarMap.GetPointIndexForTick(AtTick);
-	if (BarMapPointIndex >= 0)
-	{
-		const FTimeSignaturePoint& TimeSignaturePoint = BarMap.GetTimeSignaturePoint(BarMapPointIndex);
-		TimeSignature = TimeSignaturePoint.TimeSignature;
-	}
-	return SubdivisionToMidiTicks(Division, TimeSignature);
-}
-
-float FSongMaps::SubdivisionToBeats(EMidiClockSubdivisionQuantization Subdivision, const FTimeSignature& TimeSignature)
-{
-	// Easy cases first
-	if (Subdivision == EMidiClockSubdivisionQuantization::Bar)
-	{
-		return TimeSignature.Numerator;
-	}
-
-	if (Subdivision == EMidiClockSubdivisionQuantization::Beat)
-	{
-		return 1;
-	}
-
-	const float BeatsPerQuarter = TimeSignature.Denominator / 4.0f;
-
-	switch (Subdivision)
-	{
-	case EMidiClockSubdivisionQuantization::ThirtySecondNote:
-		return BeatsPerQuarter / 8;
-	case EMidiClockSubdivisionQuantization::SixteenthNote:
-		return BeatsPerQuarter / 4;
-	case EMidiClockSubdivisionQuantization::EighthNote:
-		return BeatsPerQuarter / 2;
-	case EMidiClockSubdivisionQuantization::QuarterNote:
-		return BeatsPerQuarter;
-	case EMidiClockSubdivisionQuantization::HalfNote:
-		return BeatsPerQuarter * 2;
-	case EMidiClockSubdivisionQuantization::WholeNote:
-		return BeatsPerQuarter * 4;
-	case EMidiClockSubdivisionQuantization::DottedSixteenthNote:
-		return BeatsPerQuarter / 4 + BeatsPerQuarter / 8;
-	case EMidiClockSubdivisionQuantization::DottedEighthNote:
-		return BeatsPerQuarter / 2 + BeatsPerQuarter / 4;
-	case EMidiClockSubdivisionQuantization::DottedQuarterNote:
-		return BeatsPerQuarter + BeatsPerQuarter / 2;
-	case EMidiClockSubdivisionQuantization::DottedHalfNote:
-		return BeatsPerQuarter * 3;
-	case EMidiClockSubdivisionQuantization::DottedWholeNote:
-		return BeatsPerQuarter * 6;
-	case EMidiClockSubdivisionQuantization::SixteenthNoteTriplet:
-		return (BeatsPerQuarter / 4) * 2 / 3;
-	case EMidiClockSubdivisionQuantization::EighthNoteTriplet:
-		return (BeatsPerQuarter / 2) * 2 / 3;
-	case EMidiClockSubdivisionQuantization::QuarterNoteTriplet:
-		return BeatsPerQuarter * 2 / 3;
-	case EMidiClockSubdivisionQuantization::HalfNoteTriplet:
-		return BeatsPerQuarter * 4 / 3;
-	default:
-		checkNoEntry();
-		return 0;
-	}
+	const FBarMap& BarMap = GetBarMap();
+	return BarMap.MusicTimestampToTick(Timestamp);
+}
+
+int32 ISongMapEvaluator::MusicTimestampBarToTick(int32 BarNumber, int32* OutBeatsPerBar, int32* OutTicksPerBeat) const
+{
+	const FBarMap& BarMap = GetBarMap();
+	return BarMap.MusicTimestampBarToTick(BarNumber, OutBeatsPerBar, OutTicksPerBeat);
 }
 
 ///////////////////////////////////////////////////////////////////////////////////
 // Sections
 
-const FSongSection* FSongMaps::GetSectionAtMs(float Ms) const
+const TArray<FSongSection>& ISongMapEvaluator::GetSections() const
+{
+	const FSectionMap& SectionMap = GetSectionMap();
+	return SectionMap.GetSections();
+}
+
+int32 ISongMapEvaluator::GetNumSections() const
+{
+	const FSectionMap& SectionMap = GetSectionMap();
+	return SectionMap.GetNumSections();
+}
+
+const FSongSection* ISongMapEvaluator::GetSectionAtMs(float Ms) const
 {
 	int32 Tick = int32(MsToTick(Ms));
 	return GetSectionAtTick(Tick);
 }
 
-const FSongSection* FSongMaps::GetSectionAtTick(int32 Tick) const
+const FSongSection* ISongMapEvaluator::GetSectionAtTick(int32 Tick) const
 {
+	const FSectionMap& SectionMap = GetSectionMap();
 	if (SectionMap.IsEmpty())
 	{
-		UE_LOG(LogMIDI, Log, TEXT("SongMaps does not contain a Section Map."));
 		return nullptr;
 	}
 	return SectionMap.TickToSection(Tick);
 }
 
-const FSongSection* FSongMaps::GetSectionWithName(const FString& Name) const
+int32 ISongMapEvaluator::GetSectionIndexAtTick(int32 Tick) const
+{
+	const FSectionMap& SectionMap = GetSectionMap();
+	if (SectionMap.IsEmpty())
+	{
+		return -1;
+	}
+	return SectionMap.TickToSectionIndex(Tick);
+}
+
+
+const FSongSection* ISongMapEvaluator::GetSectionWithName(const FString& Name) const
 {
+	const FSectionMap& SectionMap = GetSectionMap();
 	if (SectionMap.IsEmpty())
 	{
 		UE_LOG(LogMIDI, Log, TEXT("SongMaps does not contain a Section Map."));
@@ -908,14 +1065,15 @@ const FSongSection* FSongMaps::GetSectionWithName(const FString& Name) const
 	return SectionMap.FindSectionInfo(Name);
 }
 
-FString FSongMaps::GetSectionNameAtMs(float Ms) const
+FString ISongMapEvaluator::GetSectionNameAtMs(float Ms) const
 {
 	int32 Tick = int32(MsToTick(Ms));
 	return GetSectionNameAtTick(Tick);
 }
 
-FString FSongMaps::GetSectionNameAtTick(int32 Tick) const
+FString ISongMapEvaluator::GetSectionNameAtTick(int32 Tick) const
 {
+	const FSectionMap& SectionMap = GetSectionMap();
 	if (SectionMap.IsEmpty())
 	{
 		UE_LOG(LogMIDI, Log, TEXT("SongMaps does not contain a Section Map."));
@@ -924,13 +1082,13 @@ FString FSongMaps::GetSectionNameAtTick(int32 Tick) const
 	return SectionMap.GetSectionNameAtTick(Tick);
 }
 
-float FSongMaps::GetSectionLengthMsAtMs(float Ms) const
+float ISongMapEvaluator::GetSectionLengthMsAtMs(float Ms) const
 {
 	int32 Tick = int32(MsToTick(Ms));
 	return GetSectionLengthMsAtTick(Tick);
 }
 
-float FSongMaps::GetSectionStartMsAtMs(float Ms) const
+float ISongMapEvaluator::GetSectionStartMsAtMs(float Ms) const
 {
 	const int32 Tick = int32(MsToTick(Ms));
 	const FSongSection* SectionAtTick = GetSectionAtTick(Tick);
@@ -942,7 +1100,7 @@ float FSongMaps::GetSectionStartMsAtMs(float Ms) const
 	return TickToMs(float(SectionAtTick->StartTick));
 }
 
-float FSongMaps::GetSectionEndMsAtMs(float Ms) const
+float ISongMapEvaluator::GetSectionEndMsAtMs(float Ms) const
 {
 	const int32 Tick = int32(MsToTick(Ms));
 	const FSongSection* SectionAtTick = GetSectionAtTick(Tick);
@@ -954,7 +1112,7 @@ float FSongMaps::GetSectionEndMsAtMs(float Ms) const
 	return TickToMs(float(SectionAtTick->EndTick()));
 }
 
-float FSongMaps::GetSectionLengthMsAtTick(int32 Tick) const
+float ISongMapEvaluator::GetSectionLengthMsAtTick(int32 Tick) const
 {
 	const FSongSection* SectionAtTick = GetSectionAtTick(Tick);
 	if (!SectionAtTick)
@@ -967,17 +1125,24 @@ float FSongMaps::GetSectionLengthMsAtTick(int32 Tick) const
 	return EndMs - StartMs;
 }
 
+bool ISongMapEvaluator::SectionMapIsEmpty() const
+{
+	return GetSectionMap().GetNumSections() == 0; 
+}
+
+
 ///////////////////////////////////////////////////////////////////////////////////
 // Chords
 
-const FChordMapPoint* FSongMaps::GetChordAtMs(float Ms) const
+const FChordMapPoint* ISongMapEvaluator::GetChordAtMs(float Ms) const
 {
 	int32 Tick = int32(MsToTick(Ms));
 	return GetChordAtTick(Tick);
 }
 
-const FChordMapPoint* FSongMaps::GetChordAtTick(int32 Tick) const
+const FChordMapPoint* ISongMapEvaluator::GetChordAtTick(int32 Tick) const
 {
+	const FChordProgressionMap& ChordMap = GetChordMap();
 	if (ChordMap.IsEmpty())
 	{
 		UE_LOG(LogMIDI, Log, TEXT("SongMaps does not contain a Chord Map."));
@@ -986,14 +1151,15 @@ const FChordMapPoint* FSongMaps::GetChordAtTick(int32 Tick) const
 	return ChordMap.GetPointInfoForTick(Tick);
 }
 
-FName FSongMaps::GetChordNameAtMs(float Ms) const
+FName ISongMapEvaluator::GetChordNameAtMs(float Ms) const
 {
 	int32 Tick = int32(MsToTick(Ms));
 	return GetChordNameAtTick(Tick);
 }
 
-FName FSongMaps::GetChordNameAtTick(int32 Tick) const
+FName ISongMapEvaluator::GetChordNameAtTick(int32 Tick) const
 {
+	const FChordProgressionMap& ChordMap = GetChordMap();
 	if (ChordMap.IsEmpty())
 	{
 		UE_LOG(LogMIDI, Log, TEXT("SongMaps does not contain a Chord Map."));
@@ -1002,13 +1168,13 @@ FName FSongMaps::GetChordNameAtTick(int32 Tick) const
 	return ChordMap.GetChordNameAtTick(Tick);
 }
 
-float FSongMaps::GetChordLengthMsAtMs(float Ms) const
+float ISongMapEvaluator::GetChordLengthMsAtMs(float Ms) const
 {
 	int32 Tick = int32(MsToTick(Ms));
 	return GetChordLengthMsAtTick(Tick);
 }
 
-float FSongMaps::GetChordLengthMsAtTick(int32 Tick) const
+float ISongMapEvaluator::GetChordLengthMsAtTick(int32 Tick) const
 {
 	const FChordMapPoint* ChordInfo = GetChordAtTick(Tick);
 	if (!ChordInfo)
@@ -1020,3 +1186,37 @@ float FSongMaps::GetChordLengthMsAtTick(int32 Tick) const
 	return ChordEndMs - ChordStartMs;
 }
 
+bool ISongMapEvaluator::ChordMapIsEmpty() const
+{
+	return GetChordMap().GetNumChords() == 0; 
+}
+
+const FTempoMap& FSongMapsWithAlternateTempoSource::GetTempoMap() const
+{
+	return MidiFileWithTempoMap->SongMaps.GetTempoMap();
+}
+
+const FBeatMap& FSongMapsWithAlternateTempoSource::GetBeatMap() const
+{
+	return MidiFileWithOtherMaps->SongMaps.GetBeatMap();
+}
+
+const FBarMap& FSongMapsWithAlternateTempoSource::GetBarMap() const
+{
+	return MidiFileWithOtherMaps->SongMaps.GetBarMap();
+}
+
+const FSectionMap& FSongMapsWithAlternateTempoSource::GetSectionMap() const
+{
+	return MidiFileWithOtherMaps->SongMaps.GetSectionMap();
+}
+
+const FChordProgressionMap& FSongMapsWithAlternateTempoSource::GetChordMap() const
+{
+	return MidiFileWithOtherMaps->SongMaps.GetChordMap();
+}
+
+const FSongLengthData& FSongMapsWithAlternateTempoSource::GetSongLengthData() const
+{
+	return MidiFileWithOtherMaps->SongMaps.GetSongLengthData();
+}
diff --git a/Engine/Plugins/Runtime/Harmonix/Source/HarmonixMidi/Private/HarmonixMidi/TempoMap.cpp b/Engine/Plugins/Runtime/Harmonix/Source/HarmonixMidi/Private/HarmonixMidi/TempoMap.cpp
--- a/Engine/Plugins/Runtime/Harmonix/Source/HarmonixMidi/Private/HarmonixMidi/TempoMap.cpp
+++ b/Engine/Plugins/Runtime/Harmonix/Source/HarmonixMidi/Private/HarmonixMidi/TempoMap.cpp
@@ -241,6 +241,11 @@ const FTempoInfoPoint* FTempoMap::GetTempoPointAtTick(int32 Tick) const
 	return FMusicMapUtl::GetPointInfoForTick(Points, Tick);
 }
 
+int32 FTempoMap::GetTempoPointIndexAtTick(int32 Tick) const
+{
+	return FMusicMapUtl::GetPointIndexForTick(Points, Tick);
+}
+
 int32 FTempoMap::GetNumTempoChangePoints() const
 {
 	return Points.Num();
diff --git a/Engine/Plugins/Runtime/Harmonix/Source/HarmonixMidi/Public/HarmonixMidi/BarMap.h b/Engine/Plugins/Runtime/Harmonix/Source/HarmonixMidi/Public/HarmonixMidi/BarMap.h
--- a/Engine/Plugins/Runtime/Harmonix/Source/HarmonixMidi/Public/HarmonixMidi/BarMap.h
+++ b/Engine/Plugins/Runtime/Harmonix/Source/HarmonixMidi/Public/HarmonixMidi/BarMap.h
@@ -3,6 +3,7 @@
 #include "HarmonixMidi/MusicMapBase.h"
 #include "HarmonixMidi/MidiConstants.h"
 #include "Math/UnrealMathUtility.h"
+#include <limits>
 
 #include "BarMap.generated.h"
 
@@ -158,15 +159,21 @@ public:
 	bool operator==(const FBarMap& Other) const;
 
 	void Empty();
-	void Copy(const FBarMap& Other, int32 StartTick = 0, int32 EndTick = -1);
+	void Copy(const FBarMap& Other, int32 StartTick = 0, int32 EndTick = std::numeric_limits<int32>::max());
 	bool IsEmpty() const;
 
+	int32 CalculateMidiTick(const FMusicTimestamp& Timestamp, const EMidiClockSubdivisionQuantization Quantize) const;
+
+	int32 SubdivisionToMidiTicks(const EMidiClockSubdivisionQuantization Division, const int32 AtTick) const;
+
 	/** Called by the midi file importer before map points are added to this map */
 	void SetTicksPerQuarterNote(int32 InTicksPerQuarterNote)
 	{
 		TicksPerQuarterNote = InTicksPerQuarterNote;
 	}
 
+	int32 GetTicksPerQuarterNote() const { return TicksPerQuarterNote; }
+
 	void SupplyDefault() { AddTimeSignatureAtBarIncludingCountIn(0, 4, 4); }
 
 	/** 
@@ -334,6 +341,8 @@ public:
 
 	const FTimeSignature& GetTimeSignatureAtTick(int32 InTick) const;
 
+	const FTimeSignaturePoint* GetTimeSignaturePointForTick(int32 InTick) const;
+
 	const FTimeSignature& GetTimeSignatureAtBar(int32 InBar) const;
 
 	FORCEINLINE int32 GetTicksInBeatAfterPoint(int32 Index) const
@@ -361,6 +370,11 @@ public:
 	void SetStartBar(int32 InStartBar) { StartBar = InStartBar; }
 	int32 GetStartBar() const { return StartBar; }
 
+	// Returns the time signature points for inspection.
+	const TArray<FTimeSignaturePoint>& GetTimeSignaturePoints() const { return Points; }
+
+	int32 GetTimeSignatureChangePointTick(int32 PointIndex) const;
+
 protected:
 	UPROPERTY()
 	int32 StartBar = 1;
@@ -368,5 +382,7 @@ protected:
 	int32 TicksPerQuarterNote;
 	UPROPERTY()
 	TArray<FTimeSignaturePoint> Points;
+
+	static const FTimeSignaturePoint sDefaultTimeSignature;
 };
 
diff --git a/Engine/Plugins/Runtime/Harmonix/Source/HarmonixMidi/Public/HarmonixMidi/BeatMap.h b/Engine/Plugins/Runtime/Harmonix/Source/HarmonixMidi/Public/HarmonixMidi/BeatMap.h
--- a/Engine/Plugins/Runtime/Harmonix/Source/HarmonixMidi/Public/HarmonixMidi/BeatMap.h
+++ b/Engine/Plugins/Runtime/Harmonix/Source/HarmonixMidi/Public/HarmonixMidi/BeatMap.h
@@ -2,6 +2,7 @@
 #pragma once
 #include "HarmonixMidi/MusicMapBase.h"
 #include "HarmonixMidi/MidiConstants.h"
+#include <limits>
 
 #include "BeatMap.generated.h"
 
@@ -13,6 +14,8 @@ enum class EMusicalBeatType : uint8
 	Normal = 2,
 };
 
+FString HARMONIXMIDI_API MusicalBeatTypeToString(EMusicalBeatType BeatType);
+
 /**
  * A point in the music representing a "beat".
  * 
@@ -104,7 +107,7 @@ public:
 	bool operator==(const FBeatMap& Other) const;
 
 	void Empty();
-	void Copy(const FBeatMap& Other, int32 StartTick = 0, int32 EndTick = -1);
+	void Copy(const FBeatMap& Other, int32 StartTick = 0, int32 EndTick = std::numeric_limits<int32>::max());
 	bool IsEmpty() const;
 	int32 GetNumMapPoints() const;
 
diff --git a/Engine/Plugins/Runtime/Harmonix/Source/HarmonixMidi/Public/HarmonixMidi/ChordMap.h b/Engine/Plugins/Runtime/Harmonix/Source/HarmonixMidi/Public/HarmonixMidi/ChordMap.h
--- a/Engine/Plugins/Runtime/Harmonix/Source/HarmonixMidi/Public/HarmonixMidi/ChordMap.h
+++ b/Engine/Plugins/Runtime/Harmonix/Source/HarmonixMidi/Public/HarmonixMidi/ChordMap.h
@@ -2,6 +2,7 @@
 #pragma once
 #include "HarmonixMidi/MusicMapBase.h"
 #include "HarmonixMidi/MidiConstants.h"
+#include <limits>
 
 #include "ChordMap.generated.h"
 
@@ -53,7 +54,7 @@ public:
 
 	void Finalize(int32 LastTick);
 
-	void Copy(const FChordProgressionMap& Other, int32 StartTick = 0, int32 EndTick = -1);
+	void Copy(const FChordProgressionMap& Other, int32 StartTick = 0, int32 EndTick = std::numeric_limits<int32>::max());
 	bool IsEmpty() const;
 
 	/** Called by the midi file importer before map points are added to this map */
diff --git a/Engine/Plugins/Runtime/Harmonix/Source/HarmonixMidi/Public/HarmonixMidi/MidiConstants.h b/Engine/Plugins/Runtime/Harmonix/Source/HarmonixMidi/Public/HarmonixMidi/MidiConstants.h
--- a/Engine/Plugins/Runtime/Harmonix/Source/HarmonixMidi/Public/HarmonixMidi/MidiConstants.h
+++ b/Engine/Plugins/Runtime/Harmonix/Source/HarmonixMidi/Public/HarmonixMidi/MidiConstants.h
@@ -6,6 +6,9 @@
 
 HARMONIXMIDI_API DECLARE_LOG_CATEGORY_EXTERN(LogMIDI, Log, All);
 
+enum class EMidiClockSubdivisionQuantization : uint8;
+struct FTimeSignature;
+
 namespace Harmonix::Midi::Constants
 {
 	enum class EMidiTextEventEncoding
@@ -107,6 +110,8 @@ namespace Harmonix::Midi::Constants
 	// turns MIDI style tempo (microseconds per quarter note) into BPM...
 	HARMONIXMIDI_API inline float MidiTempoToBPM(int32 UsPerQuarterNote) { return UsPerQuarterNote== 0 ? 0 : 60000000.0f/(float)UsPerQuarterNote; }
 	HARMONIXMIDI_API inline int32 BPMToMidiTempo(float Bpm)       { return Bpm == 0.0 ? 0 : (int32)(60000000.0f / (float)Bpm); }
+	HARMONIXMIDI_API int32 SubdivisionToMidiTicks(const EMidiClockSubdivisionQuantization Division, const FTimeSignature& TimeSignature, const int32 TicksPerQuarterNote);
+	HARMONIXMIDI_API float SubdivisionToBeats(EMidiClockSubdivisionQuantization Subdivision, const FTimeSignature& TimeSignature);
 
 	HARMONIXMIDI_API float RoundToStandardBeatPrecision(float InBeat, int TimeSignatureDenominator);
 
diff --git a/Engine/Plugins/Runtime/Harmonix/Source/HarmonixMidi/Public/HarmonixMidi/MidiCursor.h b/Engine/Plugins/Runtime/Harmonix/Source/HarmonixMidi/Public/HarmonixMidi/MidiCursor.h
new file mode 100644
--- /dev/null
+++ b/Engine/Plugins/Runtime/Harmonix/Source/HarmonixMidi/Public/HarmonixMidi/MidiCursor.h
@@ -0,0 +1,53 @@
+// Copyright Epic Games, Inc. All Rights Reserved.
+
+#pragma once
+
+#include "MidiFile.h"
+
+class HARMONIXMIDI_API FMidiCursor
+{
+public:
+	static const int32 kAllTracks = -1;
+
+	void Prepare(UMidiFile* InMidiFile, int32 TrackIndex = kAllTracks, bool bResetState = true);
+	void Prepare(TSharedPtr<FMidiFileData> InMidiFile, int32 TrackIndex = kAllTracks, bool bResetState = true);
+
+	int32 GetNextTick() const { return NextTick; }
+
+	class HARMONIXMIDI_API FReceiver
+	{
+	public:
+
+		virtual void OnMidiMessage(int32 TrackIndex, int32 Tick, uint8 Status, uint8 Data1, uint8 Data2, bool bIsPreroll) {}
+		virtual void OnPreRollNoteOn(int32 TrackIndex, int32 EventTick, int32 Tick, float PreRollMs, uint8 Status, uint8 Data1, uint8 Data2) {}
+		virtual void OnText(int32 TrackIndex, int32 Tick, int32 TextIndex, const FString& Str, uint8 Type, bool bIsPreroll) {}
+		virtual void OnTempo(int32 TrackIndex, int32 Tick, int32 Tempo, bool bIsPreroll) {}
+		virtual void OnTimeSig(int32 TrackIndex, int32 Tick, int32 Numerator, int32 Denominator, bool bIsPreroll) {}
+		virtual bool HandleMessage(int32 TrackIndex, const FMidiTrack& Track, const FMidiEvent& Event, bool bIsPreroll);
+	};
+
+	void  SeekToNextTick(int32 NewNextTick, int32 PrerollBars = 0, FReceiver* PrerollReceiver = nullptr);
+	int32 SeekToMs(float NewPositionMs, int32 PrerollBars = 0, FReceiver* PrerollReceiver = nullptr);
+
+	void Process(int32 FirstTickToProcess, int32 LastTickToProcess, FReceiver& Receiver);
+	
+	// Preroll is like process except it filters meaningless messages that happen in the provided
+	// span of ticks. For example:
+	//    - If the span includes both a note-on AND a note-off for the same note on the same track
+	//      those events are 'eaten' and not provided to the receiver. 
+	//    - Only the last time signature and tempo in the span will be sent.
+	//    - etc. 
+	void Preroll(int32 FirstTickToProcess, int32 LastTickToProcess, FReceiver& Receiver);
+
+	bool PassedEnd() const;
+
+private:
+	void UpdateNextTick(int32 NewNextTick);
+
+	TSharedPtr<FMidiFileData> MidiFile;
+	TArray<int32> TrackNextEventIndexs;
+	float CurrentFileMs = 0.0f;
+	int32 NextTick = 0;
+	int32 LastTick = -1;
+	int32 WatchTrack = -1;
+};
\ No newline at end of file
diff --git a/Engine/Plugins/Runtime/Harmonix/Source/HarmonixMidi/Public/HarmonixMidi/MidiFile.h b/Engine/Plugins/Runtime/Harmonix/Source/HarmonixMidi/Public/HarmonixMidi/MidiFile.h
--- a/Engine/Plugins/Runtime/Harmonix/Source/HarmonixMidi/Public/HarmonixMidi/MidiFile.h
+++ b/Engine/Plugins/Runtime/Harmonix/Source/HarmonixMidi/Public/HarmonixMidi/MidiFile.h
@@ -194,7 +194,7 @@ public:
 	virtual void PostEditChangeProperty(FPropertyChangedEvent& PropertyChangedEvent) override;
 #endif
 
-	int32 GetStartBar() { return TheMidiData.SongMaps.GetBarMap().GetStartBar(); };
+	int32 GetStartBar() { return TheMidiData.SongMaps.GetStartBar(); };
 
 	// This does a quick and dirty to check to see if the length of the midi is 
 	// on a musical subdivision. It is used primarily during asset importing, and
@@ -216,6 +216,8 @@ public:
 	const FSongMaps* GetSongMaps() const { return &TheMidiData.SongMaps; }
 	FSongMaps* GetSongMaps() { return &TheMidiData.SongMaps; }
 
+	TSharedPtr<FMidiFileData> GetOrCreateRenderableCopy();
+
 	// IAudioProxyDataFactory
 	virtual TSharedPtr<Audio::IProxyData> CreateProxyData(const Audio::FProxyDataInitParams& InitParams) override;
 
diff --git a/Engine/Plugins/Runtime/Harmonix/Source/HarmonixMidi/Public/HarmonixMidi/MidiPlayCursorMgr.h b/Engine/Plugins/Runtime/Harmonix/Source/HarmonixMidi/Public/HarmonixMidi/MidiPlayCursorMgr.h
--- a/Engine/Plugins/Runtime/Harmonix/Source/HarmonixMidi/Public/HarmonixMidi/MidiPlayCursorMgr.h
+++ b/Engine/Plugins/Runtime/Harmonix/Source/HarmonixMidi/Public/HarmonixMidi/MidiPlayCursorMgr.h
@@ -116,9 +116,7 @@ public:
 
 	const UMidiFile::FMidiTrackList& Tracks() const;
 
-	const FSongMaps& GetSongMaps() const;
-	const FTempoMap& GetTempoMap() const;
-	const FBarMap& GetBarMap() const;
+	const ISongMapEvaluator& GetSongMaps() const;
 
 	int32 GetCurrentHiResTick() const   { return GetHiResTracker().CurrentTick;  }
 	float GetCurrentHiResMs() const     { return GetHiResTracker().CurrentMs;    }
diff --git a/Engine/Plugins/Runtime/Harmonix/Source/HarmonixMidi/Public/HarmonixMidi/MidiSongPos.h b/Engine/Plugins/Runtime/Harmonix/Source/HarmonixMidi/Public/HarmonixMidi/MidiSongPos.h
--- a/Engine/Plugins/Runtime/Harmonix/Source/HarmonixMidi/Public/HarmonixMidi/MidiSongPos.h
+++ b/Engine/Plugins/Runtime/Harmonix/Source/HarmonixMidi/Public/HarmonixMidi/MidiSongPos.h
@@ -8,7 +8,7 @@
 
 #include "MidiSongPos.generated.h"
 
-struct FSongMaps;
+struct ISongMapEvaluator;
 
 /////////////////////////////////////////////////////////////////////////////
 // Position within a song (midi info)
@@ -33,7 +33,7 @@ public:
 	float BarsIncludingCountIn = 0.0f; // total bars from the beginning of the song.
 	UPROPERTY(BlueprintReadOnly, Category = "MidiSongPosition")
 	float BeatsIncludingCountIn = 0.0f; // total beats from the beginning of the song.
-	UPROPERTY(BlueprintReadOnly, Category = "MidiSongPosition")
+	UPROPERTY(BlueprintReadOnly, Category	 = "MidiSongPosition")
 	EMusicalBeatType BeatType = EMusicalBeatType::Normal;
 	UPROPERTY(BlueprintReadOnly, Category = "MidiSongPosition")
 	FMusicTimestamp Timestamp;
@@ -55,11 +55,11 @@ public:
 	// etc.
 	void SetByTime(float InElapsedMs, float Bpm, int32 TimeSigNum = 4, int32 TimeSigDenom = 4, int32 StartBar = 1);
 
-	void SetByTime(float InElapsedMs, const FSongMaps& Maps);
+	void SetByTime(float InElapsedMs, const ISongMapEvaluator& Map);
 
 	// Low-level version for midi players / parsers that use the low-level
 	// "midi tick system" for advancing song position.
-	void SetByTick(float Tick, const FSongMaps& Maps);
+	void SetByTick(float Tick, const ISongMapEvaluator& Map);
 
 	void Reset()	
 	{
@@ -84,6 +84,9 @@ public:
 	bool operator>=(const FMidiSongPos& rhs) const;
 	bool operator==(const FMidiSongPos& rhs) const;
 
+	static FMidiSongPos Lerp(const FMidiSongPos& A, const FMidiSongPos& B, float Alpha);
+
 private:
-	void SetByTimeAndTick(float Ms, float Tick, const FSongMaps& Maps);
+	void SetByTimeAndTick(float Ms, float Tick, const ISongMapEvaluator& Map);
 };
+
diff --git a/Engine/Plugins/Runtime/Harmonix/Source/HarmonixMidi/Public/HarmonixMidi/MusicTimeSpan.h b/Engine/Plugins/Runtime/Harmonix/Source/HarmonixMidi/Public/HarmonixMidi/MusicTimeSpan.h
--- a/Engine/Plugins/Runtime/Harmonix/Source/HarmonixMidi/Public/HarmonixMidi/MusicTimeSpan.h
+++ b/Engine/Plugins/Runtime/Harmonix/Source/HarmonixMidi/Public/HarmonixMidi/MusicTimeSpan.h
@@ -3,7 +3,7 @@
 
 #include "MusicTimeSpan.generated.h"
 
-struct FSongMaps;
+struct ISongMapEvaluator;
 struct FMidiSongPos;
 
 UENUM()
@@ -76,18 +76,13 @@ public:
 		, Offset(InOffset)
 	{}
 
-	float CalcPositionInSpan(const FMidiSongPos& Position, const FSongMaps& Maps) const;
-	float CalcPositionInSpan(float Ms, const FSongMaps& Maps) const;
+	float CalcPositionInSpan(const FMidiSongPos& Position, const ISongMapEvaluator& Maps) const;
+	float CalcPositionInSpan(float Ms, const ISongMapEvaluator& Maps) const;
 
 private:
-	//float CalculateEnclosingVariableSizeSpanExtents(int32 Tick, int32 TimeSignatureNumerator, int32 TimeSignatureDenominator) const;
-	//float CalculateEnclosingVariableSizeSpanExtents(int32 Tick, const FSongMaps& Maps) const;
-	//float CalculateEnclosingFixedSizeSpanExtents(int32 Tick, const FSongMaps& Maps) const;
-	//float CalculateEnclosingFixedSizeSpanExtents(int32 Tick, int32 TicksPerQuarterNote) const;
+	float CalculateEnclosingVariableSizeSpanExtents(const FMidiSongPos& Position, const ISongMapEvaluator& Maps) const;
+	float CalculateEnclosingFixedSizeSpanExtents(const FMidiSongPos& Position, const ISongMapEvaluator& Maps) const;
 
-	float CalculateEnclosingVariableSizeSpanExtents(const FMidiSongPos& Position, const FSongMaps& Maps) const;
-	float CalculateEnclosingFixedSizeSpanExtents(const FMidiSongPos& Position, const FSongMaps& Maps) const;
-
-	float CalcPositionInSpanWithOffset(const FMidiSongPos& Position, const FSongMaps& Maps) const;
-	float CalcPositionInSpanNoOffset(const FMidiSongPos& Position, const FSongMaps& Maps) const;
+	float CalcPositionInSpanWithOffset(const FMidiSongPos& Position, const ISongMapEvaluator& Maps) const;
+	float CalcPositionInSpanNoOffset(const FMidiSongPos& Position, const ISongMapEvaluator& Maps) const;
 };
diff --git a/Engine/Plugins/Runtime/Harmonix/Source/HarmonixMidi/Public/HarmonixMidi/SectionMap.h b/Engine/Plugins/Runtime/Harmonix/Source/HarmonixMidi/Public/HarmonixMidi/SectionMap.h
--- a/Engine/Plugins/Runtime/Harmonix/Source/HarmonixMidi/Public/HarmonixMidi/SectionMap.h
+++ b/Engine/Plugins/Runtime/Harmonix/Source/HarmonixMidi/Public/HarmonixMidi/SectionMap.h
@@ -2,6 +2,7 @@
 #pragma once
 #include "HarmonixMidi/MusicMapBase.h"
 #include "HarmonixMidi/MidiConstants.h"
+#include <limits>
 
 #include "SectionMap.generated.h"
 
@@ -49,7 +50,7 @@ public:
 	void Finalize(int32 LastTick);
 
 	void Empty();
-	void Copy(const FSectionMap& Other, int32 StartTick = 0, int32 EndTick = -1);
+	void Copy(const FSectionMap& Other, int32 StartTick = 0, int32 EndTick = std::numeric_limits<int32>::max());
 	bool IsEmpty() const;
 
 	/** Called by the midi file importer before map points are added to this map */
diff --git a/Engine/Plugins/Runtime/Harmonix/Source/HarmonixMidi/Public/HarmonixMidi/SmoothedMidiPlayCursor.h b/Engine/Plugins/Runtime/Harmonix/Source/HarmonixMidi/Public/HarmonixMidi/SmoothedMidiPlayCursor.h
--- a/Engine/Plugins/Runtime/Harmonix/Source/HarmonixMidi/Public/HarmonixMidi/SmoothedMidiPlayCursor.h
+++ b/Engine/Plugins/Runtime/Harmonix/Source/HarmonixMidi/Public/HarmonixMidi/SmoothedMidiPlayCursor.h
@@ -28,7 +28,7 @@ protected:
 
 private:
 	static const int kFramesOfErrorHistory = 10;
-	FLocalMinimumMagnitudeTracker<kFramesOfErrorHistory> ErrorTracker;
+	FLocalMinimumMagnitudeTracker<float, kFramesOfErrorHistory> ErrorTracker;
 	FVariableSpeedTimer SmoothingTimer;
 
 	float SmoothingLatencyMs = 30.0f;
diff --git a/Engine/Plugins/Runtime/Harmonix/Source/HarmonixMidi/Public/HarmonixMidi/SongMaps.h b/Engine/Plugins/Runtime/Harmonix/Source/HarmonixMidi/Public/HarmonixMidi/SongMaps.h
--- a/Engine/Plugins/Runtime/Harmonix/Source/HarmonixMidi/Public/HarmonixMidi/SongMaps.h
+++ b/Engine/Plugins/Runtime/Harmonix/Source/HarmonixMidi/Public/HarmonixMidi/SongMaps.h
@@ -6,6 +6,7 @@
 #include "HarmonixMidi/ChordMap.h"
 #include "HarmonixMidi/SectionMap.h"
 #include "Sound/QuartzQuantizationUtilities.h"
+#include <limits>
 
 #include "SongMaps.generated.h"
 
@@ -13,6 +14,12 @@ class IMidiReader;
 class FStdMidiFileReader;
 class FSongMapReceiver;
 class UMidiFile;
+struct FMidiFileData;
+
+namespace HarmonixMetasound 
+{
+	class FMidiClock;
+}
 
 UENUM()
 enum class EMidiFileQuantizeDirection : uint8
@@ -66,6 +73,120 @@ public:
 	}
 };
 
+struct HARMONIXMIDI_API ISongMapEvaluator
+{
+public:
+	int32 GetTicksPerQuarterNote() const;
+	float TickToMs(float Tick) const;
+	float MsToTick(float Ms) const;
+	float GetCountInSeconds() const;
+
+	// tempo
+	const FTempoInfoPoint* GetTempoInfoForMs(float Ms) const;
+	const FTempoInfoPoint* GetTempoInfoForTick(int32 Tick) const;
+	int32 GetTempoPointIndexForTick(int32 Tick) const;
+	const FTempoInfoPoint* GetTempoInfoPoint(int32 PointIndex) const;
+	int32 GetNumTempoChanges() const;
+	int32 GetTempoChangePointTick(int32 PointIndex) const;
+	float GetTempoAtMs(float Ms) const;
+	float GetTempoAtTick(int32 Tick) const;
+	bool TempoMapIsEmpty() const;
+
+	// beats
+	const FBeatMapPoint* GetBeatAtMs(float Ms) const;
+	float GetMsAtBeat(float Beat) const;
+	const FBeatMapPoint* GetBeatAtTick(int32 Tick) const;
+	float GetMsPerBeatAtMs(float Ms) const;
+	float GetMsPerBeatAtTick(int32 Tick) const;
+	float GetFractionalBeatAtMs(float Ms) const;
+	float GetFractionalBeatAtTick(float Tick) const;
+	int32 GetBeatIndexAtMs(float Ms) const;
+	int32 GetBeatIndexAtTick(int32 Tick) const;
+	EMusicalBeatType GetBeatTypeAtMs(float Ms) const;
+	EMusicalBeatType GetBeatTypeAtTick(int32 Tick) const;
+	const FBeatMapPoint* GetBeatPointInfoAtTick(int32 Tick, int32* PointIndex = nullptr) const;
+
+	float GetBeatInPulseBarAtMs(float Ms) const;
+	float GetBeatInPulseBarAtTick(float Tick) const;
+	int32 GetNumBeatsInPulseBarAtMs(float Ms) const;
+	int32 GetNumBeatsInPulseBarAtTick(int32 Tick) const;
+	bool BeatMapIsEmpty() const;
+
+	// bars
+	int32 GetStartBar() const;
+	int32 GetNumTimeSignatureChanges() const;
+	const FTimeSignature* GetTimeSignatureAtMs(float Ms) const;
+	const FTimeSignature* GetTimeSignatureAtTick(int32 Tick) const;
+	int32 GetTimeSignaturePointIndexForTick(int32 Tick) const;
+	const FTimeSignature* GetTimeSignatureAtBar(int32 Bar) const;
+	const FTimeSignaturePoint* GetTimeSignaturePointAtTick(int32 Tick) const;
+	const FTimeSignaturePoint* GetTimeSignaturePoint(int32 PointIndex) const;
+	int32 GetTimeSignatureChangePointTick(int32 PointIndex);
+	float GetBarIncludingCountInAtMs(float Ms) const;
+	float GetBarIncludingCountInAtTick(float Tick) const;
+	float GetMsPerBarAtMs(float Ms) const;
+	float GetMsPerBarAtTick(float Tick) const;
+	bool BarMapIsEmpty() const;
+	FMusicTimestamp TickToMusicTimestamp(float Tick, int32* OutBeatsPerBar = nullptr) const;
+	int32 BarIncludingCountInToTick(int32 BarIndex, int32* OutBeatsPerBar = nullptr, int32* OutTicksPerBeat = nullptr) const;
+	int32 BarBeatTickIncludingCountInToTick(int32 BarIndex, int32 BeatInBar, int32 TickInBeat) const;
+	float FractionalBarIncludingCountInToTick(float FractionalBarIndex) const;
+	int32 TickToBarIncludingCountIn(int32 Tick) const;
+	float TickToFractionalBarIncludingCountIn(float Tick) const;
+	void TickToBarBeatTickIncludingCountIn(int32 RawTick, int32& OutBarIndex, int32& OutBeatInBarIndex, int32& OutTickIndexInBeat, int32* OutBeatsPerBar = nullptr, int32* OutTicksPerBeat = nullptr) const;
+	int32 CalculateMidiTick(const FMusicTimestamp& Timestamp, const EMidiClockSubdivisionQuantization Quantize) const;
+	int32 SubdivisionToMidiTicks(const EMidiClockSubdivisionQuantization Division, const int32 AtTick) const;
+	float MusicTimestampToTick(const FMusicTimestamp& Timestamp) const;
+	int32 MusicTimestampBarToTick(int32 BarNumber, int32* OutBeatsPerBar = nullptr, int32* OutTicksPerBeat = nullptr) const;
+
+	// sections
+	const TArray<FSongSection>& GetSections() const;
+	int32 GetNumSections() const;
+	float GetSectionStartMsAtMs(float Ms) const;
+	float GetSectionEndMsAtMs(float Ms) const;
+	const FSongSection* GetSectionAtMs(float Ms) const;
+	const FSongSection* GetSectionAtTick(int32 Tick) const;
+	int32 GetSectionIndexAtTick(int32 Tick) const;
+	const FSongSection* GetSectionWithName(const FString& Name) const;
+	FString GetSectionNameAtMs(float Ms) const;
+	FString GetSectionNameAtTick(int32 Tick) const;
+	float GetSectionLengthMsAtMs(float Ms) const;
+	float GetSectionLengthMsAtTick(int32 Tick) const;
+	bool SectionMapIsEmpty() const;
+
+	// chords
+	const FChordMapPoint* GetChordAtMs(float Ms) const;
+	const FChordMapPoint* GetChordAtTick(int32 Tick) const;
+	FName GetChordNameAtMs(float Ms) const;
+	FName GetChordNameAtTick(int32 Tick) const;
+	float GetChordLengthMsAtMs(float Ms) const;
+	float GetChordLengthMsAtTick(int32 Tick) const;
+	bool ChordMapIsEmpty() const;
+
+	// length
+	float GetSongLengthMs() const;
+	int32 GetSongLengthBeats() const;
+	float GetSongLengthFractionalBars() const;
+	bool LengthIsAPerfectSubdivision() const;
+	FString GetSongLengthString() const;
+
+	int32 QuantizeTickToAnyNearestSubdivision(int32 InTick, EMidiFileQuantizeDirection Direction, EMidiClockSubdivisionQuantization& Division) const;
+	int32 QuantizeTickToNearestSubdivision(int32 InTick, EMidiFileQuantizeDirection Direction, EMidiClockSubdivisionQuantization Division) const;
+	void GetTicksForNearestSubdivision(int32 InTick, EMidiClockSubdivisionQuantization Division, int32& LowerTick, int32& UpperTick) const;
+
+	//virtual FSongLengthData& GetSongLengthData() = 0;
+	virtual const FSongLengthData& GetSongLengthData() const = 0;
+
+private:
+	friend struct FSongMaps;
+	virtual const FTempoMap& GetTempoMap() const = 0;
+	virtual const FBeatMap& GetBeatMap() const = 0;
+	virtual const FBarMap& GetBarMap() const = 0;
+	virtual const FSectionMap& GetSectionMap() const = 0;
+	virtual const FChordProgressionMap& GetChordMap() const = 0;
+};
+
+
 /**
  * FSongMaps encapsulates a number of other musical/midi map types
  * that are very useful for musical gameplay and interactivity. 
@@ -75,131 +196,76 @@ public:
  * chord, etc.
  */
 USTRUCT(BlueprintType)
-struct HARMONIXMIDI_API FSongMaps
+struct HARMONIXMIDI_API FSongMaps 
+#if CPP
+	: public ISongMapEvaluator
+#endif
 {
 	GENERATED_BODY()
 
 public:
 	FSongMaps();
+	FSongMaps(float Bpm, int32 TimeSigNumerator, int32 TimeSigDenominator);
+	FSongMaps(const ISongMapEvaluator& Other);
+	virtual ~FSongMaps() = default;
+
 	bool operator==(const FSongMaps& Other) const;
 
 	void Init(int32 InTicksPerQuarterNote);
-	void Copy(const FSongMaps& Other, int32 StartTick = 0, int32 EndTick = -1);
+	void Copy(const ISongMapEvaluator& Other, int32 StartTick = 0, int32 EndTick = std::numeric_limits<int32>::max());
 
 	// For importing...
 	bool LoadFromStdMidiFile(const FString& FilePath);
 	bool LoadFromStdMidiFile(void* Buffer, int32 BufferSize, const FString& Filename);
 	bool LoadFromStdMidiFile(TSharedPtr<FArchive> Archive, const FString& Filename);
 
-	float TickToMs(float Tick) const;
-	float MsToTick(float Ms) const;
-	float GetCountInSeconds() const;
-
 	// tracks
-	TArray<FString>&       GetTrackNames() { return TrackNames; }
+	TArray<FString>& GetTrackNames() { return TrackNames; }
 	const TArray<FString>& GetTrackNames() const { return TrackNames; }
-	FString                GetTrackName(int32 Index) const;
-	bool                   TrackNamesIsEmpty() const { return TrackNames.IsEmpty(); }
-	void                   EmptyTrackNames() { TrackNames.Empty(); }
-
-	// tempo
-	const FTempoInfoPoint* GetTempoInfoForMs(float Ms) const;
-	const FTempoInfoPoint* GetTempoInfoForTick(int32 Tick) const;
-	float                  GetTempoAtMs(float Ms) const;
-	float                  GetTempoAtTick(int32 Tick) const;
-	FTempoMap&             GetTempoMap() { return TempoMap; }
-	const FTempoMap&       GetTempoMap() const { return TempoMap; }
-	bool                   TempoMapIsEmpty() { return TempoMap.GetNumTempoChangePoints() == 0; }
-	void                   EmptyTempoMap() { TempoMap.Empty(); }
-
-	// pulses
-	const FBeatMapPoint*  GetBeatAtMs(float Ms) const;
-	float                 GetMsAtBeat(float Beat) const;
-	const FBeatMapPoint*  GetBeatAtTick(int32 Tick) const;
-	float                 GetMsPerBeatAtMs(float Ms) const;
-	float                 GetMsPerBeatAtTick(int32 Tick) const;
-	float                 GetFractionalBeatAtMs(float Ms) const;
-	float                 GetFractionalBeatAtTick(float Tick) const;
-	int32                 GetBeatIndexAtMs(float Ms) const;
-	int32                 GetBeatIndexAtTick(int32 Tick) const;
-	EMusicalBeatType      GetBeatTypeAtMs(float Ms) const;
-	EMusicalBeatType      GetBeatTypeAtTick(int32 Tick) const;
-	float                 GetBeatInPulseBarAtMs(float Ms) const;
-	float                 GetBeatInPulseBarAtTick(float Tick) const;
-	int32                 GetNumBeatsInPulseBarAtMs(float Ms) const ;
-	int32                 GetNumBeatsInPulseBarAtTick(int32 Tick) const;
-	FBeatMap&             GetBeatMap() { return BeatMap; }
-	const FBeatMap&       GetBeatMap() const { return BeatMap; }
-	bool                  BeatMapIsEmpty() { return BeatMap.GetNumMapPoints() == 0; }
-	void                  EmptyBeatMap() { BeatMap.Empty(); }
+	FString GetTrackName(int32 Index) const;
+	bool TrackNamesIsEmpty() const { return TrackNames.IsEmpty(); }
+	void EmptyTrackNames() { TrackNames.Empty(); }
 
-	// bars
-	const FTimeSignature* GetTimeSignatureAtMs(float Ms) const;
-	const FTimeSignature* GetTimeSignatureAtTick(int32 Tick) const;
-	const FTimeSignature* GetTimeSignatureAtBar(int32 Bar) const;
-	float                 GetBarIncludingCountInAtMs(float Ms) const;
-	float                 GetBarIncludingCountInAtTick(float Tick) const;
-	float                 GetMsPerBarAtMs(float Ms) const;
-	float                 GetMsPerBarAtTick(float Tick) const;
-	FBarMap&              GetBarMap() { return BarMap; }
-	const FBarMap&        GetBarMap() const { return BarMap; }
-	bool                  BarMapIsEmpty() { return BarMap.GetNumTimeSignaturePoints() == 0; }
-	void                  EmptyBarMap() { BarMap.Empty(); }
-	void                  SetLengthTotalBars(int32 Bars);
-	int32                 CalculateMidiTick(const FMusicTimestamp& Timestamp, const EMidiClockSubdivisionQuantization Quantize) const;
-	int32                 SubdivisionToMidiTicks(const EMidiClockSubdivisionQuantization Division, const int32 AtTick) const;
-	int32                 SubdivisionToMidiTicks(const EMidiClockSubdivisionQuantization Division, const FTimeSignature& TimeSignature) const;
-	static float          SubdivisionToBeats(EMidiClockSubdivisionQuantization Subdivision, const FTimeSignature& TimeSignature);
-
-
-	// sections
-	float               GetSectionStartMsAtMs(float Ms) const;
-	float               GetSectionEndMsAtMs(float Ms) const;
-	const FSongSection* GetSectionAtMs(float Ms) const;
-	const FSongSection* GetSectionAtTick(int32 Tick) const;
-	const FSongSection* GetSectionWithName(const FString& Name) const;
-	FString             GetSectionNameAtMs(float Ms) const;
-	FString             GetSectionNameAtTick(int32 Tick) const;
-	float               GetSectionLengthMsAtMs(float Ms) const;
-	float               GetSectionLengthMsAtTick(int32 Tick) const;
-	FSectionMap&        GetSectionMap() { return SectionMap; }
-	const FSectionMap&  GetSectionMap() const { return SectionMap; }
-	bool                SectionMapIsEmpty() { return SectionMap.GetNumSections() == 0; }
-	void                EmptySectionMap() { SectionMap.Empty(); }
-
-	// chords
-	const FChordMapPoint* GetChordAtMs(float Ms) const;
-	const FChordMapPoint* GetChordAtTick(int32 Tick) const;
-	FName                 GetChordNameAtMs(float Ms) const;
-	FName                 GetChordNameAtTick(int32 Tick) const;
-	float                 GetChordLengthMsAtMs(float Ms) const;
-	float                 GetChordLengthMsAtTick(int32 Tick) const;
-	FChordProgressionMap& GetChordMap() { return ChordMap; }
-	const FChordProgressionMap&  GetChordMap() const { return ChordMap; }
-	bool                  ChordMapIsEmpty() const { return ChordMap.GetNumChords() == 0; }
-	void                  EmptyChordMap() { ChordMap.Empty(); }
+	void EmptyTempoMap() { TempoMap.Empty(); }
+	void EmptyBeatMap() { BeatMap.Empty(); }
+	void EmptyBarMap() { BarMap.Empty(); }
+	void SetLengthTotalBars(int32 Bars);
+	void EmptySectionMap() { SectionMap.Empty(); }
+	void EmptyChordMap() { ChordMap.Empty(); }
 
 	void EmptyAllMaps();
 	bool IsEmpty() const;
 
+	// BEGIN ISongMapEvaluator overrides (Immutable access to maps)
+	virtual const FTempoMap& GetTempoMap() const override { return TempoMap; }
+	virtual const FBeatMap& GetBeatMap() const override { return BeatMap; }
+	virtual const FBarMap& GetBarMap() const override { return BarMap; }
+	virtual const FSectionMap& GetSectionMap() const override { return SectionMap; }
+	virtual const FChordProgressionMap& GetChordMap() const override { return ChordMap; }
+	virtual const FSongLengthData& GetSongLengthData() const override { return LengthData; }
+	// END ISongMapEvaluator overrides
+
+	FTempoMap& GetTempoMap() { return TempoMap; }
+	FBeatMap& GetBeatMap() { return BeatMap; }
+	FBarMap& GetBarMap() { return BarMap; }
+	FSectionMap& GetSectionMap() { return SectionMap; }
+	FChordProgressionMap& GetChordMap() { return ChordMap; }
 	FSongLengthData& GetSongLengthData() { return LengthData; }
-	const FSongLengthData& GetSongLengthData() const { return LengthData; }
-	float GetSongLengthMs() const;
-	int32 GetSongLengthBeats() const;
-	float GetSongLengthFractionalBars() const;
 
 	int32 GetTicksPerQuarterNote() const { return TicksPerQuarterNote; }
 
+	void SetStartBar(int32 StartBar);
 	void SetSongLengthTicks(int32 NewLengthTicks);
-	bool LengthIsAPerfectSubdivision() const;
+	void FinalizeBarMap(int32 InLastTick);
 
-	int32 QuantizeTickToAnyNearestSubdivision(int32 InTick, EMidiFileQuantizeDirection Direction, EMidiClockSubdivisionQuantization& Division) const;
-	int32 QuantizeTickToNearestSubdivision(int32 InTick, EMidiFileQuantizeDirection Direction, EMidiClockSubdivisionQuantization Division) const;
-	void GetTicksForNearestSubdivision(int32 InTick, EMidiClockSubdivisionQuantization Division, int32& LowerTick, int32& UpperTick) const;
-	FString GetSongLengthString() const;
+	bool AddTempoInfoPoint(int32 MicrosecondsPerQuarterNote, int32 Tick, bool SortNow = true);
+	bool AddTimeSignatureAtBarIncludingCountIn(int32 BarIndex, int32 InNumerator, int32 InDenominator, bool SortNow = true, bool FailOnError = true);
+	FTimeSignaturePoint* GetMutableTimeSignaturePoint(int32 PointIndex);
 
 protected:
 	friend class FSongMapReceiver;
+	friend class HarmonixMetasound::FMidiClock;
+
 	UPROPERTY()
 	int32 TicksPerQuarterNote = Harmonix::Midi::Constants::GTicksPerQuarterNoteInt;
 	UPROPERTY()
@@ -219,8 +285,58 @@ private:
 	UPROPERTY()
 	FSongLengthData LengthData;
 
+
 	void StringLengthToMT(const FString& LengthString, int32& OutBars, int32& OutTicks);
 	bool ReadWithReader(FStdMidiFileReader& Reader);
 	bool FinalizeRead(IMidiReader* Reader);
 };
 
+struct HARMONIXMIDI_API FSongMapsWithAlternateTempoSource : public ISongMapEvaluator
+{
+public:
+	FSongMapsWithAlternateTempoSource(const TSharedPtr<const FMidiFileData>& MidiWithTempo, const TSharedPtr<const FMidiFileData>& MidiWithOthers)
+		: MidiFileWithTempoMap(MidiWithTempo)
+		, MidiFileWithOtherMaps(MidiWithOthers)
+	{}
+
+	FSongMapsWithAlternateTempoSource(const TSharedPtr<const FMidiFileData>& MidiData)
+		: MidiFileWithTempoMap(MidiData)
+		, MidiFileWithOtherMaps(MidiData)
+	{
+	}
+
+	FSongMapsWithAlternateTempoSource(const TSharedPtr<const FSongMapsWithAlternateTempoSource>& Other)
+		: MidiFileWithTempoMap(Other->MidiFileWithTempoMap)
+		, MidiFileWithOtherMaps(Other->MidiFileWithOtherMaps)
+	{}
+
+	FSongMapsWithAlternateTempoSource& operator=(const TSharedPtr<const FSongMapsWithAlternateTempoSource>& Other)
+	{
+		MidiFileWithTempoMap = Other->MidiFileWithTempoMap;
+		MidiFileWithOtherMaps = Other->MidiFileWithOtherMaps;
+		return *this;
+	}
+
+	virtual ~FSongMapsWithAlternateTempoSource() = default;
+
+	operator bool() const 
+	{
+		return MidiFileWithTempoMap.IsValid() && MidiFileWithOtherMaps.IsValid();
+	}
+
+	const TSharedPtr<const FMidiFileData>& GetMidiFileWithTempoMap() const { return MidiFileWithTempoMap; }
+	const TSharedPtr<const FMidiFileData>& GetMidiFileWithOtherMaps() const { return MidiFileWithOtherMaps; }
+
+	bool AllMapsHaveOneSource() const { return MidiFileWithTempoMap == MidiFileWithOtherMaps; }
+
+private:
+	virtual const FTempoMap& GetTempoMap() const override;
+	virtual const FBeatMap& GetBeatMap() const override;
+	virtual const FBarMap& GetBarMap() const override;
+	virtual const FSectionMap& GetSectionMap() const override;
+	virtual const FChordProgressionMap& GetChordMap() const override;
+	virtual const FSongLengthData& GetSongLengthData() const override;
+
+	TSharedPtr<const FMidiFileData> MidiFileWithTempoMap;
+	TSharedPtr<const FMidiFileData> MidiFileWithOtherMaps;
+};
diff --git a/Engine/Plugins/Runtime/Harmonix/Source/HarmonixMidi/Public/HarmonixMidi/TempoMap.h b/Engine/Plugins/Runtime/Harmonix/Source/HarmonixMidi/Public/HarmonixMidi/TempoMap.h
--- a/Engine/Plugins/Runtime/Harmonix/Source/HarmonixMidi/Public/HarmonixMidi/TempoMap.h
+++ b/Engine/Plugins/Runtime/Harmonix/Source/HarmonixMidi/Public/HarmonixMidi/TempoMap.h
@@ -2,6 +2,7 @@
 #pragma once
 #include "HarmonixMidi/MidiConstants.h"
 #include "HarmonixMidi/MusicMapBase.h"
+#include <limits>
 
 #include "TempoMap.generated.h"
 
@@ -33,6 +34,8 @@ public:
 
 	float GetBPM() const { return Harmonix::Midi::Constants::MidiTempoToBPM(MidiTempo); }
 
+	float GetMsPerQuarterNote() const { return (float)MidiTempo / 1000.0f; }
+
 	UPROPERTY()
 	float Ms = 0;        // The time at which the tempo is changing
 	UPROPERTY()
@@ -68,7 +71,7 @@ public:
 	bool operator==(const FTempoMap& Other) const;
 
 	void Empty();
-	void Copy(const FTempoMap& Other, int32 StartTick = 0, int32 EndTick = -1);
+	void Copy(const FTempoMap& Other, int32 StartTick = 0, int32 EndTick = std::numeric_limits<int32>::max());
 	bool IsEmpty() const;
 
 	/** Called by the midi file importer before map points are added to this map */
@@ -77,13 +80,6 @@ public:
 		TicksPerQuarterNote = InTicksPerQuarterNote;
 	}
 
-/*
-	template <typename T>
-	float TickToMs(T Tick) const;
-	template <typename T>
-	float TickToSeconds(T Tick) const;
-*/
-
 	/** Get the time (in milliseconds) at a given tick: */
 	float TickToMs(float Tick) const;
 
@@ -128,6 +124,9 @@ public:
 	/** Get the tempo info point */
 	const FTempoInfoPoint* GetTempoPointAtTick(int32 Tick) const;
 
+	/** Get the tempo info point */
+	int32 GetTempoPointIndexAtTick(int32 Tick) const;
+
 	/** The number of 'tempo events' in the song */
 	int32 GetNumTempoChangePoints() const;
 
@@ -148,6 +147,8 @@ public:
 	// Returns the tempo info points for inspection.
 	const TArray<FTempoInfoPoint>& GetTempoPoints() const { return Points; }
 
+	int32 GetTicksPerQuarterNote() const { return TicksPerQuarterNote; }
+
 protected:
 	UPROPERTY()
 	int32 TicksPerQuarterNote;
diff --git a/Engine/Plugins/Runtime/Harmonix/Source/HarmonixMidiTests/Private/ConformMidiFileLengthTests.cpp b/Engine/Plugins/Runtime/Harmonix/Source/HarmonixMidiTests/Private/ConformMidiFileLengthTests.cpp
--- a/Engine/Plugins/Runtime/Harmonix/Source/HarmonixMidiTests/Private/ConformMidiFileLengthTests.cpp
+++ b/Engine/Plugins/Runtime/Harmonix/Source/HarmonixMidiTests/Private/ConformMidiFileLengthTests.cpp
@@ -40,7 +40,6 @@ namespace HarmonixMidiTests::ConformMidiFileLength
 
 		// These will be useful below...
 		const FSongMaps* SongMaps = InMidiFile->GetSongMaps();
-		const FBarMap& BarMap = SongMaps->GetBarMap();
 
 		RemovableNotes = 0;
 		RemoveableControlChanges = 0;
@@ -54,7 +53,7 @@ namespace HarmonixMidiTests::ConformMidiFileLength
 			{
 				for (int32 Bar = 0; Bar < FMath::FloorToInt32(InFileLengthBars); ++Bar)
 				{
-					int32 DestinationTick = BarMap.BarBeatTickIncludingCountInToTick(Bar, 1, 0);
+					int32 DestinationTick = SongMaps->BarBeatTickIncludingCountInToTick(Bar, 1, 0);
 					int32 Duration = SongMaps->SubdivisionToMidiTicks(EMidiClockSubdivisionQuantization::Beat, DestinationTick) - 1;
 					AddNoteOnNoteOffPairToFile(InMidiFile, DefaultNoteNumber, DefaultNoteVelocity, TrackIndex, Channel, DestinationTick, Duration);
 
@@ -81,7 +80,7 @@ namespace HarmonixMidiTests::ConformMidiFileLength
 				// If bar length is a fractional number, add additional midi events after the last integer bar
 				if (InFileLengthBars != (int32)InFileLengthBars)
 				{
-					int32 DestinationTick = BarMap.FractionalBarIncludingCountInToTick(InFileLengthBars);
+					int32 DestinationTick = SongMaps->FractionalBarIncludingCountInToTick(InFileLengthBars);
 					int32 Duration = SongMaps->SubdivisionToMidiTicks(EMidiClockSubdivisionQuantization::Beat, DestinationTick) - 1;
 					int32 NoteDestinationTick = DestinationTick - (Duration + 1);
 					DestinationTick--;
diff --git a/Engine/Plugins/Runtime/Harmonix/Source/HarmonixMidiTests/Private/MidiFileImportExportTests.cpp b/Engine/Plugins/Runtime/Harmonix/Source/HarmonixMidiTests/Private/MidiFileImportExportTests.cpp
--- a/Engine/Plugins/Runtime/Harmonix/Source/HarmonixMidiTests/Private/MidiFileImportExportTests.cpp
+++ b/Engine/Plugins/Runtime/Harmonix/Source/HarmonixMidiTests/Private/MidiFileImportExportTests.cpp
@@ -46,7 +46,6 @@ namespace HarmonixMidiTests::ImportExportTests
 
 		// these will be useful below...
 		const FSongMaps* SongMaps =  InMidiFile->GetSongMaps();
-		const FBarMap& BarMap =SongMaps->GetBarMap();
 		
 		//Add events to track 1 - InNumTracks, Track 0 is the conductor track
 		for (int32 TrackIndex = 1; TrackIndex < (InNumTracksExcludingConductorTrack + 1); ++TrackIndex)
@@ -56,7 +55,7 @@ namespace HarmonixMidiTests::ImportExportTests
 			{
 				for (int32 Bar = 0; Bar < FMath::CeilToInt32(InFileLengthBars); ++Bar)
 				{
-					int32 DestinationTick = BarMap.BarBeatTickIncludingCountInToTick(Bar, 1, 0);
+					int32 DestinationTick = SongMaps->BarBeatTickIncludingCountInToTick(Bar, 1, 0);
 					int32 Duration = SongMaps->SubdivisionToMidiTicks(EMidiClockSubdivisionQuantization::Beat, DestinationTick) - 1;
 					AddNoteOnNoteOffPairToFile(InMidiFile, DefaultNoteNumber, DefaultNoteVelocity, TrackIndex, Channel, DestinationTick, Duration);
 
@@ -80,7 +79,7 @@ namespace HarmonixMidiTests::ImportExportTests
 				//if bar length is a fractional number, add additional midi events after the last integer bar
 				if (InFileLengthBars != (int32)InFileLengthBars)
 				{
-					int32 DestinationTick = BarMap.FractionalBarIncludingCountInToTick(InFileLengthBars) - 1;
+					int32 DestinationTick = SongMaps->FractionalBarIncludingCountInToTick(InFileLengthBars) - 1;
 					int32 Duration = SongMaps->SubdivisionToMidiTicks(EMidiClockSubdivisionQuantization::Beat, DestinationTick) - 1;
 					int32 NoteDestinationTick = DestinationTick - (Duration + 1);
 
diff --git a/Engine/Plugins/Runtime/Harmonix/Source/HarmonixMidiTests/Private/MidiTestUtility.cpp b/Engine/Plugins/Runtime/Harmonix/Source/HarmonixMidiTests/Private/MidiTestUtility.cpp
--- a/Engine/Plugins/Runtime/Harmonix/Source/HarmonixMidiTests/Private/MidiTestUtility.cpp
+++ b/Engine/Plugins/Runtime/Harmonix/Source/HarmonixMidiTests/Private/MidiTestUtility.cpp
@@ -10,7 +10,6 @@ namespace Harmonix::Testing::Utility::MidiTestUtility
 		//some default values for creating midi events
 		constexpr int32 DefaultNoteNumber = 60;//C4
 		constexpr int32 DefaultNoteVelocity = 90;
-		constexpr int32 DefaultTicksPerQuarter = 960;
 		constexpr uint8 DefaultControllerID = 69; //Hold Pedal
 		constexpr uint8 DefaultControlValue = 127;
 		constexpr uint8 DefaultPitchBendValueLSB = 64;
@@ -18,9 +17,8 @@ namespace Harmonix::Testing::Utility::MidiTestUtility
 		constexpr uint8 DefaultPolyPresValue = 127;
 
 		// Set initial tempo and time signature
-		TheMidiFile->GetSongMaps()->GetTempoMap().AddTempoInfoPoint(Midi::Constants::BPMToMidiTempo(InTempo), 0);
-		TheMidiFile->GetSongMaps()->GetBarMap().AddTimeSignatureAtBarIncludingCountIn(0, InTimeSigNum, InTimeSigDenom);
-		TheMidiFile->GetSongMaps()->GetBarMap().SetTicksPerQuarterNote(DefaultTicksPerQuarter);
+		TheMidiFile->GetSongMaps()->AddTempoInfoPoint(Midi::Constants::BPMToMidiTempo(InTempo), 0);
+		TheMidiFile->GetSongMaps()->AddTimeSignatureAtBarIncludingCountIn(0, InTimeSigNum, InTimeSigDenom);
 		// Whenever you change the song maps you need to either:
 		// a) Add corresponding midi events to the conductor track. Or...
 		// b) Build the conductor track from scratch.
@@ -28,7 +26,7 @@ namespace Harmonix::Testing::Utility::MidiTestUtility
 		TheMidiFile->BuildConductorTrack();
 
 		// now we can calculate the tick length since the midi file now has a conductor track...
-		int32 TickLength = TheMidiFile->GetSongMaps()->GetBarMap().FractionalBarIncludingCountInToTick(FileLengthBars);
+		int32 TickLength = TheMidiFile->GetSongMaps()->FractionalBarIncludingCountInToTick(FileLengthBars);
 		if (TickLength <= 0)
 		{
 			// MIDI Files must be at least one tick long. Even the initial tempo and time signature are specified on tick 0!
