Massive Midi Clock Refactor

diff --git a/Engine/Plugins/Runtime/Harmonix/Source/HarmonixMidiTests/Private/SongMapsTests.cpp b/Engine/Plugins/Runtime/Harmonix/Source/HarmonixMidiTests/Private/SongMapsTests.cpp
--- a/Engine/Plugins/Runtime/Harmonix/Source/HarmonixMidiTests/Private/SongMapsTests.cpp
+++ b/Engine/Plugins/Runtime/Harmonix/Source/HarmonixMidiTests/Private/SongMapsTests.cpp
@@ -13,12 +13,12 @@ namespace HarmonixMidiTests::SongMaps
 
 	static constexpr bool GLogQualtizationDetails = true;
 
-	void LogQuantizationDetails(const FBarMap& BarMap, int32 OriginalTick, int32 QuantizedTick, EMidiClockSubdivisionQuantization ResultDivision)
+	void LogQuantizationDetails(const ISongMapEvaluator& Map, int32 OriginalTick, int32 QuantizedTick, EMidiClockSubdivisionQuantization ResultDivision)
 	{
 		int32 BarIndex = 0;
 		int32 BeatInBar = 0;
 		int32 TickIndexInBeat = 0;
-		BarMap.TickToBarBeatTickIncludingCountIn(QuantizedTick, BarIndex, BeatInBar, TickIndexInBeat);
+		Map.TickToBarBeatTickIncludingCountIn(QuantizedTick, BarIndex, BeatInBar, TickIndexInBeat);
 		UE_LOG(LogSongMapsTest, Log, TEXT("Tick -> %d: Quantized %s to %d, Division = %s -> %d | %d | %d"), OriginalTick, (QuantizedTick < OriginalTick ? TEXT("^") : TEXT("v")), QuantizedTick, *StaticEnum<EQuartzCommandQuantization>()->GetDisplayNameTextByIndex((int32)ResultDivision).ToString(), BarIndex, BeatInBar, TickIndexInBeat);
 	}
 
@@ -28,18 +28,15 @@ namespace HarmonixMidiTests::SongMaps
 		check(MidiData);
 
 		// make 97bpm, 4/4  tempo map...
-		FTempoMap& TempoMap = MidiData->SongMaps.GetTempoMap();
-		TempoMap.Empty();
-		FBarMap& BarMap = MidiData->SongMaps.GetBarMap();
-		BarMap.Empty();
+		MidiData->SongMaps.EmptyAllMaps();
 		MidiData->Tracks.Empty();
 
 		MidiData->Tracks.Add(FMidiTrack(TEXT("conductor")));
 		MidiData->Tracks[0].AddEvent(FMidiEvent(0, FMidiMsg(Numerator, Denominator)));
-		BarMap.AddTimeSignatureAtBarIncludingCountIn(0, Numerator, Denominator);
+		MidiData->SongMaps.AddTimeSignatureAtBarIncludingCountIn(0, Numerator, Denominator);
 		const int32 MidiTempo = Harmonix::Midi::Constants::BPMToMidiTempo(97.0f);
 		MidiData->Tracks[0].AddEvent(FMidiEvent(0, FMidiMsg(MidiTempo)));
-		TempoMap.AddTempoInfoPoint(MidiTempo, 0);
+		MidiData->SongMaps.AddTempoInfoPoint(MidiTempo, 0);
 		MidiData->Tracks[0].Sort();
 		MidiData->ConformToLength(std::numeric_limits<int32>::max());
 
@@ -62,19 +59,19 @@ namespace HarmonixMidiTests::SongMaps
 			UTEST_TRUE(TEXT("Got good quantization."), QuantizedTick % SongMaps.SubdivisionToMidiTicks(ResultDivision, QuantizedTick) == 0);
 			if (GLogQualtizationDetails)
 			{
-				LogQuantizationDetails(SongMaps.GetBarMap(), i, QuantizedTick, ResultDivision);
+				LogQuantizationDetails(SongMaps, i, QuantizedTick, ResultDivision);
 			}
 			QuantizedTick = SongMaps.QuantizeTickToAnyNearestSubdivision(i, EMidiFileQuantizeDirection::Down, ResultDivision);
 			UTEST_TRUE(TEXT("Got good quantization."), QuantizedTick % SongMaps.SubdivisionToMidiTicks(ResultDivision, QuantizedTick) == 0);
 			if (GLogQualtizationDetails)
 			{
-				LogQuantizationDetails(SongMaps.GetBarMap(), i, QuantizedTick, ResultDivision);
+				LogQuantizationDetails(SongMaps, i, QuantizedTick, ResultDivision);
 			}
 			QuantizedTick = SongMaps.QuantizeTickToAnyNearestSubdivision(i, EMidiFileQuantizeDirection::Up, ResultDivision);
 			UTEST_TRUE(TEXT("Got good quantization."), QuantizedTick % SongMaps.SubdivisionToMidiTicks(ResultDivision, QuantizedTick) == 0);
 			if (GLogQualtizationDetails)
 			{
-				LogQuantizationDetails(SongMaps.GetBarMap(), i, QuantizedTick, ResultDivision);
+				LogQuantizationDetails(SongMaps, i, QuantizedTick, ResultDivision);
 			}
 		}
 		return true;
@@ -95,19 +92,19 @@ namespace HarmonixMidiTests::SongMaps
 			UTEST_TRUE(TEXT("Got good quantization."), QuantizedTick % SongMaps.SubdivisionToMidiTicks(ResultDivision, QuantizedTick) == 0);
 			if (GLogQualtizationDetails)
 			{
-				LogQuantizationDetails(SongMaps.GetBarMap(), i, QuantizedTick, ResultDivision);
+				LogQuantizationDetails(SongMaps, i, QuantizedTick, ResultDivision);
 			}
 			QuantizedTick = SongMaps.QuantizeTickToAnyNearestSubdivision(i, EMidiFileQuantizeDirection::Down, ResultDivision);
 			UTEST_TRUE(TEXT("Got good quantization."), QuantizedTick % SongMaps.SubdivisionToMidiTicks(ResultDivision, QuantizedTick) == 0);
 			if (GLogQualtizationDetails)
 			{
-				LogQuantizationDetails(SongMaps.GetBarMap(), i, QuantizedTick, ResultDivision);
+				LogQuantizationDetails(SongMaps, i, QuantizedTick, ResultDivision);
 			}
 			QuantizedTick = SongMaps.QuantizeTickToAnyNearestSubdivision(i, EMidiFileQuantizeDirection::Up, ResultDivision);
 			UTEST_TRUE(TEXT("Got good quantization."), QuantizedTick % SongMaps.SubdivisionToMidiTicks(ResultDivision, QuantizedTick) == 0);
 			if (GLogQualtizationDetails)
 			{
-				LogQuantizationDetails(SongMaps.GetBarMap(), i, QuantizedTick, ResultDivision);
+				LogQuantizationDetails(SongMaps, i, QuantizedTick, ResultDivision);
 			}
 		}
 		return true;
@@ -128,19 +125,19 @@ namespace HarmonixMidiTests::SongMaps
 			UTEST_TRUE(TEXT("Got good quantization."), QuantizedTick % SongMaps.SubdivisionToMidiTicks(ResultDivision, QuantizedTick) == 0);
 			if (GLogQualtizationDetails)
 			{
-				LogQuantizationDetails(SongMaps.GetBarMap(), i, QuantizedTick, ResultDivision);
+				LogQuantizationDetails(SongMaps, i, QuantizedTick, ResultDivision);
 			}
 			QuantizedTick = SongMaps.QuantizeTickToAnyNearestSubdivision(i, EMidiFileQuantizeDirection::Down, ResultDivision);
 			UTEST_TRUE(TEXT("Got good quantization."), QuantizedTick % SongMaps.SubdivisionToMidiTicks(ResultDivision, QuantizedTick) == 0);
 			if (GLogQualtizationDetails)
 			{
-				LogQuantizationDetails(SongMaps.GetBarMap(), i, QuantizedTick, ResultDivision);
+				LogQuantizationDetails(SongMaps, i, QuantizedTick, ResultDivision);
 			}
 			QuantizedTick = SongMaps.QuantizeTickToAnyNearestSubdivision(i, EMidiFileQuantizeDirection::Up, ResultDivision);
 			UTEST_TRUE(TEXT("Got good quantization."), QuantizedTick % SongMaps.SubdivisionToMidiTicks(ResultDivision, QuantizedTick) == 0);
 			if (GLogQualtizationDetails)
 			{
-				LogQuantizationDetails(SongMaps.GetBarMap(), i, QuantizedTick, ResultDivision);
+				LogQuantizationDetails(SongMaps, i, QuantizedTick, ResultDivision);
 			}
 		}
 		return true;
@@ -153,13 +150,12 @@ namespace HarmonixMidiTests::SongMaps
 		bool FTestSongMapsMixedTimeSig::RunTest(const FString&)
 	{
 		TSharedPtr<FMidiFileData> MidiData = BuildMidiWithOneTimeSigature(4, 4);
-		FBarMap& BarMap = MidiData->SongMaps.GetBarMap();
 		int32 TickAfterFirstBar = MidiData->TicksPerQuarterNote * 4;
 		MidiData->Tracks[0].AddEvent(FMidiEvent(TickAfterFirstBar, FMidiMsg(5, 8)));
-		BarMap.AddTimeSignatureAtBarIncludingCountIn(1, 5, 8);
+		MidiData->SongMaps.AddTimeSignatureAtBarIncludingCountIn(1, 5, 8);
 		int32 TickAfterSecondBar = MidiData->TicksPerQuarterNote / 2 * 5 + TickAfterFirstBar;
 		MidiData->Tracks[0].AddEvent(FMidiEvent(TickAfterSecondBar, FMidiMsg(3, 2)));
-		BarMap.AddTimeSignatureAtBarIncludingCountIn(2, 3, 2);
+		MidiData->SongMaps.AddTimeSignatureAtBarIncludingCountIn(2, 3, 2);
 
 		MidiData->Tracks[0].Sort();
 
@@ -173,19 +169,19 @@ namespace HarmonixMidiTests::SongMaps
 			UTEST_TRUE(TEXT("Got good quantization."), QuantizedTick % SongMaps.SubdivisionToMidiTicks(ResultDivision, i) == 0);
 			if (GLogQualtizationDetails)
 			{
-				LogQuantizationDetails(SongMaps.GetBarMap(), i, QuantizedTick, ResultDivision);
+				LogQuantizationDetails(SongMaps, i, QuantizedTick, ResultDivision);
 			}
 			QuantizedTick = SongMaps.QuantizeTickToAnyNearestSubdivision(i, EMidiFileQuantizeDirection::Down, ResultDivision);
 			UTEST_TRUE(TEXT("Got good quantization."), QuantizedTick % SongMaps.SubdivisionToMidiTicks(ResultDivision, i) == 0);
 			if (GLogQualtizationDetails)
 			{
-				LogQuantizationDetails(SongMaps.GetBarMap(), i, QuantizedTick, ResultDivision);
+				LogQuantizationDetails(SongMaps, i, QuantizedTick, ResultDivision);
 			}
 			QuantizedTick = SongMaps.QuantizeTickToAnyNearestSubdivision(i, EMidiFileQuantizeDirection::Up, ResultDivision);
 			UTEST_TRUE(TEXT("Got good quantization."), QuantizedTick % SongMaps.SubdivisionToMidiTicks(ResultDivision, i) == 0);
 			if (GLogQualtizationDetails)
 			{
-				LogQuantizationDetails(SongMaps.GetBarMap(), i, QuantizedTick, ResultDivision);
+				LogQuantizationDetails(SongMaps, i, QuantizedTick, ResultDivision);
 			}
 		}
 		// Check Second Bar...
@@ -196,19 +192,19 @@ namespace HarmonixMidiTests::SongMaps
 			UTEST_TRUE(TEXT("Got good quantization."), (QuantizedTick - TickAfterFirstBar) % SongMaps.SubdivisionToMidiTicks(ResultDivision, i) == 0);
 			if (GLogQualtizationDetails)
 			{
-				LogQuantizationDetails(SongMaps.GetBarMap(), i, QuantizedTick, ResultDivision);
+				LogQuantizationDetails(SongMaps, i, QuantizedTick, ResultDivision);
 			}
 			QuantizedTick = SongMaps.QuantizeTickToAnyNearestSubdivision(i, EMidiFileQuantizeDirection::Down, ResultDivision);
 			UTEST_TRUE(TEXT("Got good quantization."), (QuantizedTick - TickAfterFirstBar) % SongMaps.SubdivisionToMidiTicks(ResultDivision, i) == 0);
 			if (GLogQualtizationDetails)
 			{
-				LogQuantizationDetails(SongMaps.GetBarMap(), i, QuantizedTick, ResultDivision);
+				LogQuantizationDetails(SongMaps, i, QuantizedTick, ResultDivision);
 			}
 			QuantizedTick = SongMaps.QuantizeTickToAnyNearestSubdivision(i, EMidiFileQuantizeDirection::Up, ResultDivision);
 			UTEST_TRUE(TEXT("Got good quantization."), (QuantizedTick - TickAfterFirstBar) % SongMaps.SubdivisionToMidiTicks(ResultDivision, i) == 0);
 			if (GLogQualtizationDetails)
 			{
-				LogQuantizationDetails(SongMaps.GetBarMap(), i, QuantizedTick, ResultDivision);
+				LogQuantizationDetails(SongMaps, i, QuantizedTick, ResultDivision);
 			}
 		}
 		// Check Third Bar...
@@ -219,19 +215,19 @@ namespace HarmonixMidiTests::SongMaps
 			UTEST_TRUE(TEXT("Got good quantization."), (QuantizedTick-TickAfterSecondBar) % SongMaps.SubdivisionToMidiTicks(ResultDivision, i) == 0);
 			if (GLogQualtizationDetails)
 			{
-				LogQuantizationDetails(SongMaps.GetBarMap(), i, QuantizedTick, ResultDivision);
+				LogQuantizationDetails(SongMaps, i, QuantizedTick, ResultDivision);
 			}
 			QuantizedTick = SongMaps.QuantizeTickToAnyNearestSubdivision(i, EMidiFileQuantizeDirection::Down, ResultDivision);
 			UTEST_TRUE(TEXT("Got good quantization."), (QuantizedTick - TickAfterSecondBar) % SongMaps.SubdivisionToMidiTicks(ResultDivision, i) == 0);
 			if (GLogQualtizationDetails)
 			{
-				LogQuantizationDetails(SongMaps.GetBarMap(), i, QuantizedTick, ResultDivision);
+				LogQuantizationDetails(SongMaps, i, QuantizedTick, ResultDivision);
 			}
 			QuantizedTick = SongMaps.QuantizeTickToAnyNearestSubdivision(i, EMidiFileQuantizeDirection::Up, ResultDivision);
 			UTEST_TRUE(TEXT("Got good quantization."), (QuantizedTick - TickAfterSecondBar) % SongMaps.SubdivisionToMidiTicks(ResultDivision, i) == 0);
 			if (GLogQualtizationDetails)
 			{
-				LogQuantizationDetails(SongMaps.GetBarMap(), i, QuantizedTick, ResultDivision);
+				LogQuantizationDetails(SongMaps, i, QuantizedTick, ResultDivision);
 			}
 		}
 
diff --git a/Engine/Plugins/Runtime/Metasound/Source/MetasoundEngine/Private/MetasoundGeneratorHandle.cpp b/Engine/Plugins/Runtime/Metasound/Source/MetasoundEngine/Private/MetasoundGeneratorHandle.cpp
--- a/Engine/Plugins/Runtime/Metasound/Source/MetasoundEngine/Private/MetasoundGeneratorHandle.cpp
+++ b/Engine/Plugins/Runtime/Metasound/Source/MetasoundEngine/Private/MetasoundGeneratorHandle.cpp
@@ -774,6 +774,19 @@ bool UMetasoundGeneratorHandle::RemoveGraphSetCallback(const FDelegateHandle& Ha
 	return OnGeneratorsGraphChanged.Remove(Handle);
 }
 
+bool UMetasoundGeneratorHandle::TryCreateAnalyzerAddress(const FName OutputName, const FName AnalyzerName, const FName AnalyzerOutputName, Metasound::Frontend::FAnalyzerAddress& OutAnalyzerAddress)
+{
+	METASOUND_LLM_SCOPE;
+	METASOUND_TRACE_CPUPROFILER_EVENT_SCOPE(UMetasoundGeneratorHandle::WatchOutput);
+
+	if (!IsValid())
+	{
+		return false;
+	}
+
+	return GeneratorHandle->TryCreateAnalyzerAddress(OutputName, AnalyzerName, AnalyzerOutputName, OutAnalyzerAddress);
+}
+
 bool UMetasoundGeneratorHandle::WatchOutput(
 	const FName OutputName,
 	const FOnMetasoundOutputValueChanged& OnOutputValueChanged,
diff --git a/Engine/Plugins/Runtime/Metasound/Source/MetasoundEngine/Public/MetasoundGeneratorHandle.h b/Engine/Plugins/Runtime/Metasound/Source/MetasoundEngine/Public/MetasoundGeneratorHandle.h
--- a/Engine/Plugins/Runtime/Metasound/Source/MetasoundEngine/Public/MetasoundGeneratorHandle.h
+++ b/Engine/Plugins/Runtime/Metasound/Source/MetasoundEngine/Public/MetasoundGeneratorHandle.h
@@ -158,6 +158,12 @@ namespace Metasound
 			FName AnalyzerName = NAME_None,
 			FName AnalyzerOutputName = NAME_None);
 
+		bool TryCreateAnalyzerAddress(
+			const FName OutputName,
+			const FName AnalyzerName,
+			const FName AnalyzerOutputName,
+			Frontend::FAnalyzerAddress& OutAnalyzerAddress);
+
 		/**
 		 * Update any watched outputs
 		 */
@@ -223,13 +229,6 @@ namespace Metasound
 			FName AnalyzerOutputName = NAME_None);
 
 		void UpdateOutputWatchersInternal();
-
-		bool TryCreateAnalyzerAddress(
-			const FName OutputName,
-			const FName AnalyzerName,
-			const FName AnalyzerOutputName,
-			Frontend::FAnalyzerAddress& OutAnalyzerAddress);
-
 		void FixUpOutputWatchers();
 
 		void CreateOutputWatcher(
@@ -470,6 +469,12 @@ public:
 	DECLARE_MULTICAST_DELEGATE_OneParam(FOnIOUpdatedWithChanges, const TArray<Metasound::FVertexInterfaceChange>&)
 	FOnIOUpdatedWithChanges OnIOUpdatedWithChanges;
 
+	bool TryCreateAnalyzerAddress(
+		const FName OutputName,
+		const FName AnalyzerName,
+		const FName AnalyzerOutputName,
+		Metasound::Frontend::FAnalyzerAddress& OutAnalyzerAddress);
+
 	/**
 	 * Watch an output value.
 	 *
