Minor change to fix static analysis complaint

diff --git a/Engine/Plugins/Runtime/Harmonix/Source/HarmonixMetasound/Public/HarmonixMetasound/Analysis/SpmcAnalysisResultQueue.h b/Engine/Plugins/Runtime/Harmonix/Source/HarmonixMetasound/Public/HarmonixMetasound/Analysis/SpmcAnalysisResultQueue.h
--- a/Engine/Plugins/Runtime/Harmonix/Source/HarmonixMetasound/Public/HarmonixMetasound/Analysis/SpmcAnalysisResultQueue.h
+++ b/Engine/Plugins/Runtime/Harmonix/Source/HarmonixMetasound/Public/HarmonixMetasound/Analysis/SpmcAnalysisResultQueue.h
@@ -234,6 +234,7 @@ namespace HarmonixMetasound::Analysis
 				Entry = MoveTemp(Other.Entry);
 				DiscontinuityDetectedInLastRead = Other.DiscontinuityDetectedInLastRead;
 				ReadCursor = MoveTemp(Other.ReadCursor);
+				return *this;
 			}
 
 			void Release()
