Fix issue with StepSequencer when watching a looping clock input. This broke with the new clock.

Also some updated ground truth files for a few of the end-to-end tests.
#tests Harmonix Functional Tests

diff --git a/Engine/Plugins/Runtime/Harmonix/Source/HarmonixMetasound/Private/HarmonixMetasound/Nodes/StepSequencePlayerNode.cpp b/Engine/Plugins/Runtime/Harmonix/Source/HarmonixMetasound/Private/HarmonixMetasound/Nodes/StepSequencePlayerNode.cpp
--- a/Engine/Plugins/Runtime/Harmonix/Source/HarmonixMetasound/Private/HarmonixMetasound/Nodes/StepSequencePlayerNode.cpp
+++ b/Engine/Plugins/Runtime/Harmonix/Source/HarmonixMetasound/Private/HarmonixMetasound/Nodes/StepSequencePlayerNode.cpp
@@ -377,16 +377,14 @@ namespace HarmonixMetasound::Nodes::StepSequencePlayer
 							{
 								AdvanceThruTick(Event.BlockFrameIndex, AsAdvance->LastTickToProcess());
 							}
-/*
-							else if (Event.Msg.IsType<MidiClockMessageTypes::FSeekThru>())
-							{
-								SeekThruTick(Event.BlockFrameIndex, Event.Msg.Get<MidiClockMessageTypes::FSeekThru>().ThruTick);
-							}
-*/
 							else if (const FSeek* AsSeekTo = Event.TryGet<FSeek>())
 							{
 								SeekToTick(Event.BlockFrameIndex, AsSeekTo->NewNextTick);
 							}
+							else if (const FLoop* AsLoop = Event.TryGet<FLoop>())
+							{
+								SeekToTick(Event.BlockFrameIndex, AsLoop->FirstTickInLoop);
+							}
 						}
 					}
 				}
@@ -822,7 +820,7 @@ namespace HarmonixMetasound::Nodes::StepSequencePlayer
 				}
 			}
 
-			if (CellInRow != CurrentCellIndex || CurrentMaxColumns == 1)
+			if (CellInRow != CurrentCellIndex/* || CurrentMaxColumns == 1*/)
 			{	
 				for (int32 i = 0; i < CurrentPage->Rows.Num(); ++i)
 				{
