[Harmonix] Transport Event Stream Vertex Analyzer

diff --git a/Engine/Plugins/Runtime/Harmonix/Source/HarmonixMetasound/Private/HarmonixMetasound/Analysis/MusicTransportEventStreamVertexAnalyzer.cpp b/Engine/Plugins/Runtime/Harmonix/Source/HarmonixMetasound/Private/HarmonixMetasound/Analysis/MusicTransportEventStreamVertexAnalyzer.cpp
new file mode 100644
--- /dev/null
+++ b/Engine/Plugins/Runtime/Harmonix/Source/HarmonixMetasound/Private/HarmonixMetasound/Analysis/MusicTransportEventStreamVertexAnalyzer.cpp
@@ -0,0 +1,59 @@
+// Copyright Epic Games, Inc. All Rights Reserved.
+
+#include "HarmonixMetasound/Analysis/MusicTransportEventStreamVertexAnalyzer.h"
+
+#include "HarmonixMetasound/DataTypes/MusicTransport.h"
+
+namespace HarmonixMetasound::Analysis
+{
+	const FName& FMusicTransportEventStreamVertexAnalyzer::GetAnalyzerName()
+	{
+		return AnalyzerName;
+	}
+
+	const FName& FMusicTransportEventStreamVertexAnalyzer::GetDataType()
+	{
+		return Metasound::GetMetasoundDataTypeName<FMusicTransportEventStream>();
+	}
+	
+	const Metasound::Frontend::FAnalyzerOutput& FMusicTransportEventStreamVertexAnalyzer::FOutputs::GetValue()
+	{
+		return TransportEvent;
+	}
+
+	const Metasound::Frontend::FAnalyzerOutput FMusicTransportEventStreamVertexAnalyzer::FOutputs::SeekDestination = { "SeekDestination", Metasound::GetMetasoundDataTypeName<FMusicSeekTarget>() };
+	const Metasound::Frontend::FAnalyzerOutput FMusicTransportEventStreamVertexAnalyzer::FOutputs::TransportEvent = { "TransportEvent", Metasound::GetMetasoundDataTypeName<FMusicTransportEvent>() };
+
+	const TArray<Metasound::Frontend::FAnalyzerOutput>& FMusicTransportEventStreamVertexAnalyzer::FFactory::GetAnalyzerOutputs() const
+	{
+		return AnalyzerOutputs;
+	}
+
+	const TArray<Metasound::Frontend::FAnalyzerOutput> FMusicTransportEventStreamVertexAnalyzer::FFactory::AnalyzerOutputs{ FOutputs::SeekDestination, FOutputs::TransportEvent };
+
+	FMusicTransportEventStreamVertexAnalyzer::FMusicTransportEventStreamVertexAnalyzer(const Metasound::Frontend::FCreateAnalyzerParams& InParams)
+		: FVertexAnalyzerBase(InParams.AnalyzerAddress, InParams.VertexDataReference)
+		, SeekDestination(FMusicSeekTargetWriteRef::CreateNew())
+		, LastMusicTransportEvent(FMusicTransportEventWriteRef::CreateNew())
+		, FramesPerBlock(InParams.OperatorSettings.GetNumFramesPerBlock())
+		, SampleRate(InParams.OperatorSettings.GetSampleRate())
+	{
+		BindOutputData(FOutputs::SeekDestination.Name, InParams.OperatorSettings, Metasound::TDataReadReference{ SeekDestination });
+		BindOutputData(FOutputs::TransportEvent.Name, InParams.OperatorSettings, Metasound::TDataReadReference{ LastMusicTransportEvent });
+	}
+
+	void FMusicTransportEventStreamVertexAnalyzer::Execute()
+	{
+		const FMusicTransportEventStream& MusicTransportEventStream = GetVertexData<FMusicTransportEventStream>();
+		for (auto&& Event : MusicTransportEventStream.GetTransportEventsInBlock())
+		{
+			*SeekDestination = MusicTransportEventStream.GetNextSeekDestination();
+			LastMusicTransportEvent->Time = Metasound::FTime::FromSeconds(double(Event.SampleIndex + NumFrames) / SampleRate);
+			LastMusicTransportEvent->Request = Event.Request;
+			MarkOutputDirty();
+		}
+		NumFrames += FramesPerBlock;
+	}
+
+	const FName FMusicTransportEventStreamVertexAnalyzer::AnalyzerName = "Harmonix.MusicTransport";
+}
diff --git a/Engine/Plugins/Runtime/Harmonix/Source/HarmonixMetasound/Private/HarmonixMetasound/DataTypes/MusicTransport.cpp b/Engine/Plugins/Runtime/Harmonix/Source/HarmonixMetasound/Private/HarmonixMetasound/DataTypes/MusicTransport.cpp
--- a/Engine/Plugins/Runtime/Harmonix/Source/HarmonixMetasound/Private/HarmonixMetasound/DataTypes/MusicTransport.cpp
+++ b/Engine/Plugins/Runtime/Harmonix/Source/HarmonixMetasound/Private/HarmonixMetasound/DataTypes/MusicTransport.cpp
@@ -4,6 +4,7 @@
 
 #include "MetasoundDataTypeRegistrationMacro.h"
 
+REGISTER_METASOUND_DATATYPE(HarmonixMetasound::FMusicTransportEvent, "MusicTransportEvent")
 REGISTER_METASOUND_DATATYPE(HarmonixMetasound::FMusicTransportEventStream, "MusicTransport")
 
 namespace HarmonixMetasound
diff --git a/Engine/Plugins/Runtime/Harmonix/Source/HarmonixMetasound/Private/HarmonixMetasoundModule.cpp b/Engine/Plugins/Runtime/Harmonix/Source/HarmonixMetasound/Private/HarmonixMetasoundModule.cpp
--- a/Engine/Plugins/Runtime/Harmonix/Source/HarmonixMetasound/Private/HarmonixMetasoundModule.cpp
+++ b/Engine/Plugins/Runtime/Harmonix/Source/HarmonixMetasound/Private/HarmonixMetasoundModule.cpp
@@ -9,6 +9,7 @@
 #include "HarmonixMetasound/Analysis/MidiClockVertexAnalyzer.h"
 #include "HarmonixMetasound/Analysis/MidiSongPosVertexAnalyzer.h"
 #include "HarmonixMetasound/Analysis/MidiStreamVertexAnalyzer.h"
+#include "HarmonixMetasound/Analysis/MusicTransportEventStreamVertexAnalyzer.h"
 #include "HarmonixMetasound/Analysis/FFTAnalyzerResultVertexAnalyzer.h"
 #include "HarmonixMetasound/DataTypes/FFTAnalyzerResult.h"
 #include "HarmonixMetasound/DataTypes/MidiStream.h"
@@ -33,6 +34,10 @@ void FHarmonixMetasoundModule::StartupModule()
 		GetMetasoundDataTypeName<HarmonixMetasound::FMidiClock>(),
 		HarmonixMetasound::Analysis::FMidiClockVertexAnalyzer::GetAnalyzerName(),
 		HarmonixMetasound::Analysis::FMidiClockVertexAnalyzer::FOutputs::GetValue().Name);
+	UMetasoundGeneratorHandle::RegisterPassthroughAnalyzerForType(
+		GetMetasoundDataTypeName<HarmonixMetasound::FMusicTransportEventStream>(),
+		HarmonixMetasound::Analysis::FMusicTransportEventStreamVertexAnalyzer::GetAnalyzerName(),
+		HarmonixMetasound::Analysis::FMusicTransportEventStreamVertexAnalyzer::FOutputs::GetValue().Name);
 	UMetasoundGeneratorHandle::RegisterPassthroughAnalyzerForType(
 		GetMetasoundDataTypeName<FHarmonixFFTAnalyzerResults>(),
 		HarmonixMetasound::Analysis::FFFTAnalyzerResultVertexAnalyzer::GetAnalyzerName(),
@@ -42,6 +47,7 @@ void FHarmonixMetasoundModule::StartupModule()
 	METASOUND_REGISTER_VERTEX_ANALYZER_FACTORY(HarmonixMetasound::Analysis::FMidiStreamVertexAnalyzer)
 	METASOUND_REGISTER_VERTEX_ANALYZER_FACTORY(HarmonixMetasound::Analysis::FMidiClockVertexAnalyzer)
 	METASOUND_REGISTER_VERTEX_ANALYZER_FACTORY(HarmonixMetasound::Analysis::FMidiSongPosVertexAnalyzer)
+	METASOUND_REGISTER_VERTEX_ANALYZER_FACTORY(HarmonixMetasound::Analysis::FMusicTransportEventStreamVertexAnalyzer)
 	METASOUND_REGISTER_VERTEX_ANALYZER_FACTORY(HarmonixMetasound::Analysis::FFFTAnalyzerResultVertexAnalyzer)
 
 	// The first redirect for thie module
diff --git a/Engine/Plugins/Runtime/Harmonix/Source/HarmonixMetasound/Public/HarmonixMetasound/Analysis/MusicTransportEventStreamVertexAnalyzer.h b/Engine/Plugins/Runtime/Harmonix/Source/HarmonixMetasound/Public/HarmonixMetasound/Analysis/MusicTransportEventStreamVertexAnalyzer.h
new file mode 100644
--- /dev/null
+++ b/Engine/Plugins/Runtime/Harmonix/Source/HarmonixMetasound/Public/HarmonixMetasound/Analysis/MusicTransportEventStreamVertexAnalyzer.h
@@ -0,0 +1,53 @@
+// Copyright Epic Games, Inc. All Rights Reserved.
+#pragma once
+
+#include "MetasoundDataReference.h"
+#include "Analysis/MetasoundFrontendAnalyzerFactory.h"
+#include "Analysis/MetasoundFrontendVertexAnalyzer.h"
+
+#include "HarmonixMetasound/DataTypes/MusicTransport.h"
+
+namespace HarmonixMetasound::Analysis
+{
+	class HARMONIXMETASOUND_API FMusicTransportEventStreamVertexAnalyzer final : public Metasound::Frontend::FVertexAnalyzerBase
+	{
+	public:
+		static const FName& GetAnalyzerName();
+		static const FName& GetDataType();
+
+		struct HARMONIXMETASOUND_API FOutputs
+		{
+			/**
+			 * @brief Get the default output for this analyzer
+			 * @return The default output
+			 */
+			static const Metasound::Frontend::FAnalyzerOutput& GetValue();
+
+			static const Metasound::Frontend::FAnalyzerOutput SeekDestination;
+			static const Metasound::Frontend::FAnalyzerOutput TransportEvent;
+		};
+
+		class HARMONIXMETASOUND_API FFactory final : public Metasound::Frontend::TVertexAnalyzerFactory<FMusicTransportEventStreamVertexAnalyzer>
+		{
+		public:
+			virtual const TArray<Metasound::Frontend::FAnalyzerOutput>& GetAnalyzerOutputs() const override;
+
+		private:
+			static const TArray<Metasound::Frontend::FAnalyzerOutput> AnalyzerOutputs;
+		};
+
+		explicit FMusicTransportEventStreamVertexAnalyzer(const Metasound::Frontend::FCreateAnalyzerParams& InParams);
+		virtual ~FMusicTransportEventStreamVertexAnalyzer() override = default;
+
+		virtual void Execute() override;
+
+	private:
+		FMusicSeekTargetWriteRef  SeekDestination;
+		FMusicTransportEventWriteRef LastMusicTransportEvent;
+		int64 NumFrames{ 0 };
+		int32 FramesPerBlock{ 0 };
+		float SampleRate{ 0 };
+
+		static const FName AnalyzerName;
+	};
+}
diff --git a/Engine/Plugins/Runtime/Harmonix/Source/HarmonixMetasound/Public/HarmonixMetasound/DataTypes/MusicTransport.h b/Engine/Plugins/Runtime/Harmonix/Source/HarmonixMetasound/Public/HarmonixMetasound/DataTypes/MusicTransport.h
--- a/Engine/Plugins/Runtime/Harmonix/Source/HarmonixMetasound/Public/HarmonixMetasound/DataTypes/MusicTransport.h
+++ b/Engine/Plugins/Runtime/Harmonix/Source/HarmonixMetasound/Public/HarmonixMetasound/DataTypes/MusicTransport.h
@@ -24,6 +24,16 @@ namespace HarmonixMetasound
 		Count
 	};
 
+	struct HARMONIXMETASOUND_API FMusicTransportEvent
+	{
+		Metasound::FTime Time;
+		EMusicPlayerTransportRequest Request;
+	};
+
+	// Declare aliases IN the namespace...
+	DECLARE_METASOUND_DATA_REFERENCE_ALIAS_TYPES(FMusicTransportEvent, FMusicTransportEventTypeInfo, FMusicTransportEventReadRef, FMusicTransportEventWriteRef)
+
+
 	class HARMONIXMETASOUND_API FMusicTransportEventStream
 	{
 	public:
@@ -115,5 +125,8 @@ namespace HarmonixMetasound
 	};
 }
 
+// Declare reference types OUT of the namespace...
+DECLARE_METASOUND_DATA_REFERENCE_TYPES_NO_ALIASES(HarmonixMetasound::FMusicTransportEvent, HARMONIXMETASOUND_API)
+
 // Declare reference types OUT of the namespace...
 DECLARE_METASOUND_DATA_REFERENCE_TYPES_NO_ALIASES(HarmonixMetasound::FMusicTransportEventStream, HARMONIXMETASOUND_API)
