Fix inadvertent regression of a previous fix. Minor change BACK to the correct code.

diff --git a/Engine/Plugins/Runtime/Harmonix/Source/HarmonixMetasound/Private/HarmonixMetasound/Nodes/StepSequencePlayerNode.cpp b/Engine/Plugins/Runtime/Harmonix/Source/HarmonixMetasound/Private/HarmonixMetasound/Nodes/StepSequencePlayerNode.cpp
--- a/Engine/Plugins/Runtime/Harmonix/Source/HarmonixMetasound/Private/HarmonixMetasound/Nodes/StepSequencePlayerNode.cpp
+++ b/Engine/Plugins/Runtime/Harmonix/Source/HarmonixMetasound/Private/HarmonixMetasound/Nodes/StepSequencePlayerNode.cpp
@@ -820,7 +820,7 @@ namespace HarmonixMetasound::Nodes::StepSequencePlayer
 				}
 			}
 
-			if (CellInRow != CurrentCellIndex/* || CurrentMaxColumns == 1*/)
+			if (CellInRow != CurrentCellIndex || CurrentMaxColumns == 1)
 			{	
 				for (int32 i = 0; i < CurrentPage->Rows.Num(); ++i)
 				{
