Trivial change to add a little info to the test log.

diff --git a/Engine/Plugins/Runtime/Harmonix/Source/HarmonixMetasoundTests/Private/HarmonixMetasoundFunctionalTest.cpp b/Engine/Plugins/Runtime/Harmonix/Source/HarmonixMetasoundTests/Private/HarmonixMetasoundFunctionalTest.cpp
--- a/Engine/Plugins/Runtime/Harmonix/Source/HarmonixMetasoundTests/Private/HarmonixMetasoundFunctionalTest.cpp
+++ b/Engine/Plugins/Runtime/Harmonix/Source/HarmonixMetasoundTests/Private/HarmonixMetasoundFunctionalTest.cpp
@@ -1,4 +1,4 @@
-﻿// Copyright Epic Games, Inc. All Rights Reserved.
+// Copyright Epic Games, Inc. All Rights Reserved.
 
 #include "HarmonixMetasoundFunctionalTest.h"
 #include "HarmonixFunctionalTestAction.h"
@@ -267,7 +267,7 @@ void AHarmonixMetasoundFunctionalTest::CompareResults()
 	float PSNR = Harmonix::Dsp::AudioAnalysis::CalculatePSNR(AudioCaptureOutput.GetData(), AudioData.GetData(), NumChannels, NumFramesToCompare);
 
 	static constexpr float PSNRThreshold = 60.0f;
-	AssertTrue(PSNR >= PSNRThreshold, FString::Printf(TEXT("PSNR = %.2f where the acceptable range is (PSNR >= %.2f)"), PSNR, PSNRThreshold), this);
+	AssertTrue(PSNR >= PSNRThreshold, FString::Printf(TEXT("PSNR = %.2f where the acceptable range is (PSNR >= %.2f) Compared %d frames."), PSNR, PSNRThreshold, NumFramesToCompare), this);
 }
 
 void AHarmonixMetasoundFunctionalTest::StartTest()
