Remove requirement to create an FMidiFileData in places where only an FSongMaps is needed. Make FMidiFileData an ISongMapEvaluator as a facade in front of its SongMaps. Cleans up a bunch of code.

diff --git a/Engine/Plugins/Runtime/Harmonix/Source/HarmonixMetasound/Private/HarmonixMetasound/Analysis/MidiSongPosVertexAnalyzer.cpp b/Engine/Plugins/Runtime/Harmonix/Source/HarmonixMetasound/Private/HarmonixMetasound/Analysis/MidiSongPosVertexAnalyzer.cpp
--- a/Engine/Plugins/Runtime/Harmonix/Source/HarmonixMetasound/Private/HarmonixMetasound/Analysis/MidiSongPosVertexAnalyzer.cpp
+++ b/Engine/Plugins/Runtime/Harmonix/Source/HarmonixMetasound/Private/HarmonixMetasound/Analysis/MidiSongPosVertexAnalyzer.cpp
@@ -58,7 +58,7 @@ namespace HarmonixMetasound::Analysis
 
 		TSharedPtr<const FSongMapChain> MapChain = History->GetLatestMapsForProducer();
 
-		if (!MapChain->SongMaps || Clock.GetMidiDataChangedInBlock() ||
+		if (!MapChain->SongMaps || Clock.GetSongMapsChangedInBlock() ||
 			Clock.GetFirstTickInLoop() != MapChain->FirstTickInLoop ||
 			Clock.GetLoopLengthTicks() != MapChain->LoopLengthTicks)
 		{
diff --git a/Engine/Plugins/Runtime/Harmonix/Source/HarmonixMetasound/Private/HarmonixMetasound/DataTypes/MidiClock.cpp b/Engine/Plugins/Runtime/Harmonix/Source/HarmonixMetasound/Private/HarmonixMetasound/DataTypes/MidiClock.cpp
--- a/Engine/Plugins/Runtime/Harmonix/Source/HarmonixMetasound/Private/HarmonixMetasound/DataTypes/MidiClock.cpp
+++ b/Engine/Plugins/Runtime/Harmonix/Source/HarmonixMetasound/Private/HarmonixMetasound/DataTypes/MidiClock.cpp
@@ -40,7 +40,7 @@ namespace HarmonixMetasound
 	using namespace Metasound;
 
 	FMidiClock::FMidiClock(const FOperatorSettings& InSettings)
-		: SongMapEvaluator(MakeShared<FSongMapsWithAlternateTempoSource>(MakeClockConductorMidiData(120.0f, 4, 4)))
+		: SongMapEvaluator(MakeShared<FSongMapsWithAlternateTempoSource>(MakeShared<FSongMaps>(120.0f, 4, 4)))
 		, CurrentTempoInfoPointIndex(0)
 		, CurrentTimeSignaturePointIndex(0)
 		, TickResidualWhenDriven(0.0f)
@@ -154,23 +154,23 @@ namespace HarmonixMetasound
 		return *this;
 	}
 
-	void FMidiClock::AttachToMidiFile(TSharedPtr<FMidiFileData> MidiData, bool ResetToStart)
+	void FMidiClock::AttachToSongMapEvaluator(TSharedPtr<ISongMapEvaluator> SongMaps, bool ResetToStart)
 	{
 		MidiDataChangedInBlock = true;
-		if (!MidiData)
+		if (!SongMaps)
 		{
-			MidiData = MakeClockConductorMidiData(TempoAtBlockEnd, TimeSignatureAtBlockEnd.Numerator, TimeSignatureAtBlockEnd.Denominator);
+			SongMaps = MakeShared<FSongMaps>(TempoAtBlockEnd, TimeSignatureAtBlockEnd.Numerator, TimeSignatureAtBlockEnd.Denominator);
 		}
 
 		if (ExternalClockDriver)
 		{
 			// use the tempo from the external clock...
-			RebuildSongMapEvaluator(ExternalClockDriver->SongMapEvaluator->GetMidiFileWithTempoMap(), MidiData);
+			RebuildSongMapEvaluator(ExternalClockDriver->SongMapEvaluator->GetSongMapsWithTempoMap(), SongMaps);
 		}
 		else
 		{
 			// All maps come from the same source...
-			RebuildSongMapEvaluator(MidiData, MidiData);
+			RebuildSongMapEvaluator(SongMaps, SongMaps);
 		}
 
 		if (ResetToStart)
@@ -180,15 +180,15 @@ namespace HarmonixMetasound
 		PostTempoOrTimeSignatureEventsIfNeeded();
 	}
 
-	void FMidiClock::MidiChanged()
+	void FMidiClock::SongMapsChanged()
 	{
 		MidiDataChangedInBlock = true;
 		PostTempoOrTimeSignatureEventsIfNeeded();
 	}
 
-	void FMidiClock::DetachFromMidiFile()
+	void FMidiClock::DetachFromSongMaps()
 	{
-		AttachToMidiFile(nullptr, false);
+		AttachToSongMapEvaluator(nullptr, false);
 	}
 
 	void FMidiClock::SetDrivingClock(FConstSharedMidiClockPtr NewExternalClockDriver)
@@ -196,9 +196,9 @@ namespace HarmonixMetasound
 		MidiDataChangedInBlock = true;
 		ExternalClockDriver = NewExternalClockDriver;
 		RebuildSongMapEvaluator(ExternalClockDriver ? 
-									ExternalClockDriver->SongMapEvaluator->GetMidiFileWithTempoMap() : 
-									SongMapEvaluator->GetMidiFileWithOtherMaps(),
-								SongMapEvaluator->GetMidiFileWithOtherMaps());
+									ExternalClockDriver->SongMapEvaluator->GetSongMapsWithTempoMap() : 
+									SongMapEvaluator->GetSongMapsWithOtherMaps(),
+								SongMapEvaluator->GetSongMapsWithOtherMaps());
 	}
 
 	void FMidiClock::PrepareBlock()
@@ -223,7 +223,7 @@ namespace HarmonixMetasound
 
 		if (ExternalClockDriver && ExternalClockDriver->MidiDataChangedInBlock)
 		{
-			RebuildSongMapEvaluator(ExternalClockDriver->SongMapEvaluator->GetMidiFileWithTempoMap(), SongMapEvaluator->GetMidiFileWithOtherMaps());
+			RebuildSongMapEvaluator(ExternalClockDriver->SongMapEvaluator->GetSongMapsWithTempoMap(), SongMapEvaluator->GetSongMapsWithOtherMaps());
 			PostTempoOrTimeSignatureEventsIfNeeded();
 		}
 	}
@@ -568,39 +568,6 @@ namespace HarmonixMetasound
 		return Tick;
 	}
 
-	TSharedPtr<FMidiFileData> FMidiClock::MakeClockConductorMidiData(float InTempoBPM, int32 InTimeSigNum, int32 InTimeSigDen)
-	{
-		TSharedPtr<FMidiFileData> OutMidiData = MakeShared<FMidiFileData>();
-
-		// clear it all out for good measure...
-		FTempoMap& TempoMap = OutMidiData->SongMaps.GetTempoMap();
-		TempoMap.Empty();
-		FBarMap& BarMap = OutMidiData->SongMaps.GetBarMap();
-		BarMap.Empty();
-		OutMidiData->Tracks.Empty();
-
-		// create conductor track
-		FMidiTrack& Track = OutMidiData->Tracks.Add_GetRef(FMidiTrack(TEXT("conductor")));
-		// max out song length data so the midi can play indefinitely
-		OutMidiData->ConformToLength(std::numeric_limits<int32>::max());
-
-		// add time sig info
-		int32 TimeSigNum = FMath::Clamp(InTimeSigNum, 1, 64);
-		int32 TimeSigDen = FMath::Clamp(InTimeSigDen, 1, 64);
-		Track.AddEvent(FMidiEvent(0, FMidiMsg((uint8)TimeSigNum, (uint8)TimeSigDen)));
-		BarMap.AddTimeSignatureAtBarIncludingCountIn(0, TimeSigNum, TimeSigDen);
-
-		// add tempo info
-		float TempoBPM = FMath::Max(1.0f, InTempoBPM);
-		int32 MidiTempo = Harmonix::Midi::Constants::BPMToMidiTempo(TempoBPM);
-		Track.AddEvent(FMidiEvent(0, FMidiMsg(MidiTempo)));
-		TempoMap.AddTempoInfoPoint(MidiTempo, 0);
-
-		Track.Sort();
-
-		return OutMidiData;
-	}
-
 	int32 FMidiClock::GetNextTickToProcessAtBlockFrame(int32 BlockFrame) const
 	{
 		int32 FoundNextTick = NextMidiTickToProcess;
@@ -1012,9 +979,9 @@ namespace HarmonixMetasound
 		}
 	}
 
-	void FMidiClock::RebuildSongMapEvaluator(const TSharedPtr<const FMidiFileData>& MidiWithTempo, const TSharedPtr<const FMidiFileData>& MidiWithOtherMaps)
+	void FMidiClock::RebuildSongMapEvaluator(const TSharedPtr<const ISongMapEvaluator>& MapWithTempo, const TSharedPtr<const ISongMapEvaluator>& MapWithOtherMaps)
 	{
-		SongMapEvaluator = MakeShared<FSongMapsWithAlternateTempoSource>(MidiWithTempo, MidiWithOtherMaps);
+		SongMapEvaluator = MakeShared<FSongMapsWithAlternateTempoSource>(MapWithTempo, MapWithOtherMaps);
 	}
 
 }
diff --git a/Engine/Plugins/Runtime/Harmonix/Source/HarmonixMetasound/Private/HarmonixMetasound/Nodes/MetronomeNode.cpp b/Engine/Plugins/Runtime/Harmonix/Source/HarmonixMetasound/Private/HarmonixMetasound/Nodes/MetronomeNode.cpp
--- a/Engine/Plugins/Runtime/Harmonix/Source/HarmonixMetasound/Private/HarmonixMetasound/Nodes/MetronomeNode.cpp
+++ b/Engine/Plugins/Runtime/Harmonix/Source/HarmonixMetasound/Private/HarmonixMetasound/Nodes/MetronomeNode.cpp
@@ -64,7 +64,7 @@ namespace HarmonixMetasound
 
 		//** DATA
 		TSharedPtr<FMidiClock, ESPMode::NotThreadSafe> MonotonicallyIncreasingClock;
-		TSharedPtr<FMidiFileData> MidiData;
+		TSharedPtr<FSongMaps> SongMaps;
 		FSampleCount BlockSize;
 		float        SampleRate;
 		float        CurrentTempo;
@@ -74,7 +74,7 @@ namespace HarmonixMetasound
 		int32		 NextClockTickToProcess = 0;
 		bool		 bClocksArePreparedForExecute = true;
 
-		void BuildMidiData(bool ResetToStart = true);
+		void BuildSongMaps(bool ResetToStart = true);
 		void UpdateMidi();
 		void AddTempoChangeForMidi(float TempoBPM);
 		void AddTimeSigChangeForMidi(int32 TimeSigNum, int32 TimeSigDenom);
@@ -235,7 +235,7 @@ namespace HarmonixMetasound
 			MidiClockOutPin->SetDrivingClock(MonotonicallyIncreasingClock);
 		}
 
-		BuildMidiData();
+		BuildSongMaps();
 		
 		FTransportInitFn InitFn = [this](EMusicPlayerTransportState CurrentState)
 		{
@@ -298,7 +298,7 @@ namespace HarmonixMetasound
 				// Play from the beginning if we haven't received a seek call while we were stopped...
 				if (!ReceivedSeekWhileStopped())
 				{
-					BuildMidiData(true);
+					BuildSongMaps(true);
 					LastProcessedClockTick = -1;
 					NextClockTickToProcess = 0;
 				}
@@ -309,7 +309,7 @@ namespace HarmonixMetasound
 				return EMusicPlayerTransportState::Playing;
 				
 			case EMusicPlayerTransportState::Seeking:
-				BuildMidiData(false);
+				BuildSongMaps(false);
 				DrivingMidiClock.SeekTo(StartFrameIndex, TransportInPin->GetNextSeekDestination());
 				LastProcessedClockTick = DrivingMidiClock.GetLastProcessedMidiTick();
 				NextClockTickToProcess = DrivingMidiClock.GetNextMidiTickToProcess();
@@ -330,25 +330,25 @@ namespace HarmonixMetasound
 		MarkClocksAsExecuted();
 	}
 
-	void FMetronomeOperator::BuildMidiData(bool ResetToStart)
+	void FMetronomeOperator::BuildSongMaps(bool ResetToStart)
 	{
 		// make sure we have valid values
 		CurrentTempo = FMath::Max(1.0f, *TempoInPin);
 		CurrentTimeSigNum = FMath::Max(1, *TimeSigNumInPin);
 		CurrentTimeSigDenom = FMath::Max(1, *TimeSigDenomInPin);
 	
-		MidiData = FMidiClock::MakeClockConductorMidiData(CurrentTempo, CurrentTimeSigNum, CurrentTimeSigDenom);
+		SongMaps = MakeShared<FSongMaps>(CurrentTempo, CurrentTimeSigNum, CurrentTimeSigDenom);
 		
 		if (LoopInPin)
 		{
-			MonotonicallyIncreasingClock->AttachToMidiFile(MidiData, ResetToStart);
+			MonotonicallyIncreasingClock->AttachToSongMapEvaluator(SongMaps, ResetToStart);
 
 			// midi clock out will follow the tempo of the monotonically increasing clock
 			// so just assign it to some reasonable values. They will be ignored
-			TSharedPtr<FMidiFileData> MidiDataOut = FMidiClock::MakeClockConductorMidiData(120.0f, CurrentTimeSigNum, CurrentTimeSigDenom);
-			MidiClockOutPin->AttachToMidiFile(MidiDataOut, ResetToStart);
+			TSharedPtr<FSongMaps> SongMapsOut = MakeShared<FSongMaps>(120.0f, CurrentTimeSigNum, CurrentTimeSigDenom);
+			MidiClockOutPin->AttachToSongMapEvaluator(SongMapsOut, ResetToStart);
 			
-			int32 LoopEndTick = MidiDataOut->SongMaps.BarIncludingCountInToTick(FMath::Max(LoopLengthBarsInPin, 1));
+			int32 LoopEndTick = SongMapsOut->BarIncludingCountInToTick(FMath::Max(LoopLengthBarsInPin, 1));
 			// LoopEndTick == Loop Length since we are starting the loop at tick 0...
 			MidiClockOutPin->SetupPersistentLoop(0, LoopEndTick);
 		}
@@ -356,7 +356,7 @@ namespace HarmonixMetasound
 		{
 			// if we're not looping, then the MidiClockOut is going to be the driving midi clock
 			// so attach the midi data directly to it
-			MidiClockOutPin->AttachToMidiFile(MidiData);
+			MidiClockOutPin->AttachToSongMapEvaluator(SongMaps);
 			MidiClockOutPin->ClearPersistentLoop();
 		}
 	}
@@ -391,7 +391,7 @@ namespace HarmonixMetasound
 
 		if (HasMidiChanges)
 		{
-			DrivingMidiClock.MidiChanged();
+			DrivingMidiClock.SongMapsChanged();
 		}
 	}
 
@@ -399,7 +399,7 @@ namespace HarmonixMetasound
 	{
 		CurrentTempo = InTempoBPM;
 		int32 AtTick = GetDrivingMidiClock().GetNextMidiTickToProcess();
-		MidiData->AddTempoChange(0, AtTick, CurrentTempo);
+		SongMaps->AddTempoChange(AtTick, CurrentTempo);
 	}
 
 	void FMetronomeOperator::AddTimeSigChangeForMidi(int32 InTimeSigNum, int32 InTimeSigDenom)
@@ -408,12 +408,12 @@ namespace HarmonixMetasound
 		CurrentTimeSigDenom = InTimeSigDenom;
 		int32 AtTick = GetDrivingMidiClock().GetNextMidiTickToProcess();
 		// round to the next bar boundary, the bar we're actually going to apply the time sig change to
-		int32 AtBar = FMath::CeilToInt32(MidiData->SongMaps.GetBarIncludingCountInAtTick(AtTick));
-		int32 NumTimeSigPoints = MidiData->SongMaps.GetNumTimeSignatureChanges();
+		int32 AtBar = FMath::CeilToInt32(SongMaps->GetBarIncludingCountInAtTick(AtTick));
+		int32 NumTimeSigPoints = SongMaps->GetNumTimeSignatureChanges();
 
 		// check if there's already a time signature point at the bar we're trying to update
 		// the metronome clock increases monotonically, so we just have to check the _last_ point
-		FTimeSignaturePoint* Point = MidiData->SongMaps.GetMutableTimeSignaturePoint(NumTimeSigPoints - 1);
+		FTimeSignaturePoint* Point = SongMaps->GetMutableTimeSignaturePoint(NumTimeSigPoints - 1);
 		check(Point);
 		if (Point->BarIndex == AtBar)
 		{
@@ -423,7 +423,7 @@ namespace HarmonixMetasound
 		}
 		else
 		{
-			MidiData->AddTimeSigChange(0, AtTick, CurrentTimeSigNum, CurrentTimeSigDenom);
+			SongMaps->AddTimeSigChange(AtTick, CurrentTimeSigNum, CurrentTimeSigDenom);
 		}
 	}
 
diff --git a/Engine/Plugins/Runtime/Harmonix/Source/HarmonixMetasound/Private/HarmonixMetasound/Nodes/MidiClockOffsetNode.cpp b/Engine/Plugins/Runtime/Harmonix/Source/HarmonixMetasound/Private/HarmonixMetasound/Nodes/MidiClockOffsetNode.cpp
--- a/Engine/Plugins/Runtime/Harmonix/Source/HarmonixMetasound/Private/HarmonixMetasound/Nodes/MidiClockOffsetNode.cpp
+++ b/Engine/Plugins/Runtime/Harmonix/Source/HarmonixMetasound/Private/HarmonixMetasound/Nodes/MidiClockOffsetNode.cpp
@@ -109,8 +109,8 @@ namespace HarmonixMetasound::Nodes::MidiClockOffset
 			, MidiClockOut(FMidiClockWriteRef::CreateNew(InSettings))
 			, BlockSize(InSettings.GetNumFramesPerBlock())
 		{
-			TSharedPtr<FMidiFileData> ConductorMidiData = FMidiClock::MakeClockConductorMidiData(120.0f, 4, 4);
-			MidiClockOut->AttachToMidiFile(ConductorMidiData, true);
+			TSharedPtr<FSongMaps> SongMaps = MakeShared<FSongMaps>(120.0f, 4, 4);
+			MidiClockOut->AttachToSongMapEvaluator(SongMaps, true);
 			MidiClockOut->SetDrivingClock(MidiClockIn->AsShared().ToSharedPtr());
 			bClockOutNeedsPrepare = false;
 		}
diff --git a/Engine/Plugins/Runtime/Harmonix/Source/HarmonixMetasound/Private/HarmonixMetasound/Nodes/MidiPlayerNode.cpp b/Engine/Plugins/Runtime/Harmonix/Source/HarmonixMetasound/Private/HarmonixMetasound/Nodes/MidiPlayerNode.cpp
--- a/Engine/Plugins/Runtime/Harmonix/Source/HarmonixMetasound/Private/HarmonixMetasound/Nodes/MidiPlayerNode.cpp
+++ b/Engine/Plugins/Runtime/Harmonix/Source/HarmonixMetasound/Private/HarmonixMetasound/Nodes/MidiPlayerNode.cpp
@@ -579,7 +579,7 @@ namespace HarmonixMetasound
 
 		if (CurrentMidiFile.IsValid())
 		{			
-			MidiClockOut->AttachToMidiFile(CurrentMidiFile->GetMidiFile(), !IsPlaying());
+			MidiClockOut->AttachToSongMapEvaluator(CurrentMidiFile->GetMidiFile(), !IsPlaying());
 			MidiCursor.Prepare(CurrentMidiFile->GetMidiFile());
 			MidiOutPin->SetTicksPerQuarterNote(CurrentMidiFile->GetMidiFile()->TicksPerQuarterNote);
 			if (*LoopInPin)
@@ -616,7 +616,7 @@ namespace HarmonixMetasound
 		}
 		else
 		{
-			MidiClockOut->AttachToMidiFile(nullptr, !IsPlaying());
+			MidiClockOut->AttachToSongMapEvaluator(nullptr, !IsPlaying());
 			MidiCursor.Prepare(nullptr);
 		}
 	}
diff --git a/Engine/Plugins/Runtime/Harmonix/Source/HarmonixMetasound/Public/HarmonixMetasound/DataTypes/MidiClock.h b/Engine/Plugins/Runtime/Harmonix/Source/HarmonixMetasound/Public/HarmonixMetasound/DataTypes/MidiClock.h
--- a/Engine/Plugins/Runtime/Harmonix/Source/HarmonixMetasound/Public/HarmonixMetasound/DataTypes/MidiClock.h
+++ b/Engine/Plugins/Runtime/Harmonix/Source/HarmonixMetasound/Public/HarmonixMetasound/DataTypes/MidiClock.h
@@ -46,9 +46,10 @@ namespace HarmonixMetasound
 		virtual ~FMidiClock();
 		FMidiClock& operator=(const FMidiClock& Other);
 
-		void AttachToMidiFile(TSharedPtr<FMidiFileData> MidiData, bool ResetToStart = true);
-		void MidiChanged();
-		void DetachFromMidiFile();
+		void AttachToSongMapEvaluator(TSharedPtr<ISongMapEvaluator> SongMaps, bool ResetToStart = true);
+		void SongMapsChanged();
+		void DetachFromSongMaps();
+
 		const ISongMapEvaluator& GetSongMapEvaluator() const { return *SongMapEvaluator; }
 
 		void SetDrivingClock(FConstSharedMidiClockPtr NewExternalClockDriver);
@@ -162,11 +163,7 @@ namespace HarmonixMetasound
 		 */
 		int32 WrapTickIfLooping(int32 Tick) const;
 		
-		// Creates a new FMidiFileData with the given starting Tempo and Time Signature
-		// With max song length to be played indefinitely. Useful for Midi Clock Metronomes.
-		static TSharedPtr<FMidiFileData> MakeClockConductorMidiData(float TempoBPM, int32 TimeSigNum, int32 TimeSigDen);
-
-		bool GetMidiDataChangedInBlock() const { return MidiDataChangedInBlock; }
+		bool GetSongMapsChangedInBlock() const { return MidiDataChangedInBlock; }
 
 	private:
 		int32 GetNextTickToProcessAtBlockFrame(int32 BlockFrame) const;
@@ -186,7 +183,7 @@ namespace HarmonixMetasound
 		void AddLoopToBlock(int32 BlockFrameIndex, int32 FirstTick, int32 LoopLength);
 		void AddSeekToBlock(int32 BlockFrameIndex, int32 ToTick);
 		void AddAdvanceToBlock(int32 BlockFrameIndex, int32 FirstTick, int32 NumTicks);
-		void RebuildSongMapEvaluator(const TSharedPtr<const FMidiFileData>& MidiWithTempo, const TSharedPtr<const FMidiFileData>& MidiWithOtherMaps);
+		void RebuildSongMapEvaluator(const TSharedPtr<const ISongMapEvaluator>& MidiWithTempo, const TSharedPtr<const ISongMapEvaluator>& MidiWithOtherMaps);
 
 		TSharedPtr<FSongMapsWithAlternateTempoSource> SongMapEvaluator;
 		int32 CurrentTempoInfoPointIndex;
diff --git a/Engine/Plugins/Runtime/Harmonix/Source/HarmonixMetasoundTests/Private/Analysis/MidiClockVertexAnalyzerTests.cpp b/Engine/Plugins/Runtime/Harmonix/Source/HarmonixMetasoundTests/Private/Analysis/MidiClockVertexAnalyzerTests.cpp
--- a/Engine/Plugins/Runtime/Harmonix/Source/HarmonixMetasoundTests/Private/Analysis/MidiClockVertexAnalyzerTests.cpp
+++ b/Engine/Plugins/Runtime/Harmonix/Source/HarmonixMetasoundTests/Private/Analysis/MidiClockVertexAnalyzerTests.cpp
@@ -41,22 +41,12 @@ namespace HarmonixMetasoundTests::MidiClockVertexAnalyzer
 
 	void ResetAndStartClock(const HarmonixMetasound::FMidiClockWriteRef& ClockInput, float Tempo, float Speed, int32 TimeSigNumerator, int32 TimeSigDenominator)
 	{
-		const TSharedPtr<FMidiFileData> MidiData = MakeShared<FMidiFileData>();
-		check(MidiData);
-
-		MidiData->SongMaps.EmptyAllMaps();
-		MidiData->Tracks.Empty();
-
-		MidiData->Tracks.Add(FMidiTrack(TEXT("conductor")));
-		MidiData->Tracks[0].AddEvent(FMidiEvent(0, FMidiMsg(static_cast<uint8>(TimeSigNumerator), static_cast<uint8>(TimeSigDenominator))));
-		MidiData->SongMaps.AddTimeSignatureAtBarIncludingCountIn(0, TimeSigNumerator, TimeSigDenominator);
-		const int32 MidiTempo = Harmonix::Midi::Constants::BPMToMidiTempo(Tempo);
-		MidiData->Tracks[0].AddEvent(FMidiEvent(0, FMidiMsg(MidiTempo)));
-		MidiData->SongMaps.AddTempoInfoPoint(MidiTempo, 0);
-		MidiData->Tracks[0].Sort();
-		MidiData->ConformToLength(std::numeric_limits<int32>::max());
-
-		ClockInput->AttachToMidiFile(MidiData);
+		const TSharedPtr<FSongMaps> SongMaps = MakeShared<FSongMaps>(Tempo, TimeSigNumerator, TimeSigDenominator);
+		check(SongMaps);
+
+		SongMaps->SetSongLengthTicks(std::numeric_limits<int32>::max());
+
+		ClockInput->AttachToSongMapEvaluator(SongMaps);
 		ClockInput->SeekTo(0,0);
 		ClockInput->SetSpeed(0, Speed);
 		ClockInput->SetTransportState(0, HarmonixMetasound::EMusicPlayerTransportState::Playing);
diff --git a/Engine/Plugins/Runtime/Harmonix/Source/HarmonixMetasoundTests/Private/Analysis/MidiSongPosVertexAnalyzerTests.cpp b/Engine/Plugins/Runtime/Harmonix/Source/HarmonixMetasoundTests/Private/Analysis/MidiSongPosVertexAnalyzerTests.cpp
--- a/Engine/Plugins/Runtime/Harmonix/Source/HarmonixMetasoundTests/Private/Analysis/MidiSongPosVertexAnalyzerTests.cpp
+++ b/Engine/Plugins/Runtime/Harmonix/Source/HarmonixMetasoundTests/Private/Analysis/MidiSongPosVertexAnalyzerTests.cpp
@@ -41,37 +41,19 @@ namespace HarmonixMetasoundTests::MidiSongPosVertexAnalyzer
 
 	void ResetAndStartClock(const HarmonixMetasound::FMidiClockWriteRef& ClockInput)
 	{
-		const TSharedPtr<FMidiFileData> MidiData = MakeShared<FMidiFileData>();
-		check(MidiData);
+		const TSharedPtr<FSongMaps> SongMaps = MakeShared<FSongMaps>(240.3f, 4, 4);
+		check(SongMaps);
 
-		MidiData->SongMaps.EmptyAllMaps();
-		MidiData->Tracks.Empty();
+		SongMaps->AddTimeSignatureAtBarIncludingCountIn(4, 7, 8);
+		SongMaps->AddTempoChange(960, 200.0f);
+		SongMaps->SetSongLengthTicks(std::numeric_limits<int32>::max());
 
-		MidiData->Tracks.Add(FMidiTrack(TEXT("conductor")));
-		MidiData->Tracks[0].AddEvent(FMidiEvent(0, FMidiMsg(static_cast<uint8>(4), static_cast<uint8>(4))));
-		MidiData->SongMaps.AddTimeSignatureAtBarIncludingCountIn(0, 4, 4);
-
-		MidiData->Tracks[0].AddEvent(FMidiEvent(960 * 4 * 4, FMidiMsg(static_cast<uint8>(7), static_cast<uint8>(8))));
-		MidiData->SongMaps.AddTimeSignatureAtBarIncludingCountIn(4, 7, 8);
-
-		int32 MidiTempo = Harmonix::Midi::Constants::BPMToMidiTempo(240.3);
-		MidiData->Tracks[0].AddEvent(FMidiEvent(0, FMidiMsg(MidiTempo)));
-		MidiData->SongMaps.AddTempoInfoPoint(MidiTempo, 0);
-
-		MidiTempo = Harmonix::Midi::Constants::BPMToMidiTempo(200);
-		MidiData->Tracks[0].AddEvent(FMidiEvent(960, FMidiMsg(MidiTempo)));
-		MidiData->SongMaps.AddTempoInfoPoint(MidiTempo, 960);
-
-		MidiData->Tracks[0].Sort();
-
-		MidiData->ConformToLength(std::numeric_limits<int32>::max());
-
-		ClockInput->AttachToMidiFile(MidiData);
+		ClockInput->AttachToSongMapEvaluator(SongMaps);
 		ClockInput->SeekTo(0,0);
 		ClockInput->SetSpeed(0, 1.0f);
 		ClockInput->SetTransportState(0, HarmonixMetasound::EMusicPlayerTransportState::Playing);
-		int32 Bar1Tick = MidiData->SongMaps.BarBeatTickIncludingCountInToTick(1, 1, 0);
-		int32 Bar7Tick = MidiData->SongMaps.BarBeatTickIncludingCountInToTick(7, 1, 0);
+		int32 Bar1Tick = SongMaps->BarBeatTickIncludingCountInToTick(1, 1, 0);
+		int32 Bar7Tick = SongMaps->BarBeatTickIncludingCountInToTick(7, 1, 0);
 		ClockInput->SetupPersistentLoop(Bar1Tick, Bar7Tick - Bar1Tick);
 	}
 	
diff --git a/Engine/Plugins/Runtime/Harmonix/Source/HarmonixMetasoundTests/Private/Analysis/MidiStreamVertexAnalyzerTests.cpp b/Engine/Plugins/Runtime/Harmonix/Source/HarmonixMetasoundTests/Private/Analysis/MidiStreamVertexAnalyzerTests.cpp
--- a/Engine/Plugins/Runtime/Harmonix/Source/HarmonixMetasoundTests/Private/Analysis/MidiStreamVertexAnalyzerTests.cpp
+++ b/Engine/Plugins/Runtime/Harmonix/Source/HarmonixMetasoundTests/Private/Analysis/MidiStreamVertexAnalyzerTests.cpp
@@ -50,19 +50,13 @@ namespace HarmonixMetasoundTests::MidiStreamVertexAnalyzer
 		using namespace Metasound;
 
 		// Make the tempo and time sig maps
-		const TSharedPtr<FMidiFileData> MidiData = MakeShared<FMidiFileData>();
-		check(MidiData);
-		MidiData->Tracks.Add(FMidiTrack(TEXT("conductor")));
-		MidiData->Tracks[0].AddEvent(FMidiEvent(0, FMidiMsg(static_cast<uint8>(TimeSigNum), static_cast<uint8>(TimeSigDenom))));
-		MidiData->SongMaps.AddTimeSignatureAtBarIncludingCountIn(0, TimeSigNum, TimeSigDenom);
-		const int32 MidiTempo = Harmonix::Midi::Constants::BPMToMidiTempo(Tempo);
-		MidiData->Tracks[0].AddEvent(FMidiEvent(0, FMidiMsg(MidiTempo)));
-		MidiData->SongMaps.AddTempoInfoPoint(MidiTempo, 0);
-		MidiData->Tracks[0].Sort();
-		MidiData->ConformToLength(std::numeric_limits<int32>::max());
+		const TSharedPtr<FSongMaps> SongMaps = MakeShared<FSongMaps>(Tempo, TimeSigNum, TimeSigDenom);
+		check(SongMaps);
+
+		SongMaps->SetSongLengthTicks(std::numeric_limits<int32>::max());
 
 		// Attach the maps
-		Clock->AttachToMidiFile(MidiData);
+		Clock->AttachToSongMapEvaluator(SongMaps);
 	}
 	
 	IMPLEMENT_SIMPLE_AUTOMATION_TEST(
diff --git a/Engine/Plugins/Runtime/Harmonix/Source/HarmonixMetasoundTests/Private/DataTypes/MidiClock.spec.cpp b/Engine/Plugins/Runtime/Harmonix/Source/HarmonixMetasoundTests/Private/DataTypes/MidiClock.spec.cpp
--- a/Engine/Plugins/Runtime/Harmonix/Source/HarmonixMetasoundTests/Private/DataTypes/MidiClock.spec.cpp
+++ b/Engine/Plugins/Runtime/Harmonix/Source/HarmonixMetasoundTests/Private/DataTypes/MidiClock.spec.cpp
@@ -16,7 +16,7 @@ namespace HarmonixMetasoundTests::MidiClock
 
 	TUniquePtr<FMidiClock> TestClock;
 	TSharedPtr<FMidiClock, ESPMode::NotThreadSafe> DrivingClock;
-	TSharedPtr<FMidiFileData> MidiFileData;
+	TSharedPtr<FSongMaps> SongMaps;
 	Metasound::FOperatorSettings OperatorSettings {48000, 100};
 
 	void AddStateAtFrame(EMusicPlayerTransportState State, int32 Frame) const
@@ -404,11 +404,11 @@ namespace HarmonixMetasoundTests::MidiClock
 		{
 			It("Without driving clock - One Tempo Change At Span End", [this]()
 			{
-				MidiFileData = FMidiClock::MakeClockConductorMidiData(123, 5, 8);
+				SongMaps = MakeShared<FSongMaps>(123, 5, 8);
 				constexpr int32 TempoChangeTick = 234;
 				constexpr float TempoChangeTempo = 89;
-				MidiFileData->AddTempoChange(0, TempoChangeTick, TempoChangeTempo);
-				TestClock->AttachToMidiFile(MidiFileData);
+				SongMaps->AddTempoChange(TempoChangeTick, TempoChangeTempo);
+				TestClock->AttachToSongMapEvaluator(SongMaps);
 				TestClock->SetTransportState(0, HarmonixMetasound::EMusicPlayerTransportState::Playing);
 
 				while (TestClock->GetLastProcessedMidiTick() < TempoChangeTick)
@@ -448,11 +448,11 @@ namespace HarmonixMetasoundTests::MidiClock
 
 			It("Without driving clock - One Tempo Change At Span Start", [this]()
 			{
-				MidiFileData = FMidiClock::MakeClockConductorMidiData(123, 5, 8);
+				SongMaps = MakeShared<FSongMaps>(123, 5, 8);
 				constexpr int32 TempoChangeTick = 230;
 				constexpr float TempoChangeTempo = 89;
-				MidiFileData->AddTempoChange(0, TempoChangeTick, TempoChangeTempo);
-				TestClock->AttachToMidiFile(MidiFileData);
+				SongMaps->AddTempoChange(TempoChangeTick, TempoChangeTempo);
+				TestClock->AttachToSongMapEvaluator(SongMaps);
 				TestClock->SetTransportState(0, HarmonixMetasound::EMusicPlayerTransportState::Playing);
 
 				while (TestClock->GetLastProcessedMidiTick() < TempoChangeTick)
@@ -492,15 +492,15 @@ namespace HarmonixMetasoundTests::MidiClock
 
 			It("Without driving clock - Many Tempo Changes In Span", [this]()
 			{
-				MidiFileData = FMidiClock::MakeClockConductorMidiData(123, 5, 8);
+				SongMaps = MakeShared<FSongMaps>(123, 5, 8);
 				constexpr int32 NumChanges = 4;
 				constexpr int32 TempoChangeTicks[NumChanges] = { 230, 231, 232, 233 };
 				constexpr float TempoChangeTempos [NumChanges] = { 89.0, 89.2, 89.4, 89.6 };
 				for (int32 i = 0; i < NumChanges; ++i)
 				{
-					MidiFileData->AddTempoChange(0, TempoChangeTicks[i], TempoChangeTempos[i]);
+					SongMaps->AddTempoChange(TempoChangeTicks[i], TempoChangeTempos[i]);
 				}
-				TestClock->AttachToMidiFile(MidiFileData);
+				TestClock->AttachToSongMapEvaluator(SongMaps);
 				TestClock->SetTransportState(0, HarmonixMetasound::EMusicPlayerTransportState::Playing);
 
 				while (TestClock->GetLastProcessedMidiTick() < TempoChangeTicks[NumChanges - 1])
@@ -543,11 +543,11 @@ namespace HarmonixMetasoundTests::MidiClock
 		{
 			It("Without driving clock - One Change", [this]()
 			{
-				MidiFileData = FMidiClock::MakeClockConductorMidiData(123, 5, 8);
-				const int32 TimeSigChangeTick = MidiFileData->SongMaps.BarBeatTickIncludingCountInToTick(2, 1, 0);
+				SongMaps = MakeShared<FSongMaps>(123, 5, 8);
+				const int32 TimeSigChangeTick = SongMaps->BarBeatTickIncludingCountInToTick(2, 1, 0);
 				const FTimeSignature NewTimeSig{ 3, 4 };
-				MidiFileData->AddTimeSigChange(0, TimeSigChangeTick, NewTimeSig.Numerator, NewTimeSig.Denominator);
-				TestClock->AttachToMidiFile(MidiFileData);
+				SongMaps->AddTimeSigChange(TimeSigChangeTick, NewTimeSig.Numerator, NewTimeSig.Denominator);
+				TestClock->AttachToSongMapEvaluator(SongMaps);
 				TestClock->SetTransportState(0, HarmonixMetasound::EMusicPlayerTransportState::Playing);
 						
 				while (TestClock->GetLastProcessedMidiTick() < TimeSigChangeTick)
@@ -587,17 +587,17 @@ namespace HarmonixMetasoundTests::MidiClock
 
 			It("Without driving clock - One Change with tempos", [this]()
 			{
-				MidiFileData = FMidiClock::MakeClockConductorMidiData(123, 5, 8);
-				const int32 TimeSigChangeTick = MidiFileData->SongMaps.BarBeatTickIncludingCountInToTick(2, 1, 0);
+				SongMaps = MakeShared<FSongMaps>(123, 5, 8);
+				const int32 TimeSigChangeTick = SongMaps->BarBeatTickIncludingCountInToTick(2, 1, 0);
 				constexpr int32 NumTempoChanges = 3;
 				constexpr int32 TempoChangeTicks[NumTempoChanges] = { 4799, 4800, 4801 };
 				constexpr float TempoChangeTempos[NumTempoChanges] = { 155.0f, 157.2f, 158.4f };
-				MidiFileData->AddTempoChange(0, TempoChangeTicks[0], TempoChangeTempos[0]);
+				SongMaps->AddTempoChange(TempoChangeTicks[0], TempoChangeTempos[0]);
 				const FTimeSignature NewTimeSig{ 3, 4 };
-				MidiFileData->AddTimeSigChange(0, TimeSigChangeTick, NewTimeSig.Numerator, NewTimeSig.Denominator);
-				MidiFileData->AddTempoChange(0, TempoChangeTicks[1], TempoChangeTempos[1]);
-				MidiFileData->AddTempoChange(0, TempoChangeTicks[2], TempoChangeTempos[2]);
-				TestClock->AttachToMidiFile(MidiFileData);
+				SongMaps->AddTimeSigChange(TimeSigChangeTick, NewTimeSig.Numerator, NewTimeSig.Denominator);
+				SongMaps->AddTempoChange(TempoChangeTicks[1], TempoChangeTempos[1]);
+				SongMaps->AddTempoChange(TempoChangeTicks[2], TempoChangeTempos[2]);
+				TestClock->AttachToSongMapEvaluator(SongMaps);
 				TestClock->SetTransportState(0, HarmonixMetasound::EMusicPlayerTransportState::Playing);
 
 				while (TestClock->GetLastProcessedMidiTick() < TimeSigChangeTick)
diff --git a/Engine/Plugins/Runtime/Harmonix/Source/HarmonixMetasoundTests/Private/MidiOps/PulseGeneratorTests.cpp b/Engine/Plugins/Runtime/Harmonix/Source/HarmonixMetasoundTests/Private/MidiOps/PulseGeneratorTests.cpp
--- a/Engine/Plugins/Runtime/Harmonix/Source/HarmonixMetasoundTests/Private/MidiOps/PulseGeneratorTests.cpp
+++ b/Engine/Plugins/Runtime/Harmonix/Source/HarmonixMetasoundTests/Private/MidiOps/PulseGeneratorTests.cpp
@@ -21,7 +21,7 @@ namespace Harmonix::Midi::Ops::Tests
 		constexpr uint8 TimeSigNumerator = 4;
 		constexpr uint8 TimeSigDenominator = 4;
 		const auto Clock = MakeShared<HarmonixMetasound::FMidiClock, ESPMode::NotThreadSafe>(OperatorSettings);
-		Clock->AttachToMidiFile(HarmonixMetasound::FMidiClock::MakeClockConductorMidiData(Tempo, TimeSigNumerator, TimeSigDenominator));
+		Clock->AttachToSongMapEvaluator(MakeShared<FSongMaps>(Tempo, TimeSigNumerator, TimeSigDenominator));
 		
 		HarmonixMetasound::FMidiStream OutputStream;
 
diff --git a/Engine/Plugins/Runtime/Harmonix/Source/HarmonixMetasoundTests/Private/Nodes/DelayNodeTests.cpp b/Engine/Plugins/Runtime/Harmonix/Source/HarmonixMetasoundTests/Private/Nodes/DelayNodeTests.cpp
--- a/Engine/Plugins/Runtime/Harmonix/Source/HarmonixMetasoundTests/Private/Nodes/DelayNodeTests.cpp
+++ b/Engine/Plugins/Runtime/Harmonix/Source/HarmonixMetasoundTests/Private/Nodes/DelayNodeTests.cpp
@@ -240,24 +240,12 @@ namespace HarmonixMetasoundTests::DelayNode
 			{
 				return false;
 			}
-			const TSharedPtr<FMidiFileData> MidiData = MakeShared<FMidiFileData>();
-			check(MidiData);
-
-			MidiData->SongMaps.EmptyAllMaps();
-			MidiData->Tracks.Empty();
-
-			MidiData->Tracks.Add(FMidiTrack(TEXT("conductor")));
-			MidiData->Tracks[0].AddEvent(FMidiEvent(0, FMidiMsg(static_cast<uint8>(TimeSigNum), static_cast<uint8>(TimeSigDenom))));
-			MidiData->SongMaps.AddTimeSignatureAtBarIncludingCountIn(0, TimeSigNum, TimeSigDenom);
-			const int32 MidiTempo = Harmonix::Midi::Constants::BPMToMidiTempo(Tempo);
-			MidiData->Tracks[0].AddEvent(FMidiEvent(0, FMidiMsg(MidiTempo)));
-			MidiData->SongMaps.AddTempoInfoPoint(MidiTempo, 0);
-			MidiData->Tracks[0].Sort();
-			MidiData->ConformToLength(std::numeric_limits<int32>::max());
-			MidiData->SongMaps.GetSongLengthData().LengthTicks = std::numeric_limits<int32>::max();
-			MidiData->SongMaps.GetSongLengthData().LengthFractionalBars = std::numeric_limits<float>::max();
-
-			(*ClockInput)->AttachToMidiFile(MidiData);
+			const TSharedPtr<FSongMaps> SongMaps = MakeShared<FSongMaps>(Tempo, TimeSigNum, TimeSigDenom);
+			check(SongMaps);
+
+			SongMaps->SetSongLengthTicks(std::numeric_limits<int32>::max());
+
+			(*ClockInput)->AttachToSongMapEvaluator(SongMaps);
 			(*ClockInput)->SetTransportState(0, HarmonixMetasound::EMusicPlayerTransportState::Playing);
 
 			SampleRemainder = 0;
diff --git a/Engine/Plugins/Runtime/Harmonix/Source/HarmonixMetasoundTests/Private/Nodes/MetronomeNodeTests.cpp b/Engine/Plugins/Runtime/Harmonix/Source/HarmonixMetasoundTests/Private/Nodes/MetronomeNodeTests.cpp
--- a/Engine/Plugins/Runtime/Harmonix/Source/HarmonixMetasoundTests/Private/Nodes/MetronomeNodeTests.cpp
+++ b/Engine/Plugins/Runtime/Harmonix/Source/HarmonixMetasoundTests/Private/Nodes/MetronomeNodeTests.cpp
@@ -121,8 +121,8 @@ namespace HarmonixMetasoundTests::MetronomeNode
 			float SecsPerBlock = Params.NumSamplesPerBlock / Params.SampleRate;
 			float TicksPerBlock = TicksPerSec * SecsPerBlock;
 
-			TSharedPtr<FMidiFileData> MidiData = FMidiClock::MakeClockConductorMidiData(Params.Tempo, Params.TimeSigNumerator, Params.TimeSigDenominator);
-			int32 LoopLengthTicks = MidiData->SongMaps.BarIncludingCountInToTick(Params.LoopLengthBars);
+			TSharedPtr<FSongMaps> SongMaps = MakeShared<FSongMaps>(Params.Tempo, Params.TimeSigNumerator, Params.TimeSigDenominator);
+			int32 LoopLengthTicks = SongMaps->BarIncludingCountInToTick(Params.LoopLengthBars);
 			
 			//test for tempo consistency by stopping and restarting the transport (clock output and tempo map)
 			if (TestCaseString.Equals(TempoChangeTestString))
diff --git a/Engine/Plugins/Runtime/Harmonix/Source/HarmonixMetasoundTests/Private/Nodes/MidiPlayerNodeTests.cpp b/Engine/Plugins/Runtime/Harmonix/Source/HarmonixMetasoundTests/Private/Nodes/MidiPlayerNodeTests.cpp
--- a/Engine/Plugins/Runtime/Harmonix/Source/HarmonixMetasoundTests/Private/Nodes/MidiPlayerNodeTests.cpp
+++ b/Engine/Plugins/Runtime/Harmonix/Source/HarmonixMetasoundTests/Private/Nodes/MidiPlayerNodeTests.cpp
@@ -159,8 +159,8 @@ namespace HarmonixMetasoundTests::MidiPlayerNode
 
 				MetronomeClockOut = Generator->GetOutputReadReference<FMidiClock>(NodeNames::Metronome::ClockOut);
 
-				TSharedPtr<FMidiFileData> MidiData = FMidiClock::MakeClockConductorMidiData(Params.Metronome.Tempo, Params.Metronome.TimeSig.Numerator, Params.Metronome.TimeSig.Denominator);
-				MetronomeLoopLengthTicks = MidiData->SongMaps.BarIncludingCountInToTick(Params.Metronome.LoopLengthBars);
+				TSharedPtr<FSongMaps> SongMaps = MakeShared<FSongMaps>(Params.Metronome.Tempo, Params.Metronome.TimeSig.Numerator, Params.Metronome.TimeSig.Denominator);
+				MetronomeLoopLengthTicks = SongMaps->BarIncludingCountInToTick(Params.Metronome.LoopLengthBars);
 			}
 
 			if (Params.MidiPlayer.MidiFile.IsValid())
diff --git a/Engine/Plugins/Runtime/Harmonix/Source/HarmonixMetasoundTests/Private/Nodes/MorphingLfoNode.spec.cpp b/Engine/Plugins/Runtime/Harmonix/Source/HarmonixMetasoundTests/Private/Nodes/MorphingLfoNode.spec.cpp
--- a/Engine/Plugins/Runtime/Harmonix/Source/HarmonixMetasoundTests/Private/Nodes/MorphingLfoNode.spec.cpp
+++ b/Engine/Plugins/Runtime/Harmonix/Source/HarmonixMetasoundTests/Private/Nodes/MorphingLfoNode.spec.cpp
@@ -62,22 +62,12 @@ namespace HarmonixMetasoundTests::MorphingLFONode
 			return false;
 		}
 		
-		const TSharedPtr<FMidiFileData> MidiData = MakeShared<FMidiFileData>();
-		check(MidiData);
-
-		MidiData->SongMaps.EmptyAllMaps();
-		MidiData->Tracks.Empty();
-
-		MidiData->Tracks.Add(FMidiTrack(TEXT("conductor")));
-		MidiData->Tracks[0].AddEvent(FMidiEvent(0, FMidiMsg(static_cast<uint8>(TimeSigNumerator), static_cast<uint8>(TimeSigDenominator))));
-		MidiData->SongMaps.AddTimeSignatureAtBarIncludingCountIn(0, TimeSigNumerator, TimeSigNumerator);
-		const int32 MidiTempo = Harmonix::Midi::Constants::BPMToMidiTempo(Tempo);
-		MidiData->Tracks[0].AddEvent(FMidiEvent(0, FMidiMsg(MidiTempo)));
-		MidiData->SongMaps.AddTempoInfoPoint(MidiTempo, 0);
-		MidiData->Tracks[0].Sort();
-		MidiData->ConformToLength(std::numeric_limits<int32>::max());
-
-		(*ClockInput)->AttachToMidiFile(MidiData);
+		const TSharedPtr<FSongMaps> SongMaps = MakeShared<FSongMaps>(Tempo, TimeSigNumerator, TimeSigDenominator);
+		check(SongMaps);
+
+		SongMaps->SetSongLengthTicks(std::numeric_limits<int32>::max());
+
+		(*ClockInput)->AttachToSongMapEvaluator(SongMaps);
 		(*ClockInput)->SetTransportState(0, HarmonixMetasound::EMusicPlayerTransportState::Playing);
 		(*ClockInput)->SetSpeed(0, Speed);
 
diff --git a/Engine/Plugins/Runtime/Harmonix/Source/HarmonixMetasoundTests/Private/Subsystems/MidiClockUpdateSubsystemTests.cpp b/Engine/Plugins/Runtime/Harmonix/Source/HarmonixMetasoundTests/Private/Subsystems/MidiClockUpdateSubsystemTests.cpp
--- a/Engine/Plugins/Runtime/Harmonix/Source/HarmonixMetasoundTests/Private/Subsystems/MidiClockUpdateSubsystemTests.cpp
+++ b/Engine/Plugins/Runtime/Harmonix/Source/HarmonixMetasoundTests/Private/Subsystems/MidiClockUpdateSubsystemTests.cpp
@@ -17,20 +17,13 @@ namespace HarmonixMetasoundTests::MidiClockUpdateSubsystem
 			int32 TimeSigNum,
 			int32 TimeSigDenom)
 		{
-			const TSharedPtr<FMidiFileData> MidiData = MakeShared<FMidiFileData>();
-			check(MidiData);
+			const TSharedPtr<FSongMaps> SongMaps = MakeShared<FSongMaps>(Tempo, TimeSigNum, TimeSigDenom);
+			check(SongMaps);
 			
-			MidiData->Tracks.Add(FMidiTrack(TEXT("conductor")));
-			MidiData->Tracks[0].AddEvent(FMidiEvent(0, FMidiMsg(static_cast<uint8>(TimeSigNum), static_cast<uint8>(TimeSigDenom))));
-			MidiData->SongMaps.AddTimeSignatureAtBarIncludingCountIn(0, TimeSigNum, TimeSigDenom);
-			const int32 MidiTempo = Harmonix::Midi::Constants::BPMToMidiTempo(Tempo);
-			MidiData->Tracks[0].AddEvent(FMidiEvent(0, FMidiMsg(MidiTempo)));
-			MidiData->SongMaps.AddTempoInfoPoint(MidiTempo, 0);
-			MidiData->Tracks[0].Sort();
-			MidiData->ConformToLength(std::numeric_limits<int32>::max());
+			SongMaps->SetSongLengthTicks(std::numeric_limits<int32>::max());
 
 			HarmonixMetasound::FMidiClock Clock{ OperatorSettings };
-			Clock.AttachToMidiFile(MidiData);
+			Clock.AttachToSongMapEvaluator(SongMaps);
 			Clock.SetTransportState(0, HarmonixMetasound::EMusicPlayerTransportState::Playing);
 
 			return Clock;
diff --git a/Engine/Plugins/Runtime/Harmonix/Source/HarmonixMidi/Private/HarmonixMidi/SongMaps.cpp b/Engine/Plugins/Runtime/Harmonix/Source/HarmonixMidi/Private/HarmonixMidi/SongMaps.cpp
--- a/Engine/Plugins/Runtime/Harmonix/Source/HarmonixMidi/Private/HarmonixMidi/SongMaps.cpp
+++ b/Engine/Plugins/Runtime/Harmonix/Source/HarmonixMidi/Private/HarmonixMidi/SongMaps.cpp
@@ -988,6 +988,23 @@ void FSongMaps::SetLengthTotalBars(int32 Bars)
 	LengthData.LengthFractionalBars = Bars;
 }
 
+void FSongMaps::AddTempoChange(int32 Tick, float TempoBPM)
+{
+	int32 MidiTempo = Harmonix::Midi::Constants::BPMToMidiTempo(TempoBPM);
+	AddTempoInfoPoint(MidiTempo, Tick);
+}
+
+void FSongMaps::AddTimeSigChange(int32 Tick, int32 InTimeSigNum, int32 InTimeSigDenom)
+{
+	// Time signature changes can only happen at the beginning of a bar,
+	// so round up to the next bar boundary...
+	int32 AbsoluteBar = FMath::CeilToInt32(GetBarIncludingCountInAtTick(Tick));
+	Tick = BarBeatTickIncludingCountInToTick(AbsoluteBar, 1, 0);
+	int32 TimeSigNum = FMath::Clamp(InTimeSigNum, 1, 64);
+	int32 TimeSigDenom = FMath::Clamp(InTimeSigDenom, 1, 64);
+	AddTimeSignatureAtBarIncludingCountIn(AbsoluteBar, TimeSigNum, TimeSigDenom);
+}
+
 int32 ISongMapEvaluator::CalculateMidiTick(const FMusicTimestamp& Timestamp, const EMidiClockSubdivisionQuantization Quantize) const
 {
 	const FBarMap& BarMap = GetBarMap();
@@ -1193,30 +1210,30 @@ bool ISongMapEvaluator::ChordMapIsEmpty() const
 
 const FTempoMap& FSongMapsWithAlternateTempoSource::GetTempoMap() const
 {
-	return MidiFileWithTempoMap->SongMaps.GetTempoMap();
+	return SongMapsWithTempoMap->GetTempoMap();
 }
 
 const FBeatMap& FSongMapsWithAlternateTempoSource::GetBeatMap() const
 {
-	return MidiFileWithOtherMaps->SongMaps.GetBeatMap();
+	return SongMapsWithOtherMaps->GetBeatMap();
 }
 
 const FBarMap& FSongMapsWithAlternateTempoSource::GetBarMap() const
 {
-	return MidiFileWithOtherMaps->SongMaps.GetBarMap();
+	return SongMapsWithOtherMaps->GetBarMap();
 }
 
 const FSectionMap& FSongMapsWithAlternateTempoSource::GetSectionMap() const
 {
-	return MidiFileWithOtherMaps->SongMaps.GetSectionMap();
+	return SongMapsWithOtherMaps->GetSectionMap();
 }
 
 const FChordProgressionMap& FSongMapsWithAlternateTempoSource::GetChordMap() const
 {
-	return MidiFileWithOtherMaps->SongMaps.GetChordMap();
+	return SongMapsWithOtherMaps->GetChordMap();
 }
 
 const FSongLengthData& FSongMapsWithAlternateTempoSource::GetSongLengthData() const
 {
-	return MidiFileWithOtherMaps->SongMaps.GetSongLengthData();
+	return SongMapsWithOtherMaps->GetSongLengthData();
 }
diff --git a/Engine/Plugins/Runtime/Harmonix/Source/HarmonixMidi/Public/HarmonixMidi/MidiFile.h b/Engine/Plugins/Runtime/Harmonix/Source/HarmonixMidi/Public/HarmonixMidi/MidiFile.h
--- a/Engine/Plugins/Runtime/Harmonix/Source/HarmonixMidi/Public/HarmonixMidi/MidiFile.h
+++ b/Engine/Plugins/Runtime/Harmonix/Source/HarmonixMidi/Public/HarmonixMidi/MidiFile.h
@@ -10,7 +10,6 @@
 
 #include "MidiFile.generated.h"
 
-struct FSongMaps;
 class UAssetImportData;
 
 class FMidiFileProxy;
@@ -18,6 +17,9 @@ using FMidiFileProxyPtr = TSharedPtr<FMidiFileProxy, ESPMode::ThreadSafe>;
 
 USTRUCT(BlueprintType, Meta = (DisplayName = "MIDI File Data"))
 struct HARMONIXMIDI_API FMidiFileData
+#if CPP
+	: public ISongMapEvaluator
+#endif
 {
 	GENERATED_BODY()
 
@@ -25,6 +27,8 @@ struct HARMONIXMIDI_API FMidiFileData
 		: TicksPerQuarterNote(Harmonix::Midi::Constants::GTicksPerQuarterNoteInt)
 	{}
 
+	virtual ~FMidiFileData() = default;
+
 	bool operator==(const FMidiFileData& Other) const;
 
 	void Empty()
@@ -82,6 +86,15 @@ struct HARMONIXMIDI_API FMidiFileData
 	void ScanTracksForSongLengthChange();
 
 	bool LengthIsAPerfectSubdivision() const;
+
+	//** BEGIN ISongMapEvaluator Overrides ***********
+	virtual const FSongLengthData& GetSongLengthData() const override { return SongMaps.GetSongLengthData(); }
+	virtual const FTempoMap& GetTempoMap() const override { return SongMaps.GetTempoMap(); }
+	virtual const FBeatMap& GetBeatMap() const override { return SongMaps.GetBeatMap(); }
+	virtual const FBarMap& GetBarMap() const override { return SongMaps.GetBarMap(); }
+	virtual const FSectionMap& GetSectionMap() const override { return SongMaps.GetSectionMap(); }
+	virtual const FChordProgressionMap& GetChordMap() const override { return SongMaps.GetChordMap(); }
+	//** END ISongMapEvaluator Overrides ***********
 };
 
 template<>
diff --git a/Engine/Plugins/Runtime/Harmonix/Source/HarmonixMidi/Public/HarmonixMidi/SongMaps.h b/Engine/Plugins/Runtime/Harmonix/Source/HarmonixMidi/Public/HarmonixMidi/SongMaps.h
--- a/Engine/Plugins/Runtime/Harmonix/Source/HarmonixMidi/Public/HarmonixMidi/SongMaps.h
+++ b/Engine/Plugins/Runtime/Harmonix/Source/HarmonixMidi/Public/HarmonixMidi/SongMaps.h
@@ -177,8 +177,9 @@ public:
 	//virtual FSongLengthData& GetSongLengthData() = 0;
 	virtual const FSongLengthData& GetSongLengthData() const = 0;
 
-private:
+protected:
 	friend struct FSongMaps;
+	friend struct FSongMapsWithAlternateTempoSource;
 	virtual const FTempoMap& GetTempoMap() const = 0;
 	virtual const FBeatMap& GetBeatMap() const = 0;
 	virtual const FBarMap& GetBarMap() const = 0;
@@ -258,6 +259,9 @@ public:
 	void SetSongLengthTicks(int32 NewLengthTicks);
 	void FinalizeBarMap(int32 InLastTick);
 
+	void AddTempoChange(int32 Tick, float TempoBPM);
+	void AddTimeSigChange(int32 Tick, int32 TimeSigNum, int32 TimeSigDenom);
+
 	bool AddTempoInfoPoint(int32 MicrosecondsPerQuarterNote, int32 Tick, bool SortNow = true);
 	bool AddTimeSignatureAtBarIncludingCountIn(int32 BarIndex, int32 InNumerator, int32 InDenominator, bool SortNow = true, bool FailOnError = true);
 	FTimeSignaturePoint* GetMutableTimeSignaturePoint(int32 PointIndex);
@@ -294,26 +298,21 @@ private:
 struct HARMONIXMIDI_API FSongMapsWithAlternateTempoSource : public ISongMapEvaluator
 {
 public:
-	FSongMapsWithAlternateTempoSource(const TSharedPtr<const FMidiFileData>& MidiWithTempo, const TSharedPtr<const FMidiFileData>& MidiWithOthers)
-		: MidiFileWithTempoMap(MidiWithTempo)
-		, MidiFileWithOtherMaps(MidiWithOthers)
+	FSongMapsWithAlternateTempoSource(const TSharedPtr<const ISongMapEvaluator>& SongMapsWithTempo, const TSharedPtr<const ISongMapEvaluator>& SongMapsWithOthers)
+		: SongMapsWithTempoMap(SongMapsWithTempo)
+		, SongMapsWithOtherMaps(SongMapsWithOthers)
 	{}
 
-	FSongMapsWithAlternateTempoSource(const TSharedPtr<const FMidiFileData>& MidiData)
-		: MidiFileWithTempoMap(MidiData)
-		, MidiFileWithOtherMaps(MidiData)
+	FSongMapsWithAlternateTempoSource(const TSharedPtr<const ISongMapEvaluator>& SongMaps)
+		: SongMapsWithTempoMap(SongMaps)
+		, SongMapsWithOtherMaps(SongMaps)
 	{
 	}
 
-	FSongMapsWithAlternateTempoSource(const TSharedPtr<const FSongMapsWithAlternateTempoSource>& Other)
-		: MidiFileWithTempoMap(Other->MidiFileWithTempoMap)
-		, MidiFileWithOtherMaps(Other->MidiFileWithOtherMaps)
-	{}
-
 	FSongMapsWithAlternateTempoSource& operator=(const TSharedPtr<const FSongMapsWithAlternateTempoSource>& Other)
 	{
-		MidiFileWithTempoMap = Other->MidiFileWithTempoMap;
-		MidiFileWithOtherMaps = Other->MidiFileWithOtherMaps;
+		SongMapsWithTempoMap = Other->SongMapsWithTempoMap;
+		SongMapsWithOtherMaps = Other->SongMapsWithOtherMaps;
 		return *this;
 	}
 
@@ -321,15 +320,15 @@ public:
 
 	operator bool() const 
 	{
-		return MidiFileWithTempoMap.IsValid() && MidiFileWithOtherMaps.IsValid();
+		return SongMapsWithTempoMap.IsValid() && SongMapsWithOtherMaps.IsValid();
 	}
 
-	const TSharedPtr<const FMidiFileData>& GetMidiFileWithTempoMap() const { return MidiFileWithTempoMap; }
-	const TSharedPtr<const FMidiFileData>& GetMidiFileWithOtherMaps() const { return MidiFileWithOtherMaps; }
+	const TSharedPtr<const ISongMapEvaluator>& GetSongMapsWithTempoMap() const { return SongMapsWithTempoMap; }
+	const TSharedPtr<const ISongMapEvaluator>& GetSongMapsWithOtherMaps() const { return SongMapsWithOtherMaps; }
 
-	bool AllMapsHaveOneSource() const { return MidiFileWithTempoMap == MidiFileWithOtherMaps; }
+	bool AllMapsHaveOneSource() const { return SongMapsWithTempoMap == SongMapsWithOtherMaps; }
 
-private:
+protected:
 	virtual const FTempoMap& GetTempoMap() const override;
 	virtual const FBeatMap& GetBeatMap() const override;
 	virtual const FBarMap& GetBarMap() const override;
@@ -337,6 +336,6 @@ private:
 	virtual const FChordProgressionMap& GetChordMap() const override;
 	virtual const FSongLengthData& GetSongLengthData() const override;
 
-	TSharedPtr<const FMidiFileData> MidiFileWithTempoMap;
-	TSharedPtr<const FMidiFileData> MidiFileWithOtherMaps;
+	TSharedPtr<const ISongMapEvaluator> SongMapsWithTempoMap;
+	TSharedPtr<const ISongMapEvaluator> SongMapsWithOtherMaps;
 };
