[Harmonix] Delete MIDI Play Cursor and related code

#rnx

diff --git a/Engine/Plugins/Runtime/Harmonix/Source/HarmonixMetasound/Private/HarmonixMetasound/Nodes/MidiClockOffsetNode.cpp b/Engine/Plugins/Runtime/Harmonix/Source/HarmonixMetasound/Private/HarmonixMetasound/Nodes/MidiClockOffsetNode.cpp
--- a/Engine/Plugins/Runtime/Harmonix/Source/HarmonixMetasound/Private/HarmonixMetasound/Nodes/MidiClockOffsetNode.cpp
+++ b/Engine/Plugins/Runtime/Harmonix/Source/HarmonixMetasound/Private/HarmonixMetasound/Nodes/MidiClockOffsetNode.cpp
@@ -15,7 +15,6 @@
 #include "HarmonixMetasound/DataTypes/MidiAsset.h"
 #include "HarmonixMetasound/DataTypes/MusicTransport.h"
 
-#include "HarmonixMidi/MidiPlayCursor.h"
 #include "HarmonixMidi/MidiVoiceId.h"
 
 #define LOCTEXT_NAMESPACE "HarmonixMetaSound"
diff --git a/Engine/Plugins/Runtime/Harmonix/Source/HarmonixMetasound/Private/HarmonixMetasound/Nodes/StepSequencePlayerNode.cpp b/Engine/Plugins/Runtime/Harmonix/Source/HarmonixMetasound/Private/HarmonixMetasound/Nodes/StepSequencePlayerNode.cpp
--- a/Engine/Plugins/Runtime/Harmonix/Source/HarmonixMetasound/Private/HarmonixMetasound/Nodes/StepSequencePlayerNode.cpp
+++ b/Engine/Plugins/Runtime/Harmonix/Source/HarmonixMetasound/Private/HarmonixMetasound/Nodes/StepSequencePlayerNode.cpp
@@ -17,7 +17,6 @@
 #include "HarmonixMetasound/DataTypes/MidiStepSequence.h"
 #include "HarmonixMetasound/DataTypes/MusicTransport.h"
 
-#include "HarmonixMidi/MidiPlayCursor.h"
 #include "HarmonixMidi/MidiVoiceId.h"
 
 #define LOCTEXT_NAMESPACE "HarmonixMetaSound"
diff --git a/Engine/Plugins/Runtime/Harmonix/Source/HarmonixMidi/Private/HarmonixMidi/MidiPlayCursor.cpp b/Engine/Plugins/Runtime/Harmonix/Source/HarmonixMidi/Private/HarmonixMidi/MidiPlayCursor.cpp
deleted file mode 100644
diff --git a/Engine/Plugins/Runtime/Harmonix/Source/HarmonixMidi/Private/HarmonixMidi/MidiPlayCursorMgr.cpp b/Engine/Plugins/Runtime/Harmonix/Source/HarmonixMidi/Private/HarmonixMidi/MidiPlayCursorMgr.cpp
deleted file mode 100644
diff --git a/Engine/Plugins/Runtime/Harmonix/Source/HarmonixMidi/Private/HarmonixMidi/MidiPlayCursorTracker.cpp b/Engine/Plugins/Runtime/Harmonix/Source/HarmonixMidi/Private/HarmonixMidi/MidiPlayCursorTracker.cpp
deleted file mode 100644
diff --git a/Engine/Plugins/Runtime/Harmonix/Source/HarmonixMidi/Private/HarmonixMidi/SmoothedMidiPlayCursor.cpp b/Engine/Plugins/Runtime/Harmonix/Source/HarmonixMidi/Private/HarmonixMidi/SmoothedMidiPlayCursor.cpp
deleted file mode 100644
diff --git a/Engine/Plugins/Runtime/Harmonix/Source/HarmonixMidi/Public/HarmonixMidi/MidiPlayCursor.h b/Engine/Plugins/Runtime/Harmonix/Source/HarmonixMidi/Public/HarmonixMidi/MidiPlayCursor.h
deleted file mode 100644
diff --git a/Engine/Plugins/Runtime/Harmonix/Source/HarmonixMidi/Public/HarmonixMidi/MidiPlayCursorMgr.h b/Engine/Plugins/Runtime/Harmonix/Source/HarmonixMidi/Public/HarmonixMidi/MidiPlayCursorMgr.h
deleted file mode 100644
diff --git a/Engine/Plugins/Runtime/Harmonix/Source/HarmonixMidi/Public/HarmonixMidi/MidiPlayCursorTracker.h b/Engine/Plugins/Runtime/Harmonix/Source/HarmonixMidi/Public/HarmonixMidi/MidiPlayCursorTracker.h
deleted file mode 100644
diff --git a/Engine/Plugins/Runtime/Harmonix/Source/HarmonixMidi/Public/HarmonixMidi/SmoothedMidiPlayCursor.h b/Engine/Plugins/Runtime/Harmonix/Source/HarmonixMidi/Public/HarmonixMidi/SmoothedMidiPlayCursor.h
deleted file mode 100644
