[Harmonix] Submit left behind files for 34132680 to unbreak build.
#rnx

diff --git a/Engine/Plugins/Runtime/Harmonix/Source/HarmonixMetasound/Public/HarmonixMetasound/Components/MusicClockComponent.h b/Engine/Plugins/Runtime/Harmonix/Source/HarmonixMetasound/Public/HarmonixMetasound/Components/MusicClockComponent.h
--- a/Engine/Plugins/Runtime/Harmonix/Source/HarmonixMetasound/Public/HarmonixMetasound/Components/MusicClockComponent.h
+++ b/Engine/Plugins/Runtime/Harmonix/Source/HarmonixMetasound/Public/HarmonixMetasound/Components/MusicClockComponent.h
@@ -4,8 +4,8 @@
 #include "Components/ActorComponent.h"
 #include "Harmonix/MusicalTimebase.h"
 #include "HarmonixMidi/MidiSongPos.h"
+#include "HarmonixMidi/SongMaps.h"
 #include "Delegates/DelegateCombinations.h"
-#include "HarmonixMidi/SmoothedMidiPlayCursor.h"
 #include "Templates/UniquePtr.h"
 
 #include "MusicClockComponent.generated.h"
diff --git a/Engine/Plugins/Runtime/Harmonix/Source/HarmonixMetasound/Public/HarmonixMetasound/DataTypes/MidiClockEvent.h b/Engine/Plugins/Runtime/Harmonix/Source/HarmonixMetasound/Public/HarmonixMetasound/DataTypes/MidiClockEvent.h
--- a/Engine/Plugins/Runtime/Harmonix/Source/HarmonixMetasound/Public/HarmonixMetasound/DataTypes/MidiClockEvent.h
+++ b/Engine/Plugins/Runtime/Harmonix/Source/HarmonixMetasound/Public/HarmonixMetasound/DataTypes/MidiClockEvent.h
@@ -2,7 +2,6 @@
 
 #pragma once
 #include "HarmonixMetasound/DataTypes/MusicTransport.h"
-#include "HarmonixMidi/MidiPlayCursor.h"
 
 namespace HarmonixMetasound
 {
