Add "raw audio render time" as a timebase. SHOULD ONLY BE USED FOR DEBUGGING! It is NOT generally useful 99.9999% of the time.

diff --git a/Engine/Plugins/Runtime/Harmonix/Source/Harmonix/Public/Harmonix/MusicalTimebase.h b/Engine/Plugins/Runtime/Harmonix/Source/Harmonix/Public/Harmonix/MusicalTimebase.h
--- a/Engine/Plugins/Runtime/Harmonix/Source/Harmonix/Public/Harmonix/MusicalTimebase.h
+++ b/Engine/Plugins/Runtime/Harmonix/Source/Harmonix/Public/Harmonix/MusicalTimebase.h
@@ -26,6 +26,12 @@ enum class ECalibratedMusicTimebase : uint8
 	 * Tells you what you should be drawing right now so it appears *in sync with the music.* (when properly calibrated).
 	 * Useful for synchronizing animations, ui, and other visuals to the music.
 	 */
-	VideoRenderTime
+	VideoRenderTime,
+
+	/*
+	 * Tells you exactly where the audio renderer is right now. NOTE: Un-smoothed, jittery,
+	 * raw position of the audio rendering. Only useful for debugging!
+	 */
+	RawAudioRenderTime
 };
-ENUM_RANGE_BY_FIRST_AND_LAST(ECalibratedMusicTimebase, ECalibratedMusicTimebase::AudioRenderTime, ECalibratedMusicTimebase::VideoRenderTime);
+ENUM_RANGE_BY_FIRST_AND_LAST(ECalibratedMusicTimebase, ECalibratedMusicTimebase::AudioRenderTime, ECalibratedMusicTimebase::RawAudioRenderTime);
diff --git a/Engine/Plugins/Runtime/Harmonix/Source/HarmonixMetasound/Private/HarmonixMetasound/Components/MetasoundMusicClockDriver.cpp b/Engine/Plugins/Runtime/Harmonix/Source/HarmonixMetasound/Private/HarmonixMetasound/Components/MetasoundMusicClockDriver.cpp
--- a/Engine/Plugins/Runtime/Harmonix/Source/HarmonixMetasound/Private/HarmonixMetasound/Components/MetasoundMusicClockDriver.cpp
+++ b/Engine/Plugins/Runtime/Harmonix/Source/HarmonixMetasound/Private/HarmonixMetasound/Components/MetasoundMusicClockDriver.cpp
@@ -355,13 +355,10 @@ void FMetasoundMusicClockDriver::RefreshCurrentSongPosFromWallClock()
 
 	double FreeRunTime = (ClockComponent->GetWorld()->GetTimeSeconds() - FreeRunStartTimeSecs) * ClockComponent->CurrentClockAdvanceRate;
 
-	ClockComponent->CurrentSmoothedAudioRenderSongPos.SetByTime(((float)FreeRunTime * 1000.0) + SongPosOffsetMs, ClockComponent->DefaultMaps);
+	ClockComponent->CurrentRawAudioRenderSongPos.SetByTime(((float)FreeRunTime * 1000.0) + SongPosOffsetMs, ClockComponent->DefaultMaps);
+	ClockComponent->CurrentSmoothedAudioRenderSongPos = ClockComponent->CurrentRawAudioRenderSongPos;
 	ClockComponent->CurrentPlayerExperiencedSongPos.SetByTime(ClockComponent->CurrentSmoothedAudioRenderSongPos.SecondsIncludingCountIn * 1000.0f - FHarmonixModule::GetMeasuredUserExperienceAndReactionToAudioRenderOffsetMs(), ClockComponent->DefaultMaps);
 	ClockComponent->CurrentVideoRenderSongPos.SetByTime(ClockComponent->CurrentSmoothedAudioRenderSongPos.SecondsIncludingCountIn * 1000.0f - FHarmonixModule::GetMeasuredVideoToAudioRenderOffsetMs(), ClockComponent->DefaultMaps);
-	if (ClockComponent->CurrentSmoothedAudioRenderSongPos.SecondsIncludingCountIn > ClockComponent->RawUnsmoothedAudioRenderPos.SecondsIncludingCountIn)
-	{
-		ClockComponent->RawUnsmoothedAudioRenderPos.SetByTime(ClockComponent->CurrentSmoothedAudioRenderSongPos.SecondsIncludingCountIn * 1000.0f, ClockComponent->DefaultMaps);
-	}
 
 	if (TempoChanged)
 	{
@@ -410,7 +407,7 @@ void FMetasoundMusicClockDriver::RefreshCurrentSongPosFromHistory()
 	}
 
 	auto Entry = ClockHistory->Positions.GetEntry(ClockHistory->Positions.GetLastWriteIndex());
-	ClockComponent->RawUnsmoothedAudioRenderPos.SetByTick(Entry->Item.UpToTick, *(CurrentMapChain->SongMaps));
+	ClockComponent->CurrentRawAudioRenderSongPos.SetByTick(Entry->Item.UpToTick, *(CurrentMapChain->SongMaps));
 	Metasound::FSampleCount LastRenderPosSampleCount = Entry->Item.SampleCount;
 	float SpeedAtRawRenderTime = Entry->Item.CurrentSpeed;
 	LastTickSeen = Entry->Item.UpToTick;
diff --git a/Engine/Plugins/Runtime/Harmonix/Source/HarmonixMetasound/Private/HarmonixMetasound/Components/MusicClockComponent.cpp b/Engine/Plugins/Runtime/Harmonix/Source/HarmonixMetasound/Private/HarmonixMetasound/Components/MusicClockComponent.cpp
--- a/Engine/Plugins/Runtime/Harmonix/Source/HarmonixMetasound/Private/HarmonixMetasound/Components/MusicClockComponent.cpp
+++ b/Engine/Plugins/Runtime/Harmonix/Source/HarmonixMetasound/Private/HarmonixMetasound/Components/MusicClockComponent.cpp
@@ -195,6 +195,9 @@ FMidiSongPos UMusicClockComponent::CalculateSongPosWithOffset(float MsOffset, EC
 	const FSongMaps* Maps = &DefaultMaps;
 	switch (Timebase)
 	{
+	case ECalibratedMusicTimebase::RawAudioRenderTime:
+		Result.SetByTime((CurrentRawAudioRenderSongPos.SecondsIncludingCountIn * 1000.0f) + MsOffset, *Maps);
+		break;
 	case ECalibratedMusicTimebase::AudioRenderTime:
 		Result.SetByTime((CurrentSmoothedAudioRenderSongPos.SecondsIncludingCountIn * 1000.0f) + MsOffset, *Maps);
 		break;
@@ -251,7 +254,7 @@ void UMusicClockComponent::Start()
 	CurrentSmoothedAudioRenderSongPos.Reset();
 	CurrentVideoRenderSongPos.Reset();
 	CurrentPlayerExperiencedSongPos.Reset();
-	RawUnsmoothedAudioRenderPos.Reset();
+	CurrentRawAudioRenderSongPos.Reset();
 	State = EMusicClockState::Running;
 	PlayStateEvent.Broadcast(State);
 }
@@ -298,7 +301,7 @@ void UMusicClockComponent::Stop()
 	CurrentSmoothedAudioRenderSongPos.Reset();
 	CurrentVideoRenderSongPos.Reset();
 	CurrentPlayerExperiencedSongPos.Reset();
-	RawUnsmoothedAudioRenderPos.Reset();
+	CurrentRawAudioRenderSongPos.Reset();
 	PlayStateEvent.Broadcast(State);
 }
 
@@ -307,10 +310,11 @@ float UMusicClockComponent::GetSecondsIncludingCountIn(ECalibratedMusicTimebase
 	EnsureClockIsValidForGameFrame();
 	switch (Timebase)
 	{
-	case ECalibratedMusicTimebase::AudioRenderTime:	return CurrentSmoothedAudioRenderSongPos.SecondsIncludingCountIn;
-	case ECalibratedMusicTimebase::ExperiencedTime:	return CurrentPlayerExperiencedSongPos.SecondsIncludingCountIn;
+	case ECalibratedMusicTimebase::RawAudioRenderTime: return CurrentRawAudioRenderSongPos.SecondsIncludingCountIn;
+	case ECalibratedMusicTimebase::AudioRenderTime:	   return CurrentSmoothedAudioRenderSongPos.SecondsIncludingCountIn;
+	case ECalibratedMusicTimebase::ExperiencedTime:	   return CurrentPlayerExperiencedSongPos.SecondsIncludingCountIn;
 	case ECalibratedMusicTimebase::VideoRenderTime:
-	default:                                        return CurrentVideoRenderSongPos.SecondsIncludingCountIn;
+	default:                                           return CurrentVideoRenderSongPos.SecondsIncludingCountIn;
 	}
 }
 
@@ -319,10 +323,11 @@ float UMusicClockComponent::GetSecondsFromBarOne(ECalibratedMusicTimebase Timeba
 	EnsureClockIsValidForGameFrame();
 	switch (Timebase)
 	{
-	case ECalibratedMusicTimebase::AudioRenderTime:	return CurrentSmoothedAudioRenderSongPos.SecondsFromBarOne;
-	case ECalibratedMusicTimebase::ExperiencedTime:	return CurrentPlayerExperiencedSongPos.SecondsFromBarOne;
+	case ECalibratedMusicTimebase::RawAudioRenderTime: return CurrentRawAudioRenderSongPos.SecondsFromBarOne;
+	case ECalibratedMusicTimebase::AudioRenderTime:	   return CurrentSmoothedAudioRenderSongPos.SecondsFromBarOne;
+	case ECalibratedMusicTimebase::ExperiencedTime:	   return CurrentPlayerExperiencedSongPos.SecondsFromBarOne;
 	case ECalibratedMusicTimebase::VideoRenderTime:
-	default:                                        return CurrentVideoRenderSongPos.SecondsFromBarOne;
+	default:                                           return CurrentVideoRenderSongPos.SecondsFromBarOne;
 	}
 }
 
@@ -331,10 +336,11 @@ float UMusicClockComponent::GetBarsIncludingCountIn(ECalibratedMusicTimebase Tim
 	EnsureClockIsValidForGameFrame();
 	switch (Timebase)
 	{
-	case ECalibratedMusicTimebase::AudioRenderTime:	return CurrentSmoothedAudioRenderSongPos.BarsIncludingCountIn;
-	case ECalibratedMusicTimebase::ExperiencedTime:	return CurrentPlayerExperiencedSongPos.BarsIncludingCountIn;
+	case ECalibratedMusicTimebase::RawAudioRenderTime: return CurrentRawAudioRenderSongPos.BarsIncludingCountIn;
+	case ECalibratedMusicTimebase::AudioRenderTime:	   return CurrentSmoothedAudioRenderSongPos.BarsIncludingCountIn;
+	case ECalibratedMusicTimebase::ExperiencedTime:	   return CurrentPlayerExperiencedSongPos.BarsIncludingCountIn;
 	case ECalibratedMusicTimebase::VideoRenderTime:
-	default:                                        return CurrentVideoRenderSongPos.BarsIncludingCountIn;
+	default:                                           return CurrentVideoRenderSongPos.BarsIncludingCountIn;
 	}
 }
 
@@ -343,10 +349,11 @@ float UMusicClockComponent::GetBeatsIncludingCountIn(ECalibratedMusicTimebase Ti
 	EnsureClockIsValidForGameFrame();
 	switch (Timebase)
 	{
-	case ECalibratedMusicTimebase::AudioRenderTime: return CurrentSmoothedAudioRenderSongPos.BeatsIncludingCountIn;
-	case ECalibratedMusicTimebase::ExperiencedTime:	return CurrentPlayerExperiencedSongPos.BeatsIncludingCountIn;
+	case ECalibratedMusicTimebase::RawAudioRenderTime: return CurrentRawAudioRenderSongPos.BarsIncludingCountIn;
+	case ECalibratedMusicTimebase::AudioRenderTime:    return CurrentSmoothedAudioRenderSongPos.BeatsIncludingCountIn;
+	case ECalibratedMusicTimebase::ExperiencedTime:	   return CurrentPlayerExperiencedSongPos.BeatsIncludingCountIn;
 	case ECalibratedMusicTimebase::VideoRenderTime:
-	default:                                        return CurrentVideoRenderSongPos.BeatsIncludingCountIn;
+	default:                                           return CurrentVideoRenderSongPos.BeatsIncludingCountIn;
 	}
 }
 
@@ -370,10 +377,11 @@ FMusicTimestamp UMusicClockComponent::GetCurrentTimestamp(ECalibratedMusicTimeba
 	EnsureClockIsValidForGameFrame();
 	switch (Timebase)
 	{
-	case ECalibratedMusicTimebase::AudioRenderTime:	return CurrentSmoothedAudioRenderSongPos.Timestamp;
-	case ECalibratedMusicTimebase::ExperiencedTime:	return CurrentPlayerExperiencedSongPos.Timestamp;
+	case ECalibratedMusicTimebase::RawAudioRenderTime: return CurrentRawAudioRenderSongPos.Timestamp;
+	case ECalibratedMusicTimebase::AudioRenderTime:	   return CurrentSmoothedAudioRenderSongPos.Timestamp;
+	case ECalibratedMusicTimebase::ExperiencedTime:	   return CurrentPlayerExperiencedSongPos.Timestamp;
 	case ECalibratedMusicTimebase::VideoRenderTime:
-	default:                                        return CurrentVideoRenderSongPos.Timestamp;
+	default:                                           return CurrentVideoRenderSongPos.Timestamp;
 	}
 }
 
@@ -436,10 +444,11 @@ float UMusicClockComponent::GetDeltaBar(ECalibratedMusicTimebase Timebase) const
 	EnsureClockIsValidForGameFrame();
 	switch (Timebase)
 	{
-	case ECalibratedMusicTimebase::AudioRenderTime: return AudioRenderDeltaBarF;
-	case ECalibratedMusicTimebase::ExperiencedTime:	return PlayerExperienceDeltaBarF;
+	case ECalibratedMusicTimebase::RawAudioRenderTime: return RawAudioRenderDeltaBarF;
+	case ECalibratedMusicTimebase::AudioRenderTime:    return AudioRenderDeltaBarF;
+	case ECalibratedMusicTimebase::ExperiencedTime:	   return PlayerExperienceDeltaBarF;
 	case ECalibratedMusicTimebase::VideoRenderTime:
-	default:                                        return VideoRenderDeltaBarF;
+	default:                                           return VideoRenderDeltaBarF;
 	}
 }
 
@@ -448,10 +457,11 @@ float UMusicClockComponent::GetDeltaBeat(ECalibratedMusicTimebase Timebase) cons
 	EnsureClockIsValidForGameFrame();
 	switch (Timebase)
 	{
-	case ECalibratedMusicTimebase::AudioRenderTime: return AudioRenderDeltaBeatF;
-	case ECalibratedMusicTimebase::ExperiencedTime:	return PlayerExperienceDeltaBeatF;
+	case ECalibratedMusicTimebase::RawAudioRenderTime: return RawAudioRenderDeltaBeatF;
+	case ECalibratedMusicTimebase::AudioRenderTime:    return AudioRenderDeltaBeatF;
+	case ECalibratedMusicTimebase::ExperiencedTime:	   return PlayerExperienceDeltaBeatF;
 	case ECalibratedMusicTimebase::VideoRenderTime:
-	default:                                        return VideoRenderDeltaBeatF;
+	default:                                           return VideoRenderDeltaBeatF;
 	}
 }
 
@@ -563,6 +573,8 @@ const FMidiSongPos& UMusicClockComponent::GetSongPos(ECalibratedMusicTimebase Ti
 
 	switch (Timebase)
 	{
+	case ECalibratedMusicTimebase::RawAudioRenderTime:
+		return CurrentRawAudioRenderSongPos;
 	case ECalibratedMusicTimebase::AudioRenderTime:
 		return CurrentSmoothedAudioRenderSongPos;
 	case ECalibratedMusicTimebase::ExperiencedTime:
@@ -588,6 +600,11 @@ FMidiSongPos UMusicClockComponent::GetCurrentPlayerExperiencedSongPos() const
 	return GetSongPos(ECalibratedMusicTimebase::ExperiencedTime);
 }
 
+FMidiSongPos UMusicClockComponent::GetCurrentRawAudioRenderSongPos() const
+{
+	return GetSongPos(ECalibratedMusicTimebase::RawAudioRenderTime);
+}
+
 float UMusicClockComponent::MeasureSpanProgress(const FMusicalTimeSpan& Span, ECalibratedMusicTimebase Timebase) const
 {
 	EnsureClockIsValidForGameFrame();
@@ -601,6 +618,9 @@ float UMusicClockComponent::MeasureSpanProgress(const FMusicalTimeSpan& Span, EC
 	const FMidiSongPos* Basis = &CurrentVideoRenderSongPos;
 	switch (Timebase)
 	{
+	case ECalibratedMusicTimebase::RawAudioRenderTime:
+		Basis = &CurrentRawAudioRenderSongPos;
+		break;
 	case ECalibratedMusicTimebase::AudioRenderTime:
 		Basis = &CurrentSmoothedAudioRenderSongPos;
 		break;
@@ -617,6 +637,9 @@ void UMusicClockComponent::BroadcastSongPosChanges()
 	const FMidiSongPos* Basis = &CurrentVideoRenderSongPos;
 	switch (TimebaseForBarAndBeatEvents)
 	{
+	case ECalibratedMusicTimebase::RawAudioRenderTime:
+		Basis = &CurrentRawAudioRenderSongPos;
+		break;
 	case ECalibratedMusicTimebase::AudioRenderTime:
 		Basis = &CurrentSmoothedAudioRenderSongPos;
 		break;
@@ -658,12 +681,15 @@ void FMusicClockDriverBase::EnsureClockIsValidForGameFrame()
 	// to make sure its current state is appropriate to the current musical time. See 
 	// UMusicClockComponent::EnsureClockIsValidForGameFrame for more details as to why 
 	// this is so.
+	ClockComponent->PrevRawAudioRenderSongPos = ClockComponent->CurrentRawAudioRenderSongPos;
 	ClockComponent->PrevAudioRenderSongPos = ClockComponent->CurrentSmoothedAudioRenderSongPos;
 	ClockComponent->PrevPlayerExperiencedSongPos = ClockComponent->CurrentPlayerExperiencedSongPos;
 	ClockComponent->PrevVideoRenderSongPos = ClockComponent->CurrentVideoRenderSongPos;
 
 	if (RefreshCurrentSongPos())
 	{
+		ClockComponent->RawAudioRenderDeltaBarF = ClockComponent->CurrentRawAudioRenderSongPos.BarsIncludingCountIn - ClockComponent->PrevRawAudioRenderSongPos.BarsIncludingCountIn;
+		ClockComponent->RawAudioRenderDeltaBeatF = ClockComponent->CurrentRawAudioRenderSongPos.BeatsIncludingCountIn - ClockComponent->PrevRawAudioRenderSongPos.BeatsIncludingCountIn;
 		ClockComponent->AudioRenderDeltaBarF = ClockComponent->CurrentSmoothedAudioRenderSongPos.BarsIncludingCountIn - ClockComponent->PrevAudioRenderSongPos.BarsIncludingCountIn;
 		ClockComponent->AudioRenderDeltaBeatF = ClockComponent->CurrentSmoothedAudioRenderSongPos.BeatsIncludingCountIn - ClockComponent->PrevAudioRenderSongPos.BeatsIncludingCountIn;
 		ClockComponent->PlayerExperienceDeltaBarF = ClockComponent->CurrentPlayerExperiencedSongPos.BarsIncludingCountIn - ClockComponent->PrevPlayerExperiencedSongPos.BarsIncludingCountIn;
diff --git a/Engine/Plugins/Runtime/Harmonix/Source/HarmonixMetasound/Public/HarmonixMetasound/Components/MusicClockComponent.h b/Engine/Plugins/Runtime/Harmonix/Source/HarmonixMetasound/Public/HarmonixMetasound/Components/MusicClockComponent.h
--- a/Engine/Plugins/Runtime/Harmonix/Source/HarmonixMetasound/Public/HarmonixMetasound/Components/MusicClockComponent.h
+++ b/Engine/Plugins/Runtime/Harmonix/Source/HarmonixMetasound/Public/HarmonixMetasound/Components/MusicClockComponent.h
@@ -290,6 +290,8 @@ private:
 	FMidiSongPos CurrentVideoRenderSongPos;
 	UPROPERTY(BlueprintGetter = GetCurrentPlayerExperiencedSongPos, Category = "MusicClock")
 	FMidiSongPos CurrentPlayerExperiencedSongPos;
+	UPROPERTY(BlueprintGetter = GetCurrentRawAudioRenderSongPos, Category = "MusicClock")
+	FMidiSongPos CurrentRawAudioRenderSongPos;
 
 public:
 	// Getter functions for the Blueprint properties exposed above...
@@ -302,6 +304,9 @@ public:
 	UFUNCTION(BlueprintCallable, BlueprintPure, Category = "MusicClock")
 	FMidiSongPos GetCurrentPlayerExperiencedSongPos() const;
 
+	UFUNCTION(BlueprintCallable, BlueprintPure, Category = "MusicClock")
+	FMidiSongPos GetCurrentRawAudioRenderSongPos() const;
+
 	// Note: Not const as it might cause the clock to update from its source.
 	UFUNCTION(BlueprintPure, Category = "MusicClock")
 	float MeasureSpanProgress(const FMusicalTimeSpan& Span, ECalibratedMusicTimebase Timebase = ECalibratedMusicTimebase::VideoRenderTime) const;
@@ -321,7 +326,7 @@ public:
 	// Note: Not const as it might cause the clock to update from its source.
 	FMidiSongPos CalculateSongPosWithOffset(float MsOffset, ECalibratedMusicTimebase Timebase = ECalibratedMusicTimebase::VideoRenderTime) const;
 
-	const FMidiSongPos& GetRawUnsmoothedAudioRenderPos() const { return RawUnsmoothedAudioRenderPos; }
+	const FMidiSongPos& GetRawUnsmoothedAudioRenderPos() const { return CurrentRawAudioRenderSongPos; }
 
 protected:
 	virtual void TickComponent(float DeltaTime, enum ELevelTick TickType, FActorComponentTickFunction* ThisTickFunction);
@@ -337,8 +342,9 @@ private:
 	EMusicClockState State = EMusicClockState::Stopped;
 
 	FSongMaps DefaultMaps;
-	FMidiSongPos RawUnsmoothedAudioRenderPos;
 
+	float RawAudioRenderDeltaBarF = 0.0f;
+	float RawAudioRenderDeltaBeatF = 0.0f;
 	float AudioRenderDeltaBarF = 0.0f;
 	float AudioRenderDeltaBeatF = 0.0f;
 	float PlayerExperienceDeltaBarF = 0.0f;
@@ -349,6 +355,7 @@ private:
 	int32 LastBroadcastBeat = -1;
 	FSongSection LastBroadcastSongSection;
 
+	FMidiSongPos PrevRawAudioRenderSongPos;
 	FMidiSongPos PrevAudioRenderSongPos;
 	FMidiSongPos PrevPlayerExperiencedSongPos;
 	FMidiSongPos PrevVideoRenderSongPos;
