- Don't create song maps with nonsensical time signatures
- Don't crash if a song map has a nonsensical time signature

diff --git a/Engine/Plugins/Runtime/Harmonix/Source/HarmonixMidi/Private/HarmonixMidi/BarMap.cpp b/Engine/Plugins/Runtime/Harmonix/Source/HarmonixMidi/Private/HarmonixMidi/BarMap.cpp
--- a/Engine/Plugins/Runtime/Harmonix/Source/HarmonixMidi/Private/HarmonixMidi/BarMap.cpp
+++ b/Engine/Plugins/Runtime/Harmonix/Source/HarmonixMidi/Private/HarmonixMidi/BarMap.cpp
@@ -290,6 +290,23 @@ FMusicTimestamp FBarMap::TickToMusicTimestamp(float Tick, int32* OutBeatsPerBar)
 		TimeSigIndex = 0;
 	}
 
+	if (Points[TimeSigIndex].TimeSignature.Numerator == 0 || Points[TimeSigIndex].TimeSignature.Denominator == 0)
+	{
+		// Assume 4/4 time.
+		Result.Bar = int32(Tick) / (TicksPerQuarterNote * 4);
+		Tick -= (Result.Bar * TicksPerQuarterNote * 4);
+		if (Tick < 0)
+		{
+			Result.Bar--;
+			Tick += TicksPerQuarterNote * 4;
+		}
+		Result.Beat = Tick / TicksPerQuarterNote;
+		Result.Bar += StartBar;
+		Result.Beat += 1.0f; // 1 based
+		if (OutBeatsPerBar) *OutBeatsPerBar = 4;
+		return Result;
+	}
+
 	int32 TicksPerBar = GetTicksInBarAfterPoint(TimeSigIndex);
 	float TicksPassed = Tick - Points[TimeSigIndex].StartTick;
 	if (TicksPassed < 0)
diff --git a/Engine/Plugins/Runtime/Harmonix/Source/HarmonixMidi/Private/HarmonixMidi/SongMaps.cpp b/Engine/Plugins/Runtime/Harmonix/Source/HarmonixMidi/Private/HarmonixMidi/SongMaps.cpp
--- a/Engine/Plugins/Runtime/Harmonix/Source/HarmonixMidi/Private/HarmonixMidi/SongMaps.cpp
+++ b/Engine/Plugins/Runtime/Harmonix/Source/HarmonixMidi/Private/HarmonixMidi/SongMaps.cpp
@@ -12,7 +12,7 @@ FSongMaps::FSongMaps()
 FSongMaps::FSongMaps(float Bpm, int32 TimeSigNumerator, int32 TimeSigDenominator)
 {
 	TempoMap.AddTempoInfoPoint(Harmonix::Midi::Constants::BPMToMidiTempo(Bpm), 0);
-	BarMap.AddTimeSignatureAtBarIncludingCountIn(0, TimeSigNumerator, TimeSigDenominator);
+	BarMap.AddTimeSignatureAtBarIncludingCountIn(0, TimeSigNumerator == 0 ? 4 : TimeSigNumerator, TimeSigDenominator == 0 ? 4 : TimeSigDenominator);
 }
 
 FSongMaps::FSongMaps(const ISongMapEvaluator& Other)
