Update ClockDriver when input binding is updated in MidiPlayerNode
#tests PIE, build

diff --git a/Engine/Plugins/Runtime/Harmonix/Source/HarmonixMetasound/Private/HarmonixMetasound/Nodes/MidiPlayerNode.cpp b/Engine/Plugins/Runtime/Harmonix/Source/HarmonixMetasound/Private/HarmonixMetasound/Nodes/MidiPlayerNode.cpp
--- a/Engine/Plugins/Runtime/Harmonix/Source/HarmonixMetasound/Private/HarmonixMetasound/Nodes/MidiPlayerNode.cpp
+++ b/Engine/Plugins/Runtime/Harmonix/Source/HarmonixMetasound/Private/HarmonixMetasound/Nodes/MidiPlayerNode.cpp
@@ -291,6 +291,7 @@ namespace HarmonixMetasound
 		FMidiPlayerOperator::BindInputs(InVertexData);
 		using namespace CommonPinNames;
 		InVertexData.BindReadVertex(METASOUND_GET_PARAM_NAME(Inputs::MidiClock), MidiClockIn);
+		MidiClockOut->SetDrivingClock(MidiClockIn->AsShared().ToSharedPtr());
 	}
 
 	void FExternallyClockedMidiPlayerOperator::BindOutputs(FOutputVertexInterfaceData& InVertexData)
