Fix issue where rebinding a MidiPlayer to a new "driving clock" wasn't causing the midi players clock to seek to the driving clock's position.

diff --git a/Engine/Plugins/Runtime/Harmonix/Source/HarmonixMetasound/Private/HarmonixMetasound/DataTypes/MidiClock.cpp b/Engine/Plugins/Runtime/Harmonix/Source/HarmonixMetasound/Private/HarmonixMetasound/DataTypes/MidiClock.cpp
--- a/Engine/Plugins/Runtime/Harmonix/Source/HarmonixMetasound/Private/HarmonixMetasound/DataTypes/MidiClock.cpp
+++ b/Engine/Plugins/Runtime/Harmonix/Source/HarmonixMetasound/Private/HarmonixMetasound/DataTypes/MidiClock.cpp
@@ -70,6 +70,7 @@ namespace HarmonixMetasound
 		, FirstTickInLoop(-1)
 		, LoopLengthTicks(0)
 		, MidiDataChangedInBlock(false)
+		, NeedsSeekToDrivingClock(false)
 	{
 	}
 
@@ -106,6 +107,7 @@ namespace HarmonixMetasound
 		, FirstTickInLoop(Other.FirstTickInLoop)
 		, LoopLengthTicks(Other.LoopLengthTicks)
 		, MidiDataChangedInBlock(Other.MidiDataChangedInBlock)
+		, NeedsSeekToDrivingClock(Other.NeedsSeekToDrivingClock)
 		, MidiClockEventsInBlock(Other.MidiClockEventsInBlock)
 	{
 	}
@@ -148,6 +150,7 @@ namespace HarmonixMetasound
 			FirstTickInLoop = Other.FirstTickInLoop;
 			LoopLengthTicks = Other.LoopLengthTicks;
 			MidiDataChangedInBlock = Other.MidiDataChangedInBlock;
+			NeedsSeekToDrivingClock = Other.NeedsSeekToDrivingClock;
 			MidiClockEventsInBlock = Other.MidiClockEventsInBlock;
 		}
 
@@ -195,6 +198,8 @@ namespace HarmonixMetasound
 	{
 		MidiDataChangedInBlock = true;
 		ExternalClockDriver = NewExternalClockDriver;
+		TickResidualWhenDriven = 0.0f;
+		NeedsSeekToDrivingClock = true;
 		RebuildSongMapEvaluator(ExternalClockDriver ? 
 									ExternalClockDriver->SongMapEvaluator->GetSongMapsWithTempoMap() : 
 									SongMapEvaluator->GetSongMapsWithOtherMaps(),
@@ -332,6 +337,20 @@ namespace HarmonixMetasound
 		int32 EndFrame = StartFrame + NumFrames;
 		const TArray<FMidiClockEvent>& ClockEvents = DrivingClock.GetMidiClockEventsInBlock();
 		int32 Index = Algo::LowerBoundBy(ClockEvents, StartFrame, &FMidiClockEvent::BlockFrameIndex);
+		if (ClockEvents.IsValidIndex(0) && NeedsSeekToDrivingClock)
+		{
+			NeedsSeekToDrivingClock = false;
+			int32 OurStartTick = DrivingClock.GetNextTickToProcessAtBlockFrame(StartFrame);
+			if (!FMath::IsNearlyEqual(CurrentLocalSpeed, 1.0f, 0.0001))
+			{
+				float FractionalToTick = (float)OurStartTick * CurrentLocalSpeed;
+				OurStartTick = FMath::FloorToInt32(FractionalToTick);
+				TickResidualWhenDriven = FMath::Fractional(FractionalToTick);
+			}
+			OurStartTick = WrapTickIfLooping(OurStartTick);
+			SeekTo(ClockEvents[0].BlockFrameIndex, OurStartTick);
+		}
+
 		while (ClockEvents.IsValidIndex(Index))
 		{
 			const FMidiClockEvent& Event = ClockEvents[Index];
@@ -577,7 +596,7 @@ namespace HarmonixMetasound
 		// we're going to have to look through the clock events...
 		for (const FMidiClockEvent& Event : MidiClockEventsInBlock)
 		{
-			if (Event.BlockFrameIndex >= BlockFrame)
+			if (Event.BlockFrameIndex > BlockFrame)
 			{
 				break;
 			}
@@ -589,7 +608,14 @@ namespace HarmonixMetasound
 			}
 			else if (const FAdvance* AsProcess = Event.TryGet<FAdvance>())
 			{
-				FoundNextTick = AsProcess->FirstTickToProcess + AsProcess->NumberOfTicksToProcess;
+				if (Event.BlockFrameIndex == BlockFrame)
+				{
+					FoundNextTick = AsProcess->FirstTickToProcess;
+				}
+				else
+				{
+					FoundNextTick = AsProcess->FirstTickToProcess + AsProcess->NumberOfTicksToProcess;
+				}
 				AtBlockIndex = Event.BlockFrameIndex;
 			}
 			else if (const FLoop* AsLoop = Event.TryGet<FLoop>())
@@ -632,7 +658,7 @@ namespace HarmonixMetasound
 		{
 			// Advance based on the delta ticks, and not based on the absolute tick
 			int32 UpToTick;
-			if (FMath::IsNearlyEqual(SpeedAtBlockEnd, 1.0f, 0.0001))
+			if (FMath::IsNearlyEqual(CurrentLocalSpeed, 1.0f, 0.0001))
 			{
 				UpToTick = NextMidiTickToProcess + AsAdvance->NumberOfTicksToProcess;
 			}
@@ -642,6 +668,7 @@ namespace HarmonixMetasound
 				UpToTick = FMath::FloorToInt32(FractionalToTick);
 				TickResidualWhenDriven = FMath::Fractional(FractionalToTick);
 			}
+			// No need to wrap the tick here because AdvanceToTick will handle that.
 			AdvanceToTick(Event.BlockFrameIndex, UpToTick);
 		}
 		else if (const FTempoChange* AsTempoChange = Event.TryGet<FTempoChange>())
diff --git a/Engine/Plugins/Runtime/Harmonix/Source/HarmonixMetasound/Public/HarmonixMetasound/DataTypes/MidiClock.h b/Engine/Plugins/Runtime/Harmonix/Source/HarmonixMetasound/Public/HarmonixMetasound/DataTypes/MidiClock.h
--- a/Engine/Plugins/Runtime/Harmonix/Source/HarmonixMetasound/Public/HarmonixMetasound/DataTypes/MidiClock.h
+++ b/Engine/Plugins/Runtime/Harmonix/Source/HarmonixMetasound/Public/HarmonixMetasound/DataTypes/MidiClock.h
@@ -222,6 +222,7 @@ namespace HarmonixMetasound
 		int32 LoopLengthTicks;
 
 		bool MidiDataChangedInBlock;
+		bool NeedsSeekToDrivingClock;
 
 		FMidiClockEvents MidiClockEventsInBlock;
 	};
