Tweaks so that the Metronome Node doesn't have to keep infinite history.

diff --git a/Engine/Plugins/Runtime/Harmonix/Source/HarmonixMetasound/Private/HarmonixMetasound/Nodes/MetronomeNode.cpp b/Engine/Plugins/Runtime/Harmonix/Source/HarmonixMetasound/Private/HarmonixMetasound/Nodes/MetronomeNode.cpp
--- a/Engine/Plugins/Runtime/Harmonix/Source/HarmonixMetasound/Private/HarmonixMetasound/Nodes/MetronomeNode.cpp
+++ b/Engine/Plugins/Runtime/Harmonix/Source/HarmonixMetasound/Private/HarmonixMetasound/Nodes/MetronomeNode.cpp
@@ -14,16 +14,35 @@
 #include <algorithm>
 
 #include "HAL/PlatformMath.h"
+#include "Harmonix.h"
 
 #define LOCTEXT_NAMESPACE "HarmonixMetaSound"
 
 DEFINE_LOG_CATEGORY_STATIC(LogMetronomeNode, Log, All);
 
-namespace HarmonixMetasound
+namespace HarmonixMetasound::Nodes::MetronomeNode
 {
 	using namespace Metasound;
 
-	class FMetronomeOperator : public TExecutableOperator<FMetronomeOperator>, public FMusicTransportControllable
+	namespace Inputs
+	{
+		DEFINE_METASOUND_PARAM_ALIAS(Transport, CommonPinNames::Inputs::Transport);
+		DEFINE_METASOUND_PARAM_ALIAS(Loop, CommonPinNames::Inputs::Loop);
+		DEFINE_METASOUND_PARAM_ALIAS(LoopLengthBars, CommonPinNames::Inputs::LoopLengthBars);
+		DEFINE_METASOUND_PARAM_ALIAS(TimeSigNumerator, CommonPinNames::Inputs::TimeSigNumerator);
+		DEFINE_METASOUND_PARAM_ALIAS(TimeSigDenominator, CommonPinNames::Inputs::TimeSigDenominator);
+		DEFINE_METASOUND_PARAM_ALIAS(Tempo, CommonPinNames::Inputs::Tempo);
+		DEFINE_METASOUND_PARAM_ALIAS(Speed, CommonPinNames::Inputs::Speed);
+		DEFINE_METASOUND_PARAM_ALIAS(PrerollBars, CommonPinNames::Inputs::PrerollBars);
+		DEFINE_INPUT_METASOUND_PARAM(HistoryLengthSeconds, "History Length Secs.", "Amount of accurate tempo change history to maintain. Any game system queries earlier than this number of second behind the current render time will be estimates and prone to jitter!");
+	}
+
+	namespace Outputs
+	{
+		DEFINE_METASOUND_PARAM_ALIAS(MidiClock, CommonPinNames::Outputs::MidiClock);
+	}
+
+	class FMetronomeOperator final : public TExecutableOperator<FMetronomeOperator>, public FMusicTransportControllable
 	{
 	public:
 		static const FNodeClassMetadata& GetNodeInfo();
@@ -38,7 +57,8 @@ namespace HarmonixMetasound
 						   const FInt32ReadRef& InTimeSigDenominator,
 						   const FFloatReadRef& InTempo,
 		                   const FFloatReadRef& InSpeedMultiplier,
-		                   const int32 InSeekPrerollBars);
+		                   const int32 InSeekPrerollBars,
+					       const FFloatReadRef& InHistoryLength);
 
 		virtual void BindInputs(FInputVertexInterfaceData& InVertexData) override;
 		virtual void BindOutputs(FOutputVertexInterfaceData& InVertexData) override;
@@ -58,10 +78,13 @@ namespace HarmonixMetasound
 		FFloatReadRef TempoInPin;
 		FFloatReadRef SpeedMultInPin;
 		const int32 SeekPreRollBarsInPin;
+		FFloatReadRef HistoryLengthSecondsInPin;
 
 		//** OUTPUTS
 		FMidiClockWriteRef MidiClockOutPin;
 
+		static constexpr float kExtraHistoryForCalibrationMs = 100.0f;
+
 		//** DATA
 		TSharedPtr<FMidiClock, ESPMode::NotThreadSafe> MonotonicallyIncreasingClock;
 		TSharedPtr<FSongMaps> SongMaps;
@@ -121,8 +144,6 @@ namespace HarmonixMetasound
 
 	const FVertexInterface& FMetronomeOperator::GetVertexInterface()
 	{
-		using namespace CommonPinNames;
-
 		static const FVertexInterface Interface(
 			FInputVertexInterface(
 				TInputDataVertex<FMusicTransportEventStream>(METASOUND_GET_PARAM_NAME_AND_METADATA(Inputs::Transport)),
@@ -132,8 +153,9 @@ namespace HarmonixMetasound
 				TInputDataVertex<int32>(METASOUND_GET_PARAM_NAME_AND_METADATA(Inputs::TimeSigDenominator), 4),
 				TInputDataVertex<float>(METASOUND_GET_PARAM_NAME_AND_METADATA(Inputs::Tempo), 120.0f),
 				TInputDataVertex<float>(METASOUND_GET_PARAM_NAME_AND_METADATA(Inputs::Speed), 1.0f),
-				TInputConstructorVertex<int32>(METASOUND_GET_PARAM_NAME_AND_METADATA(Inputs::PrerollBars), 8)
-			),
+				TInputConstructorVertex<int32>(METASOUND_GET_PARAM_NAME_AND_METADATA(Inputs::PrerollBars), 8),
+				TInputDataVertex<float>(METASOUND_GET_PARAM_NAME_AND_METADATA(Inputs::HistoryLengthSeconds), 0.5f)
+				),
 			FOutputVertexInterface(
 				TOutputDataVertex<FMidiClock>(METASOUND_GET_PARAM_NAME_AND_METADATA(Outputs::MidiClock))
 			)
@@ -144,8 +166,6 @@ namespace HarmonixMetasound
 
 	TUniquePtr<IOperator> FMetronomeOperator::CreateOperator(const FBuildOperatorParams& InParams, FBuildResults& OutResults)
 	{
-		using namespace CommonPinNames;
-
 		const FMetronomeNode& TempoClockNode = static_cast<const FMetronomeNode&>(InParams.Node);
 
 		const FInputVertexInterfaceData& InputData = InParams.InputData;
@@ -159,8 +179,9 @@ namespace HarmonixMetasound
 		FFloatReadRef InTempo = InputData.GetOrCreateDefaultDataReadReference<float>(METASOUND_GET_PARAM_NAME(Inputs::Tempo), Settings);
 		FFloatReadRef InSpeed = InputData.GetOrCreateDefaultDataReadReference<float>(METASOUND_GET_PARAM_NAME(Inputs::Speed), Settings);
 		int32 InPreRollBars = InputData.GetOrCreateDefaultValue<int32>(METASOUND_GET_PARAM_NAME(Inputs::PrerollBars), Settings);
+		FFloatReadRef InHistorySecs = InputData.GetOrCreateDefaultDataReadReference<float>(METASOUND_GET_PARAM_NAME(Inputs::HistoryLengthSeconds), Settings);
 
-		return MakeUnique<FMetronomeOperator>(InParams, InTransport, InLoop, InLoopLengthBars, InTimeSigNumerator, InTimeSigDenominator, InTempo, InSpeed, InPreRollBars);
+		return MakeUnique<FMetronomeOperator>(InParams, InTransport, InLoop, InLoopLengthBars, InTimeSigNumerator, InTimeSigDenominator, InTempo, InSpeed, InPreRollBars, InHistorySecs);
 	}
 
 	FMetronomeOperator::FMetronomeOperator(const FBuildOperatorParams& InParams, 
@@ -171,7 +192,8 @@ namespace HarmonixMetasound
 	                                       const FInt32ReadRef& InTimeSigDenominator,
 	                                       const FFloatReadRef& InTempo,
 	                                       const FFloatReadRef& InSpeedMultiplier,
-		                                   const int32 InPreRollBars)
+		                                   const int32 InPreRollBars,
+										   const FFloatReadRef& InHistoryLength)
 		: FMusicTransportControllable(EMusicPlayerTransportState::Prepared) 
 		, TransportInPin(InTransport)
 		, LoopInPin(InLoop)
@@ -181,6 +203,7 @@ namespace HarmonixMetasound
 		, TempoInPin(InTempo)
 		, SpeedMultInPin(InSpeedMultiplier)
 		, SeekPreRollBarsInPin(InPreRollBars)
+		, HistoryLengthSecondsInPin(InHistoryLength)
 		, MidiClockOutPin(FMidiClockWriteRef::CreateNew(InParams.OperatorSettings))
 		, MonotonicallyIncreasingClock(MakeShared<FMidiClock, ESPMode::NotThreadSafe>(InParams.OperatorSettings))
 		, BlockSize(InParams.OperatorSettings.GetNumFramesPerBlock())
@@ -195,7 +218,6 @@ namespace HarmonixMetasound
 
 	void FMetronomeOperator::BindInputs(FInputVertexInterfaceData& InVertexData)
 	{
-		using namespace CommonPinNames;
 		InVertexData.BindReadVertex(METASOUND_GET_PARAM_NAME(Inputs::Transport), TransportInPin);
 		InVertexData.SetValue(METASOUND_GET_PARAM_NAME(Inputs::Loop), LoopInPin);
 		InVertexData.SetValue(METASOUND_GET_PARAM_NAME(Inputs::LoopLengthBars), LoopLengthBarsInPin);
@@ -204,13 +226,13 @@ namespace HarmonixMetasound
 		InVertexData.BindReadVertex(METASOUND_GET_PARAM_NAME(Inputs::Tempo), TempoInPin);
 		InVertexData.BindReadVertex(METASOUND_GET_PARAM_NAME(Inputs::Speed), SpeedMultInPin);
 		InVertexData.SetValue(METASOUND_GET_PARAM_NAME(Inputs::PrerollBars), SeekPreRollBarsInPin);
+		InVertexData.BindReadVertex(METASOUND_GET_PARAM_NAME(Inputs::HistoryLengthSeconds), HistoryLengthSecondsInPin);
 
 		Init();
 	}
 
 	void FMetronomeOperator::BindOutputs(FOutputVertexInterfaceData& InVertexData)
 	{
-		using namespace CommonPinNames;
 		InVertexData.BindReadVertex(METASOUND_GET_PARAM_NAME(Outputs::MidiClock), MidiClockOutPin);
 	}
 
@@ -399,7 +421,12 @@ namespace HarmonixMetasound
 	{
 		CurrentTempo = InTempoBPM;
 		int32 AtTick = GetDrivingMidiClock().GetNextMidiTickToProcess();
-		SongMaps->AddTempoChange(AtTick, CurrentTempo);
+		
+		float HistoryLengthSeconds = FMath::Max(*HistoryLengthSecondsInPin, (FHarmonixModule::GetMeasuredUserExperienceAndReactionToAudioRenderOffsetMs() + kExtraHistoryForCalibrationMs) / 1000.0f);
+		HistoryLengthSeconds = FMath::Max(HistoryLengthSeconds, (FHarmonixModule::GetMeasuredVideoToAudioRenderOffsetMs() + kExtraHistoryForCalibrationMs) / 1000.0f);
+		
+		FTempoMap& TempoMap = SongMaps->GetTempoMap();
+		TempoMap.AddTempo(InTempoBPM, AtTick, HistoryLengthSeconds);
 	}
 
 	void FMetronomeOperator::AddTimeSigChangeForMidi(int32 InTimeSigNum, int32 InTimeSigDenom)
diff --git a/Engine/Plugins/Runtime/Harmonix/Source/HarmonixMidi/Private/HarmonixMidi/TempoMap.cpp b/Engine/Plugins/Runtime/Harmonix/Source/HarmonixMidi/Private/HarmonixMidi/TempoMap.cpp
--- a/Engine/Plugins/Runtime/Harmonix/Source/HarmonixMidi/Private/HarmonixMidi/TempoMap.cpp
+++ b/Engine/Plugins/Runtime/Harmonix/Source/HarmonixMidi/Private/HarmonixMidi/TempoMap.cpp
@@ -57,6 +57,7 @@ bool FTempoMap::operator==(const FTempoMap& Other) const
 void FTempoMap::Empty()
 {
 	Points.Empty();
+	bEarlyMapIsEstimate = false;
 }
 
 void FTempoMap::Copy(const FTempoMap& Other, int32 StartTick, int32 EndTick)
@@ -212,10 +213,9 @@ bool FTempoMap::AddTempoInfoPoint(int32 Tempo, int32 Tick, bool SortNow)
 {
 	if (Points.Num() == 0)
 	{
-		ensureMsgf(Tick == 0, TEXT("FTempoMap::AddTempoInfoPoint(): tried to add point (%d, %d) to an empty tempo map! First tempo point must be at tick 0."), Tick, Tempo);
-		if (Tick != 0)
+		if (!ensureMsgf(Tick == 0, TEXT("FTempoMap::AddTempoInfoPoint(): tried to add point (%d, %d) to an empty tempo map! First tempo point must be at tick 0. Adding at tick 0."), Tick, Tempo))
 		{
-			return false;
+			Tick = 0;
 		}
 	}
 	else
@@ -227,10 +227,126 @@ bool FTempoMap::AddTempoInfoPoint(int32 Tempo, int32 Tick, bool SortNow)
 			return false;
 		}
 	}
+
 	FMusicMapUtl::AddPoint<FTempoInfoPoint,float,int32>(TickToMs(Tick), Tempo, Points, Tick, SortNow);
 	return true;
 }
 
+void FTempoMap::AddTempo(float Bpm, int32 Tick, int32 MaxLength)
+{
+	float MsToNewPoint = TickToMs(Tick);
+
+	int32 MidiTempo = Harmonix::Midi::Constants::BPMToMidiTempo(Bpm);
+
+	if (Points.IsEmpty())
+	{
+		Tick = 0;
+	}
+
+	if (Points.Max() < MaxLength && MaxLength < std::numeric_limits<int32>::max())
+	{
+		Points.Reserve(MaxLength);
+	}
+
+	if (Points.Num() < MaxLength)
+	{
+		AddTempoInfoPoint(MidiTempo, Tick);
+	}
+	else
+	{
+		check(TicksPerQuarterNote > 0);
+
+		bEarlyMapIsEstimate = true;
+
+		if (MaxLength < 3)
+		{
+			// If HistoryLengthEntries == 2 we don't do any "shifting" of the history.
+			// We just need to update the first and last entry...
+			float QuartersToHere = (float)Tick / (float)TicksPerQuarterNote;
+			float UsPerQuarterToHere = (MsToNewPoint * 1000.0f) / QuartersToHere;
+			Points[0].MidiTempo = FMath::RoundToInt32(UsPerQuarterToHere);
+			Points[1].Ms = MsToNewPoint;
+			Points[1].StartTick = Tick;
+			Points[1].MidiTempo = MidiTempo;
+		}
+		else
+		{
+			// We have to shift the history...
+			// It is possible we now want less history than we already have.
+			int32 FirstEntryToShift = 2;
+			if (Points.Num() > MaxLength)
+			{
+				FirstEntryToShift += Points.Num() - MaxLength;
+			}
+			ShiftEntriesAndFixUpPoints(FirstEntryToShift, MsToNewPoint, Tick, MidiTempo);
+		}
+	}
+}
+
+void FTempoMap::AddTempo(float Bpm, int32 Tick, float MaxLengthSecs)
+{
+	check(TicksPerQuarterNote > 0);
+
+	int32 MidiTempo = Harmonix::Midi::Constants::BPMToMidiTempo(Bpm);
+
+	if (Points.IsEmpty())
+	{
+		Tick = 0;
+	}
+
+	if (Points.Num() < 3)
+	{
+		AddTempoInfoPoint(MidiTempo, Tick);
+		return;
+	}
+
+	float MsToNewPoint = TickToMs(Tick);
+	float SecondsOfHistory = (MsToNewPoint - Points[1].Ms) / 1000.0f;
+	if (SecondsOfHistory < MaxLengthSecs)
+	{
+		AddTempoInfoPoint(MidiTempo, Tick);
+		return;
+	}
+
+	int32 NewHistoryStartIndex = 2;
+	SecondsOfHistory = (MsToNewPoint - Points[NewHistoryStartIndex].Ms) / 1000.0f;
+	while (NewHistoryStartIndex < (Points.Num() - 1) && SecondsOfHistory > MaxLengthSecs)
+	{
+		NewHistoryStartIndex++;
+		SecondsOfHistory = (MsToNewPoint - Points[NewHistoryStartIndex].Ms) / 1000.0f;
+	}
+	NewHistoryStartIndex--;
+	if (NewHistoryStartIndex == 1)
+	{
+		AddTempoInfoPoint(MidiTempo, Tick);
+		return;
+	}
+
+	ShiftEntriesAndFixUpPoints(NewHistoryStartIndex, MsToNewPoint, Tick, MidiTempo);
+}
+
+
+void FTempoMap::ShiftEntriesAndFixUpPoints(int32 NewHistoryStartIndex, float MsToNewPoint, int32 TickOfNewPoint, int32 MidiTempoOfNewPoint)
+{
+	check(TicksPerQuarterNote > 0);
+
+	int32 NumToMove = Points.Num() - NewHistoryStartIndex;
+	FMemory::Memmove(&Points[1], &Points[NewHistoryStartIndex], sizeof(FTempoInfoPoint) * NumToMove);
+	Points.SetNum(NumToMove + 2, EAllowShrinking::Yes); // +2 because... point 0 (root point) + last point (new point)
+
+	// Now update point 0...
+	float QuartersToFirstHistory = (float)Points[1].StartTick / (float)TicksPerQuarterNote;
+	float UsPerQuarterToHistory = (Points[1].Ms * 1000.0f) / QuartersToFirstHistory;
+	Points[0].MidiTempo = FMath::RoundToInt32(UsPerQuarterToHistory);
+
+	// Update the last entry...
+	Points.Last().Ms = MsToNewPoint;
+	Points.Last().StartTick = TickOfNewPoint;
+	Points.Last().MidiTempo = MidiTempoOfNewPoint;
+
+	bEarlyMapIsEstimate = true;
+}
+
 void FTempoMap::WipeTempoInfoPoints(const int32 Tick)
 {
 	FMusicMapUtl::RemovePointsOnAndAfterTick(Points, Tick);
diff --git a/Engine/Plugins/Runtime/Harmonix/Source/HarmonixMidi/Public/HarmonixMidi/TempoMap.h b/Engine/Plugins/Runtime/Harmonix/Source/HarmonixMidi/Public/HarmonixMidi/TempoMap.h
--- a/Engine/Plugins/Runtime/Harmonix/Source/HarmonixMidi/Public/HarmonixMidi/TempoMap.h
+++ b/Engine/Plugins/Runtime/Harmonix/Source/HarmonixMidi/Public/HarmonixMidi/TempoMap.h
@@ -138,6 +138,20 @@ public:
 
 	bool AddTempoInfoPoint(int32 MicrosecondsPerQuarterNote, int32 Tick, bool SortNow = true);
 
+	/** Use this to add tempo changes to a dynamic map where you want to hard limit the 
+	    number of history entries to a fixed number. Any attempts to get a musical timing
+		(ie FMidiSongPos) for a position earlier in the music will be an estimate and is
+		prone to jitter!
+		*/
+	void AddTempo(float Bpm, int32 Tick, int32 MaxLength = std::numeric_limits<int32>::max());
+
+	/** Use this to add tempo changes to a dynamic map where you want to limit the number
+		of history entries to a specific amount of time. Any attempts to get a musical timing
+		(ie FMidiSongPos) for a position earlier in the music will be an estimate and is
+		prone to jitter!
+		*/
+	void AddTempo(float Bpm, int32 Tick, float MaxLengthSecs);
+
 	/** Remove every point from the given tick onward */
 	void WipeTempoInfoPoints(int32 Tick);
 
@@ -155,6 +169,8 @@ protected:
 	UPROPERTY()
 	TArray<FTempoInfoPoint> Points;
 
+	bool bEarlyMapIsEstimate = false;
+
 private:
 
 	// Given a time, find the TempoInfoPoint at or before it.
@@ -164,5 +180,6 @@ private:
 	// calculate tick as time in ms
 	float TickToMsInternal(float Tick, const FTempoInfoPoint& prevTempoInfoPoint) const;
 	float MsToTickInternal(float TimeMs, const FTempoInfoPoint& prevTempoInfoPoint) const;
-
+	
+	void  ShiftEntriesAndFixUpPoints(int32 NewHistoryStartIndex, float MsToNewPoint, int32 TickOfNewPoint, int32 MidiTempoOfNewPoint);
 };
