Add Trigger Compare (MusicTimestamp) Node using template in metasound code lib. Required a tweak to a Metasound header that was approved by [REDACTED_PERSON]

diff --git a/Engine/Plugins/Runtime/Harmonix/Source/HarmonixMetasound/Private/HarmonixMetasound/Nodes/MusicTimestampBuilder.cpp b/Engine/Plugins/Runtime/Harmonix/Source/HarmonixMetasound/Private/HarmonixMetasound/Nodes/MusicTimestampBuilder.cpp
--- a/Engine/Plugins/Runtime/Harmonix/Source/HarmonixMetasound/Private/HarmonixMetasound/Nodes/MusicTimestampBuilder.cpp
+++ b/Engine/Plugins/Runtime/Harmonix/Source/HarmonixMetasound/Private/HarmonixMetasound/Nodes/MusicTimestampBuilder.cpp
@@ -7,6 +7,7 @@
 #include "MetasoundSampleCounter.h"
 #include "MetasoundStandardNodesCategories.h"
 #include "MetasoundVertex.h"
+#include "MetasoundTriggerCompareNode.h"
 
 #include "HarmonixMetasound/Common.h"
 #include "HarmonixMetasound/DataTypes/MusicTransport.h"
@@ -18,6 +19,9 @@ namespace HarmonixMetasound
 {
 	using namespace Metasound;
 
+	using FTriggerCompareNodeMusicTimestamp = TTriggerCompareNode<FMusicTimestamp>;
+	METASOUND_REGISTER_NODE(FTriggerCompareNodeMusicTimestamp)
+
 	class FMusicTimestampBuilderOperator : public TExecutableOperator<FMusicTimestampBuilderOperator>
 	{
 	public:
diff --git a/Engine/Plugins/Runtime/Harmonix/Source/HarmonixMidi/Public/HarmonixMidi/BarMap.h b/Engine/Plugins/Runtime/Harmonix/Source/HarmonixMidi/Public/HarmonixMidi/BarMap.h
--- a/Engine/Plugins/Runtime/Harmonix/Source/HarmonixMidi/Public/HarmonixMidi/BarMap.h
+++ b/Engine/Plugins/Runtime/Harmonix/Source/HarmonixMidi/Public/HarmonixMidi/BarMap.h
@@ -38,6 +38,12 @@ public:
 		Beat = 1.0f;
 	}
 
+	bool operator>=(const FMusicTimestamp& Other) const
+	{
+		return Bar > Other.Bar ||
+			(Bar == Other.Bar && Beat >= Other.Beat);
+	}
+
 	bool operator>(const FMusicTimestamp& Other) const
 	{
 		return Bar > Other.Bar ||
@@ -59,7 +65,13 @@ public:
 		return Bar < Other.Bar ||
 			(Bar == Other.Bar && Beat < Other.Beat);
 	}
-	
+
+	bool operator<=(const FMusicTimestamp& Other) const
+	{
+		return Bar < Other.Bar ||
+			(Bar == Other.Bar && Beat <= Other.Beat);
+	}
+
 	friend FORCEINLINE uint32 GetTypeHash(const FMusicTimestamp& InTimestamp)
 	{
 		return HashCombineFast(GetTypeHash(InTimestamp.Bar), GetTypeHash(InTimestamp.Beat));
diff --git a/Engine/Plugins/Runtime/Metasound/Source/MetasoundStandardNodes/Public/MetasoundTriggerCompareNode.h b/Engine/Plugins/Runtime/Metasound/Source/MetasoundStandardNodes/Public/MetasoundTriggerCompareNode.h
--- a/Engine/Plugins/Runtime/Metasound/Source/MetasoundStandardNodes/Public/MetasoundTriggerCompareNode.h
+++ b/Engine/Plugins/Runtime/Metasound/Source/MetasoundStandardNodes/Public/MetasoundTriggerCompareNode.h
@@ -230,7 +230,7 @@ namespace Metasound
 	 *  Compares two inputs against enumerated comparison types.
 	 */
 	template<typename ValueType>
-	class METASOUNDSTANDARDNODES_API TTriggerCompareNode : public FNodeFacade
+	class TTriggerCompareNode : public FNodeFacade
 	{
 	public:
 		/**
