Fix bug where a looping metronome that has its tempo change doesn't properly notify its midi clock output that the song maps have changed. This caused the MusicClockComponent to NOT report tempo changes.

diff --git a/Engine/Plugins/Runtime/Harmonix/Source/HarmonixMetasound/Private/HarmonixMetasound/Nodes/MetronomeNode.cpp b/Engine/Plugins/Runtime/Harmonix/Source/HarmonixMetasound/Private/HarmonixMetasound/Nodes/MetronomeNode.cpp
--- a/Engine/Plugins/Runtime/Harmonix/Source/HarmonixMetasound/Private/HarmonixMetasound/Nodes/MetronomeNode.cpp
+++ b/Engine/Plugins/Runtime/Harmonix/Source/HarmonixMetasound/Private/HarmonixMetasound/Nodes/MetronomeNode.cpp
@@ -391,7 +391,15 @@ namespace HarmonixMetasound
 
 		if (HasMidiChanges)
 		{
-			DrivingMidiClock.SongMapsChanged();
+			// We only have to tell the MonotonicallyIncreasingClock that the song maps have changed
+			// if we are looping...
+			if (LoopInPin)
+			{
+				MonotonicallyIncreasingClock->SongMapsChanged();
+			}
+			// Regardless of which clock is currently driving, MonotonicallyIncreasingClock or MidiClockOutPin,
+			// we need to inform the MidiClockOutPin that the song maps have changed. 
+			MidiClockOutPin->SongMapsChanged();
 		}
 	}
 
