- New Test for the SecsUntilMusicTimestamp Node.
- A few fixes / tweaks to the SecsUntilMusicTimestamp Node now that there are tests.

diff --git a/Engine/Plugins/Runtime/Harmonix/Source/HarmonixMetasound/Private/HarmonixMetasound/Nodes/SecsUntilMusicTimestampNode.cpp b/Engine/Plugins/Runtime/Harmonix/Source/HarmonixMetasound/Private/HarmonixMetasound/Nodes/SecsUntilMusicTimestampNode.cpp
--- a/Engine/Plugins/Runtime/Harmonix/Source/HarmonixMetasound/Private/HarmonixMetasound/Nodes/SecsUntilMusicTimestampNode.cpp
+++ b/Engine/Plugins/Runtime/Harmonix/Source/HarmonixMetasound/Private/HarmonixMetasound/Nodes/SecsUntilMusicTimestampNode.cpp
@@ -1,6 +1,6 @@
 // Copyright Epic Games, Inc. All Rights Reserved.
 
-#include "HarmonixMetasound/Nodes/MidiCCTriggerNode.h"
+#include "HarmonixMetasound/Nodes/SecsUntilMusicTimestampNode.h"
 
 #include "MetasoundExecutableOperator.h"
 #include "MetasoundFacade.h"
@@ -48,7 +48,7 @@ namespace HarmonixMetasound::Nodes::SecsUntilMusicTimestampNode
 
 	namespace Outputs
 	{
-		METASOUND_PARAM(SecsUntilTimestamp, "Secs. Until Timestamp", "Looks at the MIDI clock and calculates how long it will be until the timestamp triggers. It DOES consider the speed of the clock as well!");
+		DEFINE_OUTPUT_METASOUND_PARAM(SecsUntilTimestamp, "Secs. Until Timestamp", "Looks at the MIDI clock and calculates how long it will be until the timestamp triggers. It DOES consider the speed of the clock as well!");
 	}
 
 	class FOp final : public TExecutableOperator<FOp>
@@ -119,6 +119,8 @@ namespace HarmonixMetasound::Nodes::SecsUntilMusicTimestampNode
 			const FOperatorSettings& OperatorSettings = InParams.OperatorSettings;
 			const FInputVertexInterfaceData& InputData = InParams.InputData;
 
+			bool bClockIsBound = InputData.IsVertexBound(Inputs::MidiClockName);
+
 			FInputs Inputs
 			{
 				InputData.GetOrCreateDefaultDataReadReference<bool>(Inputs::EnableName,OperatorSettings),
@@ -128,24 +130,31 @@ namespace HarmonixMetasound::Nodes::SecsUntilMusicTimestampNode
 
 			FOutputs Outputs
 			{
-				FFloatWriteRef::CreateNew(0.0f)
+				FFloatWriteRef::CreateNew(Never)
 			};
 
-			return MakeUnique<FOp>(InParams, MoveTemp(Inputs), MoveTemp(Outputs));
+			return MakeUnique<FOp>(InParams, MoveTemp(Inputs), bClockIsBound, MoveTemp(Outputs));
 		}
 
-		FOp(const FBuildOperatorParams& Params, FInputs&& InInputs, FOutputs&& InOutputs)
+		FOp(const FBuildOperatorParams& Params, FInputs&& InInputs, bool bClockIsBoundIn,  FOutputs&& InOutputs)
 			: Inputs(MoveTemp(InInputs))
 			, Outputs(MoveTemp(InOutputs))
+			, bClockIsBound(bClockIsBoundIn)
 		{
 			Reset(Params);
 		}
 
 		virtual void BindInputs(FInputVertexInterfaceData& InVertexData) override
 		{
+			// See if the clock input is connected to anything...
+			bClockIsBound = InVertexData.IsVertexBound(Inputs::MidiClockName);
+
 			InVertexData.BindReadVertex(Inputs::EnableName, Inputs.Enable);
 			InVertexData.BindReadVertex(Inputs::MidiClockName, Inputs.MidiClock);
 			InVertexData.BindReadVertex(Inputs::TimestampName, Inputs.Timestamp);
+			
+			// Force a recompute of the time...
+			MsOfTimestamp = Never;
 		}
 
 		virtual void BindOutputs(FOutputVertexInterfaceData& InVertexData) override
@@ -160,21 +169,27 @@ namespace HarmonixMetasound::Nodes::SecsUntilMusicTimestampNode
 
 		void Execute()
 		{
-			if (!*Inputs.Enable)
+			if (!*Inputs.Enable || !bClockIsBound)
 			{
-				*Outputs.SecsUntilTimestamp = std::numeric_limits<float>::max();
+				*Outputs.SecsUntilTimestamp = Never;
 				return;
 			}
 
 			const ISongMapEvaluator& SongMaps = Inputs.MidiClock->GetSongMapEvaluator();
 
-			if (*Inputs.Timestamp != CurrentTimestamp || Inputs.MidiClock->GetSongMapsChangedInBlock())
+			if (MsOfTimestamp == Never || *Inputs.Timestamp != CurrentTimestamp || Inputs.MidiClock->GetSongMapsChangedInBlock())
 			{
 				int32 Tick = SongMaps.MusicTimestampToTick(*Inputs.Timestamp);
 				MsOfTimestamp = SongMaps.TickToMs((float)Tick);
 			}
 
 			float Speed = Inputs.MidiClock->GetSpeedAtStartOfBlock();
+			if (FMath::IsNearlyZero(Speed))
+			{
+				*Outputs.SecsUntilTimestamp = Never;
+				return;
+			}
+
 			const int32 SongTick = Inputs.MidiClock->GetNextMidiTickToProcess();
 			float SongPosMs = SongMaps.TickToMs((float)SongTick);
 
@@ -182,11 +197,13 @@ namespace HarmonixMetasound::Nodes::SecsUntilMusicTimestampNode
 		}
 
 	private:
+		static constexpr float Never = std::numeric_limits<float>::max();
+
 		FInputs Inputs;
 		FOutputs Outputs;
-
+		bool bClockIsBound = false;
 		FMusicTimestamp CurrentTimestamp;
-		float MsOfTimestamp = 0.0f;
+		float MsOfTimestamp = Never;
 	};
 
 	class FSecsUntilMusicTimestampNode final : public FNodeFacade
diff --git a/Engine/Plugins/Runtime/Harmonix/Source/HarmonixMetasound/Public/HarmonixMetasound/Nodes/SecsUntilMusicTimestampNode.h b/Engine/Plugins/Runtime/Harmonix/Source/HarmonixMetasound/Public/HarmonixMetasound/Nodes/SecsUntilMusicTimestampNode.h
new file mode 100644
--- /dev/null
+++ b/Engine/Plugins/Runtime/Harmonix/Source/HarmonixMetasound/Public/HarmonixMetasound/Nodes/SecsUntilMusicTimestampNode.h
@@ -0,0 +1,25 @@
+// Copyright Epic Games, Inc. All Rights Reserved.
+
+#pragma once
+
+#include "MetasoundNodeInterface.h"
+
+#include "HarmonixMetasound/Common.h"
+
+namespace HarmonixMetasound::Nodes::SecsUntilMusicTimestampNode
+{
+	const HARMONIXMETASOUND_API Metasound::FNodeClassName& GetClassName();
+	HARMONIXMETASOUND_API int32 GetCurrentMajorVersion();
+
+	namespace Inputs
+	{
+		DECLARE_METASOUND_PARAM_ALIAS(Enable);
+		DECLARE_METASOUND_PARAM_ALIAS(MidiClock);
+		DECLARE_METASOUND_PARAM_ALIAS(Timestamp);
+	}
+
+	namespace Outputs
+	{
+		DECLARE_METASOUND_PARAM_EXTERN(SecsUntilTimestamp)
+	}
+}
\ No newline at end of file
diff --git a/Engine/Plugins/Runtime/Harmonix/Source/HarmonixMetasoundTests/Private/Nodes/TimestampRelatedNodeTests.cpp b/Engine/Plugins/Runtime/Harmonix/Source/HarmonixMetasoundTests/Private/Nodes/TimestampRelatedNodeTests.cpp
new file mode 100644
--- /dev/null
+++ b/Engine/Plugins/Runtime/Harmonix/Source/HarmonixMetasoundTests/Private/Nodes/TimestampRelatedNodeTests.cpp
@@ -0,0 +1,136 @@
+// Copyright Epic Games, Inc. All Rights Reserved.
+
+#include "NodeTestGraphBuilder.h"
+#include "HarmonixDsp/AudioBuffer.h"
+#include "HarmonixMetasound/Common.h"
+#include "HarmonixMetasound/DataTypes/MidiStream.h"
+#include "HarmonixMetasound/DataTypes/MusicTimestamp.h"
+#include "HarmonixMetasound/Nodes/SecsUntilMusicTimestampNode.h"
+#include "Misc/AutomationTest.h"
+
+#if WITH_DEV_AUTOMATION_TESTS
+
+namespace HarmonixMetasoundTests::TimestampRelated
+{
+	using FGraphBuilder = Metasound::Test::FNodeTestGraphBuilder;
+	using namespace Metasound;
+	using namespace Metasound::Frontend;
+	using namespace HarmonixMetasound;
+	using namespace HarmonixMetasound::Nodes;
+
+	IMPLEMENT_SIMPLE_AUTOMATION_TEST(
+		FSecsUntilTimestampTestDefaults,
+		"Harmonix.Metasound.Nodes.SecsUntilMusicTimestamp.Defaults",
+		EAutomationTestFlags::EditorContext | EAutomationTestFlags::EngineFilter)
+		bool FSecsUntilTimestampTestDefaults::RunTest(const FString&)
+	{
+		// Build the graph.
+		constexpr int32 NumSamplesPerBlock = 256;
+		const TUniquePtr<FMetasoundGenerator> Generator = FGraphBuilder::MakeSingleNodeGraph(
+			SecsUntilMusicTimestampNode::GetClassName(),
+			SecsUntilMusicTimestampNode::GetCurrentMajorVersion(),
+			48000,
+			NumSamplesPerBlock);
+		UTEST_TRUE("Graph successfully built", Generator.IsValid());
+
+		// execute a block
+		{
+			TAudioBuffer<float> Buffer{ Generator->GetNumChannels(), NumSamplesPerBlock, EAudioBufferCleanupMode::Delete};
+			Generator->OnGenerateAudio(Buffer.GetRawChannelData(0), Buffer.GetNumTotalValidSamples());
+		}
+
+		// Validate defaults.
+		TOptional<FFloatReadRef> OutputSeconds = Generator->GetOutputReadReference<float>(SecsUntilMusicTimestampNode::Outputs::SecsUntilTimestampName);
+		UTEST_TRUE("Output exists", OutputSeconds.IsSet());
+		UTEST_EQUAL("Secs. while enabled but no clock connected.", *(*OutputSeconds) , std::numeric_limits<float>::max());
+
+		return true;
+	}
+
+	IMPLEMENT_SIMPLE_AUTOMATION_TEST(
+		FSecsUntilTimestampTestBasic,
+		"Harmonix.Metasound.Nodes.SecsUntilMusicTimestamp.Basic",
+		EAutomationTestFlags::EditorContext | EAutomationTestFlags::EngineFilter)
+		bool FSecsUntilTimestampTestBasic::RunTest(const FString&)
+	{
+		FGraphBuilder Builder;
+
+		constexpr int32 TimeSigNumerator = 6;
+		constexpr int32 TimeSigDenominator = 8;
+		constexpr float Tempo = 60.0;
+		constexpr int32 Bar = 3;
+		constexpr float Beat = 3.0f;
+
+		const FNodeHandle TransportNode = Builder.AddNode({ HarmonixNodeNamespace, TEXT("TriggerToTransport"), TEXT("") },
+			0);
+
+		const FNodeHandle MetronomeNode = Builder.AddNode({HarmonixMetasound::HarmonixNodeNamespace, "Metronome", "" },
+			0	);
+
+		const FNodeHandle SecsUntilNode = Builder.AddNode(Nodes::SecsUntilMusicTimestampNode::GetClassName(),
+			Nodes::SecsUntilMusicTimestampNode::GetCurrentMajorVersion());
+
+		FNodeHandle StartInTrigger = Builder.AddAndConnectDataReferenceInput(TransportNode, CommonPinNames::Inputs::TransportPlayName, Metasound::GetMetasoundDataTypeName<FTrigger>(), "StartTest");
+		FNodeHandle TempoIn = Builder.AddAndConnectDataReferenceInput(MetronomeNode, CommonPinNames::Inputs::TempoName, Metasound::GetMetasoundDataTypeName<float>(), "TestTempo");
+		FNodeHandle TimeSigNumeratorIn = Builder.AddAndConnectDataReferenceInput(MetronomeNode, CommonPinNames::Inputs::TimeSigNumeratorName, Metasound::GetMetasoundDataTypeName<int32>(), "TestTSNum");
+		FNodeHandle TimeSigDenominatorIn = Builder.AddAndConnectDataReferenceInput(MetronomeNode, CommonPinNames::Inputs::TimeSigDenominatorName, Metasound::GetMetasoundDataTypeName<int32>(), "TestTSDenom");
+		FNodeHandle TimestampIn = Builder.AddAndConnectDataReferenceInput(SecsUntilNode, SecsUntilMusicTimestampNode::Inputs::TimestampName, Metasound::GetMetasoundDataTypeName<FMusicTimestamp>(), "TestTimestamp");
+		FNodeHandle SecsOut = Builder.AddAndConnectDataReferenceOutput(SecsUntilNode, SecsUntilMusicTimestampNode::Outputs::SecsUntilTimestampName, Metasound::GetMetasoundDataTypeName<float>(), "TestResult");
+
+		UTEST_TRUE("Connected Transport To Metronome",
+			FGraphBuilder::ConnectNodes(TransportNode, CommonPinNames::Outputs::TransportName,
+				MetronomeNode, CommonPinNames::Inputs::TransportName));
+
+		UTEST_TRUE("Connected Metronome To Secs Until Node",
+			FGraphBuilder::ConnectNodes(MetronomeNode, CommonPinNames::Outputs::MidiClockName,
+				SecsUntilNode, CommonPinNames::Inputs::MidiClockName));
+
+		Builder.AddOutput("AudioOut", GetMetasoundDataTypeName<FAudioBuffer>());
+
+		constexpr int32 NumSamplesPerBlock = 256;
+		constexpr FSampleRate SampleRate = 48000;
+
+		TUniquePtr<FMetasoundGenerator> Generator = Builder.BuildGenerator(48000, NumSamplesPerBlock);
+		UTEST_TRUE("Made Generator", Generator.IsValid());
+		FAudioBuffer Buffer{ Generator->OperatorSettings };
+
+		TOptional<FTriggerWriteRef> StartTestTrigger = Generator->GetInputWriteReference<FTrigger>("StartTest");
+		UTEST_TRUE("Got Start Test Trigger", StartTestTrigger.IsSet());
+
+		TOptional<FFloatWriteRef> TestTempoInput = Generator->GetInputWriteReference<float>("TestTempo");
+		UTEST_TRUE("Got Test Tempo Input", TestTempoInput.IsSet());
+
+		TOptional<FInt32WriteRef> TestTSNumInput = Generator->GetInputWriteReference<int32>("TestTSNum");
+		UTEST_TRUE("Got Test Time Sig Numerator Input", TestTSNumInput.IsSet());
+
+		TOptional<FInt32WriteRef> TestTSDenomInput = Generator->GetInputWriteReference<int32>("TestTSDenom");
+		UTEST_TRUE("Got Test Time Sig Denominator Input", TestTSDenomInput.IsSet());
+
+		TOptional<FMusicTimestampWriteRef> TestTimestampInput = Generator->GetInputWriteReference<FMusicTimestamp>("TestTimestamp");
+		UTEST_TRUE("Got Test Timestamp Input", TestTimestampInput.IsSet());
+
+		TOptional<FFloatReadRef> ResultOutput = Generator->GetOutputReadReference<float>("TestResult");
+		UTEST_TRUE("Got Test Result Output", ResultOutput.IsSet());
+
+		*(*TestTempoInput) = Tempo;
+		(*TestTimestampInput)->Bar = Bar;
+		(*TestTimestampInput)->Beat = Beat;
+		*(*TestTSNumInput) = TimeSigNumerator;
+		*(*TestTSDenomInput) = TimeSigDenominator;
+		(*StartTestTrigger)->TriggerFrame(0);
+
+		Generator->OnGenerateAudio(Buffer.GetData(), Buffer.Num());
+
+		float TotalBeats = ((Bar - 1) * TimeSigNumerator) + (Beat - 1);
+		float SecsPerQuarter = 60.0f / Tempo;
+		float SecsPerBeat = SecsPerQuarter / (TimeSigDenominator / 4);
+		float ExpectedSeconds = (SecsPerBeat * TotalBeats) - ((float)NumSamplesPerBlock / (float)SampleRate);
+		float SecsUntil = *(*ResultOutput);
+
+		UTEST_TRUE("Time Until As Expected", FMath::IsNearlyEqual(SecsUntil, ExpectedSeconds, 0.0005));
+		
+		return true;
+	}
+}
+
+#endif
