Fix log spam in tests

diff --git a/Engine/Plugins/Runtime/Harmonix/Source/HarmonixMetasoundTests/Private/Analysis/MidiSongPosVertexAnalyzerTests.cpp b/Engine/Plugins/Runtime/Harmonix/Source/HarmonixMetasoundTests/Private/Analysis/MidiSongPosVertexAnalyzerTests.cpp
--- a/Engine/Plugins/Runtime/Harmonix/Source/HarmonixMetasoundTests/Private/Analysis/MidiSongPosVertexAnalyzerTests.cpp
+++ b/Engine/Plugins/Runtime/Harmonix/Source/HarmonixMetasoundTests/Private/Analysis/MidiSongPosVertexAnalyzerTests.cpp
@@ -186,20 +186,20 @@ namespace HarmonixMetasoundTests::MidiSongPosVertexAnalyzer
 
 				if (Pos.UpToTick < LastTick)
 				{
-					UE_LOG(LogMIDI, Log, TEXT("----------------- LOOPED----------------- "));
+					UE_LOG(LogMIDI, VeryVerbose, TEXT("----------------- LOOPED----------------- "));
 				}
 				LastTick = Pos.UpToTick;
 
-				UE_LOG(LogMIDI, Log, TEXT("POSITION: SampleCount = %d, Transport = %s"), Pos.SampleCount, *MusicPlayerTransportStateToString(Pos.CurrentTransportState));
+				UE_LOG(LogMIDI, VeryVerbose, TEXT("POSITION: SampleCount = %d, Transport = %s"), Pos.SampleCount, *MusicPlayerTransportStateToString(Pos.CurrentTransportState));
 				
 				NewSongPos.SetByTick(LastTick, *(ClockHistory->GetLatestMapsForConsumer()->SongMaps));
 
 				LogSongPos(NewSongPos);
 
-				UE_LOG(LogMIDI, Log, TEXT("LERP POSITION:"));
+				UE_LOG(LogMIDI, VeryVerbose, TEXT("LERP POSITION:"));
 				if (PreviousMarkerType == HarmonixMetasound::Analysis::FMidiClockSongPosition::EMarkerType::LastPositionBeforeSeekLoop)
 				{
-					UE_LOG(LogMIDI, Log, TEXT("\t<none>"));
+					UE_LOG(LogMIDI, VeryVerbose, TEXT("\t<none>"));
 				}
 				else
 				{
diff --git a/Engine/Plugins/Runtime/Harmonix/Source/HarmonixMidiTests/Private/SongMapsTests.cpp b/Engine/Plugins/Runtime/Harmonix/Source/HarmonixMidiTests/Private/SongMapsTests.cpp
--- a/Engine/Plugins/Runtime/Harmonix/Source/HarmonixMidiTests/Private/SongMapsTests.cpp
+++ b/Engine/Plugins/Runtime/Harmonix/Source/HarmonixMidiTests/Private/SongMapsTests.cpp
@@ -11,7 +11,7 @@ DEFINE_LOG_CATEGORY_STATIC(LogSongMapsTest, Log, All);
 namespace HarmonixMidiTests::SongMaps
 {
 
-	static constexpr bool GLogQualtizationDetails = true;
+	static constexpr bool GLogQualtizationDetails = false;
 
 	void LogQuantizationDetails(const ISongMapEvaluator& Map, int32 OriginalTick, int32 QuantizedTick, EMidiClockSubdivisionQuantization ResultDivision)
 	{
