- Fix seek to driving clock logic

diff --git a/Engine/Plugins/Runtime/Harmonix/Source/HarmonixMetasound/Private/HarmonixMetasound/DataTypes/MidiClock.cpp b/Engine/Plugins/Runtime/Harmonix/Source/HarmonixMetasound/Private/HarmonixMetasound/DataTypes/MidiClock.cpp
--- a/Engine/Plugins/Runtime/Harmonix/Source/HarmonixMetasound/Private/HarmonixMetasound/DataTypes/MidiClock.cpp
+++ b/Engine/Plugins/Runtime/Harmonix/Source/HarmonixMetasound/Private/HarmonixMetasound/DataTypes/MidiClock.cpp
@@ -46,6 +46,7 @@ namespace HarmonixMetasound
 		, TickResidualWhenDriven(0.0f)
 		, BlockSize(InSettings.GetNumFramesPerBlock())
 		, CurrentBlockFrameIndex(0)
+		, FirstTickProcessedThisBlock(-1)
 		, LastProcessedMidiTick(-1)
 		, NextMidiTickToProcess(0)
 		, SampleRate(InSettings.GetSampleRate())
@@ -83,6 +84,7 @@ namespace HarmonixMetasound
 		, TickResidualWhenDriven(Other.TickResidualWhenDriven)
 		, BlockSize(Other.BlockSize)
 		, CurrentBlockFrameIndex(Other.CurrentBlockFrameIndex)
+		, FirstTickProcessedThisBlock(Other.FirstTickProcessedThisBlock)
 		, LastProcessedMidiTick(Other.LastProcessedMidiTick)
 		, NextMidiTickToProcess(Other.NextMidiTickToProcess)
 		, SampleRate(Other.SampleRate)
@@ -127,6 +129,7 @@ namespace HarmonixMetasound
 			TickResidualWhenDriven = Other.TickResidualWhenDriven;
 			BlockSize = Other.BlockSize;
 			CurrentBlockFrameIndex = Other.CurrentBlockFrameIndex;
+			FirstTickProcessedThisBlock = Other.FirstTickProcessedThisBlock;
 			LastProcessedMidiTick = Other.LastProcessedMidiTick;
 			NextMidiTickToProcess = Other.NextMidiTickToProcess;
 			SampleRate = Other.SampleRate;
@@ -208,6 +211,8 @@ namespace HarmonixMetasound
 
 	void FMidiClock::PrepareBlock()
 	{
+		FirstTickProcessedThisBlock = -1;
+
 		NumTransportChangeInBlock = 0;
 		TransportAtBlockStart = TransportAtBlockEnd;
 		
@@ -341,14 +346,17 @@ namespace HarmonixMetasound
 		{
 			NeedsSeekToDrivingClock = false;
 			int32 OurStartTick = DrivingClock.GetNextTickToProcessAtBlockFrame(StartFrame);
-			if (!FMath::IsNearlyEqual(CurrentLocalSpeed, 1.0f, 0.0001))
+			if (OurStartTick >= 0)
 			{
-				float FractionalToTick = (float)OurStartTick * CurrentLocalSpeed;
-				OurStartTick = FMath::FloorToInt32(FractionalToTick);
-				TickResidualWhenDriven = FMath::Fractional(FractionalToTick);
+				if (!FMath::IsNearlyEqual(CurrentLocalSpeed, 1.0f, 0.0001))
+				{
+					float FractionalToTick = (float)OurStartTick * CurrentLocalSpeed;
+					OurStartTick = FMath::FloorToInt32(FractionalToTick);
+					TickResidualWhenDriven = FMath::Fractional(FractionalToTick);
+				}
+				OurStartTick = WrapTickIfLooping(OurStartTick);
+				SeekTo(ClockEvents[0].BlockFrameIndex, OurStartTick);
 			}
-			OurStartTick = WrapTickIfLooping(OurStartTick);
-			SeekTo(ClockEvents[0].BlockFrameIndex, OurStartTick);
 		}
 
 		while (ClockEvents.IsValidIndex(Index))
@@ -589,7 +597,12 @@ namespace HarmonixMetasound
 
 	int32 FMidiClock::GetNextTickToProcessAtBlockFrame(int32 BlockFrame) const
 	{
-		int32 FoundNextTick = NextMidiTickToProcess;
+		if (BlockFrame == 0)
+		{
+			return FirstTickProcessedThisBlock < 0 ? 0 : FirstTickProcessedThisBlock;
+		}
+
+		int32 FoundNextTick = FirstTickProcessedThisBlock;
 		int32 AtBlockIndex = 0;
 
 		using namespace MidiClockMessageTypes;
@@ -623,9 +636,19 @@ namespace HarmonixMetasound
 				FoundNextTick = AsLoop->FirstTickInLoop;
 				AtBlockIndex = Event.BlockFrameIndex;
 			}
+			else if (const FTempoChange* AsTempoChange = Event.TryGet<FTempoChange>())
+			{
+				FoundNextTick = AsTempoChange->Tick;
+				AtBlockIndex = Event.BlockFrameIndex;
+			}
+			else if (const FTimeSignatureChange* AsTimeSigChange = Event.TryGet<FTimeSignatureChange>())
+			{
+				FoundNextTick = AsTimeSigChange->Tick;
+				AtBlockIndex = Event.BlockFrameIndex;
+			}
 		}
 
-		return FoundNextTick;
+		return FoundNextTick < 0 ? 0 : FoundNextTick;
 	}
 
 	void FMidiClock::AddEvent(const FMidiClockEvent& InEvent, bool bRequireSequential)
@@ -774,6 +797,11 @@ namespace HarmonixMetasound
 
 		check(Tick >= NextMidiTickToProcess);
 
+		if (FirstTickProcessedThisBlock == -1)
+		{
+			FirstTickProcessedThisBlock = Tick;
+		}
+
 		CurrentBlockFrameIndex = BlockFrameIndex;
 		NextMidiTickToProcess = Tick;
 
@@ -810,6 +838,11 @@ namespace HarmonixMetasound
 		CurrentBlockFrameIndex = BlockFrameIndex;
 		NextMidiTickToProcess = Tick;
 
+		if (FirstTickProcessedThisBlock == -1)
+		{
+			FirstTickProcessedThisBlock = Tick;
+		}
+
 		if (TempoAtBlockEnd == NewTempo)
 		{
 			// no need to add the tempo... it is already the current tempo. 
@@ -939,6 +972,11 @@ namespace HarmonixMetasound
 
 		CurrentBlockFrameIndex = BlockFrameIndex;
 
+		if (FirstTickProcessedThisBlock == -1)
+		{
+			FirstTickProcessedThisBlock = FirstTick;
+		}
+
 		// moving forward may cause us to move into a new tempo and/or time signature...
 		while (FirstTick <= NextTempoOrTimeSigChangeTick && NextTempoOrTimeSigChangeTick < (FirstTick + NumTicks))
 		{
diff --git a/Engine/Plugins/Runtime/Harmonix/Source/HarmonixMetasound/Public/HarmonixMetasound/DataTypes/MidiClock.h b/Engine/Plugins/Runtime/Harmonix/Source/HarmonixMetasound/Public/HarmonixMetasound/DataTypes/MidiClock.h
--- a/Engine/Plugins/Runtime/Harmonix/Source/HarmonixMetasound/Public/HarmonixMetasound/DataTypes/MidiClock.h
+++ b/Engine/Plugins/Runtime/Harmonix/Source/HarmonixMetasound/Public/HarmonixMetasound/DataTypes/MidiClock.h
@@ -194,6 +194,7 @@ namespace HarmonixMetasound
 
 		int32 BlockSize;
 		int32 CurrentBlockFrameIndex;
+		int32 FirstTickProcessedThisBlock;
 		int32 LastProcessedMidiTick;
 		int32 NextMidiTickToProcess;
 		float SampleRate;
