Fix issue where Midi Player Node caused midi clock to seek with wrong block frame index

diff --git a/Engine/Plugins/Runtime/Harmonix/Source/HarmonixMetasound/Private/HarmonixMetasound/Nodes/MidiPlayerNode.cpp b/Engine/Plugins/Runtime/Harmonix/Source/HarmonixMetasound/Private/HarmonixMetasound/Nodes/MidiPlayerNode.cpp
--- a/Engine/Plugins/Runtime/Harmonix/Source/HarmonixMetasound/Private/HarmonixMetasound/Nodes/MidiPlayerNode.cpp
+++ b/Engine/Plugins/Runtime/Harmonix/Source/HarmonixMetasound/Private/HarmonixMetasound/Nodes/MidiPlayerNode.cpp
@@ -488,8 +488,6 @@ namespace HarmonixMetasound
 		InVertexData.SetValue(METASOUND_GET_PARAM_NAME(MidiPlayerNodePinNames::KillVoicesOnSeek), bKillVoicesOnSeek);
 		InVertexData.SetValue(METASOUND_GET_PARAM_NAME(MidiPlayerNodePinNames::KillVoicesOnMidiChange), bKillVoicesOnMidiChange);
 
-		SetupNewMidiFile(MidiAssetInPin->GetMidiProxy());
-
 		NeedsTransportInit = true;
 	}
 
