Fix music clock component bug where the current tempo and beat durations would be wrong on looping and offset midi clocks.
- Add TempoMapTick to midi clock events so that a real mapping can be done to get from a "local tick", which might be looping round-n-round or offset, to a "tempo map tick" so that the proper tempo for that local tick can be looked up in the tempo map.
- Fix a few bugs related to offset clocks.
- Fixed a bug in FMidiClock when adding a loop event to the event list it might put in the wrong destination tick. This in cases where one clock is driving another.
- Fixed a bug where the clock offset node and externally clocked midi player node were not properly updating their output midi clock's tempo map when their input tempo map was changed.

#tests Hmx Automated Tests

diff --git a/Engine/Plugins/Runtime/Harmonix/Source/HarmonixMetasound/Private/HarmonixMetasound/Analysis/MidiSongPosVertexAnalyzer.cpp b/Engine/Plugins/Runtime/Harmonix/Source/HarmonixMetasound/Private/HarmonixMetasound/Analysis/MidiSongPosVertexAnalyzer.cpp
--- a/Engine/Plugins/Runtime/Harmonix/Source/HarmonixMetasound/Private/HarmonixMetasound/Analysis/MidiSongPosVertexAnalyzer.cpp
+++ b/Engine/Plugins/Runtime/Harmonix/Source/HarmonixMetasound/Private/HarmonixMetasound/Analysis/MidiSongPosVertexAnalyzer.cpp
@@ -101,7 +101,7 @@ namespace HarmonixMetasound::Analysis
 			if (TryProcessAsSpeedChange(Event, BlockFrameAdvance, FirstEventAfterSeekOrLoop)) continue;
 			if (TryProcessAsTransportChange(Event, BlockFrameAdvance, FirstEventAfterSeekOrLoop)) continue;
 			if (TryProcessAsLoop(Event, BlockFrameAdvance, FirstEventAfterSeekOrLoop)) continue;
-			// This is the last type re recognize, so if it isn't a seek we're in trouble...
+			// This is the last type we recognize, so if it isn't a seek we're in trouble...
 			ensureMsgf(TryProcessAsSeek(Event, BlockFrameAdvance, FirstEventAfterSeekOrLoop), TEXT("Unrecognized Clock Advance Event Type!"));
 		}
 		SampleCount += BlockSize;
@@ -124,7 +124,9 @@ namespace HarmonixMetasound::Analysis
 				return true;
 			}
 			LastTickProcessed = AsAdvance->FirstTickToProcess;
+			LastTempoMapTickProcessed = AsAdvance->TempoMapTick;
 			LastMidiClockSongPos->UpToTick = LastTickProcessed;
+			LastMidiClockSongPos->TempoMapTick = LastTempoMapTickProcessed;
 			LastMidiClockSongPos->SampleCount = SampleCount + Event.BlockFrameIndex + BlockFrameAdvance;
 			LastMidiClockSongPos->MarkerType = (FirstEventAfterSeekOrLoop) ? FMidiClockSongPosition::EMarkerType::FirstPositionAfterSeekLoop : FMidiClockSongPosition::EMarkerType::None;
 			{
@@ -155,7 +157,9 @@ namespace HarmonixMetasound::Analysis
 				return true;
 			}
 			LastTickProcessed = AsTempo->Tick;
+			LastTempoMapTickProcessed = AsTempo->TempoMapTick;
 			LastMidiClockSongPos->UpToTick = LastTickProcessed;
+			LastMidiClockSongPos->TempoMapTick = LastTempoMapTickProcessed;
 			LastMidiClockSongPos->SampleCount = SampleCount + Event.BlockFrameIndex + BlockFrameAdvance;
 			LastMidiClockSongPos->MarkerType = (FirstEventAfterSeekOrLoop) ? FMidiClockSongPosition::EMarkerType::FirstPositionAfterSeekLoop : FMidiClockSongPosition::EMarkerType::None;
 			{
@@ -182,7 +186,9 @@ namespace HarmonixMetasound::Analysis
 				return true;
 			}
 			LastTickProcessed = AsTimeSig->Tick;
+			LastTempoMapTickProcessed = AsTimeSig->TempoMapTick;
 			LastMidiClockSongPos->UpToTick = LastTickProcessed;
+			LastMidiClockSongPos->TempoMapTick = LastTempoMapTickProcessed;
 			LastMidiClockSongPos->SampleCount = SampleCount + Event.BlockFrameIndex + BlockFrameAdvance;
 			LastMidiClockSongPos->MarkerType = (FirstEventAfterSeekOrLoop) ? FMidiClockSongPosition::EMarkerType::FirstPositionAfterSeekLoop : FMidiClockSongPosition::EMarkerType::None;
 
@@ -207,6 +213,7 @@ namespace HarmonixMetasound::Analysis
 				LastMidiClockSongPos->SampleCount = SampleCount + Event.BlockFrameIndex + BlockFrameAdvance;
 				LastMidiClockSongPos->MarkerType = (FirstEventAfterSeekOrLoop) ? FMidiClockSongPosition::EMarkerType::FirstPositionAfterSeekLoop : FMidiClockSongPosition::EMarkerType::None;
 				LastMidiClockSongPos->UpToTick = LastTickProcessed;
+				LastMidiClockSongPos->TempoMapTick = LastTempoMapTickProcessed;
 				{
 					FMidiClockSongPositionHistory::FScopedItemWriteRef Slot = History->Positions.GetNextAtomicWriteSlot();
 					*Slot = *LastMidiClockSongPos.Get();
@@ -228,6 +235,7 @@ namespace HarmonixMetasound::Analysis
 				LastMidiClockSongPos->SampleCount = SampleCount + Event.BlockFrameIndex + BlockFrameAdvance;
 				LastMidiClockSongPos->MarkerType = (FirstEventAfterSeekOrLoop) ? FMidiClockSongPosition::EMarkerType::FirstPositionAfterSeekLoop : FMidiClockSongPosition::EMarkerType::None;
 				LastMidiClockSongPos->UpToTick = LastTickProcessed;
+				LastMidiClockSongPos->TempoMapTick = LastTempoMapTickProcessed;
 				{
 					FMidiClockSongPositionHistory::FScopedItemWriteRef Slot = History->Positions.GetNextAtomicWriteSlot();
 					*Slot = *LastMidiClockSongPos.Get();
@@ -247,7 +255,9 @@ namespace HarmonixMetasound::Analysis
 			// We need to insert a song position for where we got to BEFORE this loop back...
 			if (LastAdvanceUpToTick != -1)
 			{
+				LastTempoMapTickProcessed = AsLoop->TempoMapTick;
 				LastMidiClockSongPos->UpToTick = LastAdvanceUpToTick;
+				LastMidiClockSongPos->TempoMapTick = LastTempoMapTickProcessed;
 				LastMidiClockSongPos->SampleCount = SampleCount + Event.BlockFrameIndex + BlockFrameAdvance;
 				LastMidiClockSongPos->MarkerType = FMidiClockSongPosition::EMarkerType::LastPositionBeforeSeekLoop;
 				LastAdvanceUpToTick = -1;
@@ -270,7 +280,9 @@ namespace HarmonixMetasound::Analysis
 			// We need to insert a song position for where we got to BEFORE this seek...
 			if (LastAdvanceUpToTick != -1)
 			{
+				LastTempoMapTickProcessed = AsSeek->TempoMapTick;
 				LastMidiClockSongPos->UpToTick = LastAdvanceUpToTick;
+				LastMidiClockSongPos->TempoMapTick = LastTempoMapTickProcessed;
 				LastMidiClockSongPos->SampleCount = SampleCount + Event.BlockFrameIndex + BlockFrameAdvance;
 				LastMidiClockSongPos->MarkerType = FMidiClockSongPosition::EMarkerType::LastPositionBeforeSeekLoop;
 				LastAdvanceUpToTick = -1;
diff --git a/Engine/Plugins/Runtime/Harmonix/Source/HarmonixMetasound/Private/HarmonixMetasound/Components/MetasoundMusicClockDriver.cpp b/Engine/Plugins/Runtime/Harmonix/Source/HarmonixMetasound/Private/HarmonixMetasound/Components/MetasoundMusicClockDriver.cpp
--- a/Engine/Plugins/Runtime/Harmonix/Source/HarmonixMetasound/Private/HarmonixMetasound/Components/MetasoundMusicClockDriver.cpp
+++ b/Engine/Plugins/Runtime/Harmonix/Source/HarmonixMetasound/Private/HarmonixMetasound/Components/MetasoundMusicClockDriver.cpp
@@ -64,68 +64,96 @@ bool FMetasoundMusicClockDriver::CalculateSongPosWithOffset(float MsOffset, ECal
 		return false;
 	}
 
-	// start with the offset...
-	float AbsMs = MsOffset;
-	
-	// now add the current time from the proper time base...
+	const FPerTimebaseSmoothedClockState* ClockState = nullptr;
+
 	switch (Timebase)
 	{
-	case ECalibratedMusicTimebase::AudioRenderTime:
-		AbsMs += ClockComponent->CurrentSmoothedAudioRenderSongPos.SecondsIncludingCountIn * 1000.0f;
-		break;
-	case ECalibratedMusicTimebase::ExperiencedTime:
-		AbsMs += ClockComponent->CurrentPlayerExperiencedSongPos.SecondsIncludingCountIn * 1000.0f;
-		break;
+	case ECalibratedMusicTimebase::AudioRenderTime:  ClockState = &AudioRenderState; break;
+	case ECalibratedMusicTimebase::ExperiencedTime:  ClockState = &PlayerExperienceState; break;
 	case ECalibratedMusicTimebase::VideoRenderTime:
-	default:
-		AbsMs += ClockComponent->CurrentVideoRenderSongPos.SecondsIncludingCountIn * 1000.0f;
-		break;
+	default:                                         ClockState = &VideoRenderState; break;
 	}
 
-	// if looping map to range of loop...
+	float AbsMs = ClockState->TempoMapMs + MsOffset;
+	float TempoMapOffsetTick = CurrentMapChain->SongMaps->MsToTick(AbsMs);
+	float RelativeTicks = TempoMapOffsetTick - ClockState->TempoMapTick;
+	float SmoothedOffsetTick = ClockState->LocalTick + RelativeTicks;
+
 	if (CurrentMapChain->LoopLengthTicks > 0)
 	{
-		float LoopStartMs = CurrentMapChain->SongMaps->TickToMs(CurrentMapChain->FirstTickInLoop);
-		float LoopEndMs = CurrentMapChain->SongMaps->TickToMs(CurrentMapChain->FirstTickInLoop + CurrentMapChain->LoopLengthTicks);
-		if (AbsMs > LoopEndMs)
+		while (SmoothedOffsetTick < 0.0f)
 		{
-			float LoopLengthMs = LoopEndMs - LoopStartMs;
-			AbsMs = LoopStartMs + FMath::Fmod(AbsMs - LoopStartMs, LoopLengthMs);
+			SmoothedOffsetTick += CurrentMapChain->LoopLengthTicks;
+		}
+		while (SmoothedOffsetTick >= CurrentMapChain->LoopLengthTicks)
+		{
+			SmoothedOffsetTick -= CurrentMapChain->LoopLengthTicks;
 		}
 	}
 
-	// and finally we can calculate a song position...
-	OutResult.SetByTime(AbsMs, *(CurrentMapChain->SongMaps));
+	// first 99% of the song pos...
+	OutResult.SetByTick(SmoothedOffsetTick, *(CurrentMapChain->SongMaps));
+	// but tempo needs to be set to authoritiesTempo...
+	OutResult.Tempo = CurrentMapChain->SongMaps->GetTempoAtTick(FMath::FloorToInt32(TempoMapOffsetTick));
 
 	return true;
 }
 
-FMidiSongPos FMetasoundMusicClockDriver::CalculateSongPosAtMs(float AbsoluteMs) const
+FMidiSongPos FMetasoundMusicClockDriver::CalculateSongPosAtMsForLoopingOrMonotonicClock(float AbsoluteMs, float& PositionTick) const
 {
 	using namespace HarmonixMetasound::Analysis;
 	FMidiSongPos OutSongPos;
 	if (!ClockHistory)
 	{
+		PositionTick = 0.0f;
 		return OutSongPos;
 	}
 
 	if (!CurrentMapChain || !CurrentMapChain->SongMaps)
 	{
+		PositionTick = 0.0f;
 		return OutSongPos;
 	}
 
 	if (CurrentMapChain->LoopLengthTicks > 0)
 	{
-		float LoopStartMs = CurrentMapChain->SongMaps->TickToMs(CurrentMapChain->FirstTickInLoop);
-		float LoopEndMs = CurrentMapChain->SongMaps->TickToMs(CurrentMapChain->FirstTickInLoop + CurrentMapChain->LoopLengthTicks);
-		if (AbsoluteMs > LoopEndMs)
+		float DrivingTick = CurrentMapChain->SongMaps->MsToTick(AbsoluteMs);
+		if (DrivingTick >= (float)(CurrentMapChain->FirstTickInLoop + CurrentMapChain->LoopLengthTicks))
 		{
-			float LoopLengthMs = LoopEndMs - LoopStartMs;
-			AbsoluteMs = LoopStartMs + FMath::Fmod(AbsoluteMs - LoopStartMs, LoopLengthMs);
+			PositionTick = FMath::Fmod(DrivingTick - CurrentMapChain->FirstTickInLoop, (float)CurrentMapChain->LoopLengthTicks);
+			OutSongPos.SetByTick(PositionTick, *CurrentMapChain->SongMaps);
+			OutSongPos.Tempo = CurrentMapChain->SongMaps->GetTempoAtTick(FMath::FloorToInt32(DrivingTick));
+			return OutSongPos;
 		}
 	}
 
-	OutSongPos.SetByTime(AbsoluteMs, *(CurrentMapChain->SongMaps));
+	PositionTick = CurrentMapChain->SongMaps->MsToTick(AbsoluteMs);
+	OutSongPos.SetByTimeAndTick(AbsoluteMs, PositionTick, *(CurrentMapChain->SongMaps));
+
+	return OutSongPos;
+}
+
+FMidiSongPos FMetasoundMusicClockDriver::CalculateSongPosAtMsForOffsetClock(float PositionMs, float ClockTickOffsetFromDrivingClock, float& PositionTick) const
+{
+	using namespace HarmonixMetasound::Analysis;
+	
+	FMidiSongPos OutSongPos;
+	
+	if (!ClockHistory)
+	{
+		PositionTick = 0.0f;
+		return OutSongPos;
+	}
+
+	if (!CurrentMapChain || !CurrentMapChain->SongMaps)
+	{
+		PositionTick = 0.0f;
+		return OutSongPos;
+	}
+
+	PositionTick = CurrentMapChain->SongMaps->MsToTick(PositionMs);
+	OutSongPos.SetByTick(PositionTick, *CurrentMapChain->SongMaps);
+	OutSongPos.Tempo = CurrentMapChain->SongMaps->GetTempoAtTick(FMath::FloorToInt32(PositionTick - ClockTickOffsetFromDrivingClock));
 
 	return OutSongPos;
 }
@@ -458,8 +486,9 @@ void FMetasoundMusicClockDriver::RefreshCurrentSongPosFromHistory()
 	// it means we are rendering in such large blocks that we need to push up our "look behind" number for smoothing.
 	float SpeedAtSmoothRenderedTime = 1.0f;
 	float SmoothedTick = ClockComponent->GetTicksIncludingCountIn(ECalibratedMusicTimebase::AudioRenderTime);
+	float SmoothedTempoMapTick = AudioRenderState.TempoMapTick;
 	EHistoryFailureType Result = CalculateSmoothedTick(ExpectedRenderPosSampleCount, LastRenderPosSampleCount,
-		SmoothedTick, SpeedAtSmoothRenderedTime, SmoothedAudioRenderClockHistoryCursor, RenderSmoothingLagSeconds);
+		SmoothedTick, SmoothedTempoMapTick, SpeedAtSmoothRenderedTime, SmoothedAudioRenderClockHistoryCursor, RenderSmoothingLagSeconds);
 	if (Result != EHistoryFailureType::None)
 	{
 		if (LastRenderPosSampleCount >  (RenderSmoothingLagSeconds * ClockHistory->SampleRate * 2))
@@ -481,25 +510,17 @@ void FMetasoundMusicClockDriver::RefreshCurrentSongPosFromHistory()
 		return;
 	}
 
-	// Now... We are behind the actual render time because of the lag we introduce to have enough
-	// history... SO... Push forward to get a time that is approx. 30 ms behind the
-	// render time. Why 30ms? Historical reasons. 
-	float SmoothedSongMs = CurrentMapChain->SongMaps->TickToMs(SmoothedTick);
-	SmoothedSongMs += RenderSmoothingLagSeconds * 1000.0f;
-	ClockComponent->CurrentSmoothedAudioRenderSongPos = CalculateSongPosAtMs(SmoothedSongMs);
-
-	// Now the time the user should actually be "experiencing"... ie "hearing" can be calculated as an offset
-	// from the smooth audio rendering time...
-	float LookBehindLagMs = FHarmonixModule::Get().GetMeasuredUserExperienceAndReactionToAudioRenderOffsetMs();
-	float Ms = ClockComponent->CurrentSmoothedAudioRenderSongPos.SecondsIncludingCountIn * 1000.0f;
-	Ms -= LookBehindLagMs; // - (RenderSmoothingLagSeconds * 1000.0f);
-	ClockComponent->CurrentPlayerExperiencedSongPos = CalculateSongPosAtMs(Ms);
-
-	// Now the time we should be rendering graphics for...
-	LookBehindLagMs = FHarmonixModule::Get().GetMeasuredVideoToAudioRenderOffsetMs();
-	Ms = ClockComponent->CurrentSmoothedAudioRenderSongPos.SecondsIncludingCountIn * 1000.0f;
-	Ms -= LookBehindLagMs; // - (RenderSmoothingLagSeconds * 1000.0f);
-	ClockComponent->CurrentVideoRenderSongPos = CalculateSongPosAtMs(Ms);
+	if (SmoothedTempoMapTick != SmoothedTick && CurrentMapChain->LoopLengthTicks <= 0)
+	{
+		// The clock is offset from its SongMaps (eg. It is the output of a 
+		// clock offset node.)
+		UpdateCurrentTicksForOffsetClock(SmoothedTick, SmoothedTempoMapTick);
+	}
+	else
+	{
+		// The clock is looping or monotonically increasing... so we deal with it this way...
+		UpdateCurrentTicksForLoopingOrMonotonicClock(SmoothedTick, SmoothedTempoMapTick);
+	}
 
 	ClockComponent->TimeSignatureNum = ClockComponent->CurrentSmoothedAudioRenderSongPos.TimeSigNumerator;
 	ClockComponent->TimeSignatureDenom = ClockComponent->CurrentSmoothedAudioRenderSongPos.TimeSigDenominator;
@@ -513,8 +534,68 @@ void FMetasoundMusicClockDriver::RefreshCurrentSongPosFromHistory()
 	}
 }
 
+void FMetasoundMusicClockDriver::UpdateCurrentTicksForOffsetClock(float SmoothedTick, float SmoothedTempoMapTick)
+{
+	float SmoothedPositionMs = CurrentMapChain->SongMaps->TickToMs(SmoothedTick);
+	// We are behind the actual render time because of the lag we introduce to have enough
+	// history... SO... Push forward to get a time that is approx. where the renderer is.
+	SmoothedPositionMs += RenderSmoothingLagSeconds * 1000.0f;
+
+	// Calculate the song position AND get the "local tick" for the Audio Render timebase
+	ClockComponent->CurrentSmoothedAudioRenderSongPos = CalculateSongPosAtMsForOffsetClock(SmoothedPositionMs, SmoothedTick - SmoothedTempoMapTick, AudioRenderState.LocalTick);
+	float SmoothedTempoMapMs = CurrentMapChain->SongMaps->TickToMs(SmoothedTempoMapTick);
+	AudioRenderState.TempoMapMs = SmoothedTempoMapMs + (RenderSmoothingLagSeconds * 1000.0f);
+	AudioRenderState.TempoMapTick = CurrentMapChain->SongMaps->MsToTick(AudioRenderState.TempoMapMs);
+
+	// Now the time the user should actually be "experiencing" (ie "hearing") can be calculated as an offset
+	// from the smooth audio rendering time...
+	float LookBehindLagMs = FHarmonixModule::Get().GetMeasuredUserExperienceAndReactionToAudioRenderOffsetMs();
+	float Ms = SmoothedPositionMs;
+	Ms -= LookBehindLagMs;
+	// Calculate the song position AND get the "local tick" for the Player Experience timebase
+	ClockComponent->CurrentPlayerExperiencedSongPos = CalculateSongPosAtMsForOffsetClock(Ms, SmoothedTick - SmoothedTempoMapTick, PlayerExperienceState.LocalTick);
+	PlayerExperienceState.TempoMapMs = AudioRenderState.TempoMapMs - LookBehindLagMs;
+	PlayerExperienceState.TempoMapTick = CurrentMapChain->SongMaps->MsToTick(PlayerExperienceState.TempoMapMs);
+
+	// Now the time the game should be rendering graphics for can be calculated as an offset
+	// from the smooth audio rendering time...
+	LookBehindLagMs = FHarmonixModule::Get().GetMeasuredVideoToAudioRenderOffsetMs();
+	Ms = SmoothedPositionMs;
+	Ms -= LookBehindLagMs;
+	// Calculate the song position AND get the "local tick" for the Video Render timebase
+	ClockComponent->CurrentVideoRenderSongPos = CalculateSongPosAtMsForOffsetClock(Ms, SmoothedTick - SmoothedTempoMapTick, VideoRenderState.LocalTick);
+	VideoRenderState.TempoMapMs = AudioRenderState.TempoMapMs - LookBehindLagMs;
+	VideoRenderState.TempoMapTick = CurrentMapChain->SongMaps->MsToTick(VideoRenderState.TempoMapMs);
+}
+
+void FMetasoundMusicClockDriver::UpdateCurrentTicksForLoopingOrMonotonicClock(float SmoothedTick, float SmoothedTempoMapTick)
+{
+	AudioRenderState.TempoMapMs = CurrentMapChain->SongMaps->TickToMs(SmoothedTempoMapTick);
+	// We are behind the actual render time because of the lag we introduce to have enough
+	// history... SO... Push forward to get a time that is approx. where the renderer is.
+	AudioRenderState.TempoMapMs += RenderSmoothingLagSeconds * 1000.0f;
+
+	// Calculate the song position AND get the "local tick" for the Audio Render timebase
+	AudioRenderState.TempoMapTick = CurrentMapChain->SongMaps->MsToTick(AudioRenderState.TempoMapMs);
+	ClockComponent->CurrentSmoothedAudioRenderSongPos = CalculateSongPosAtMsForLoopingOrMonotonicClock(AudioRenderState.TempoMapMs, AudioRenderState.LocalTick);
+
+	// Now the time the user should actually be "experiencing" (ie "hearing") can be calculated as an offset
+	// from the smooth audio rendering time...
+	float LookBehindLagMs = FHarmonixModule::Get().GetMeasuredUserExperienceAndReactionToAudioRenderOffsetMs();
+	PlayerExperienceState.TempoMapMs = AudioRenderState.TempoMapMs - LookBehindLagMs;
+	PlayerExperienceState.TempoMapTick = CurrentMapChain->SongMaps->MsToTick(PlayerExperienceState.TempoMapMs);
+	ClockComponent->CurrentPlayerExperiencedSongPos = CalculateSongPosAtMsForLoopingOrMonotonicClock(PlayerExperienceState.TempoMapMs, PlayerExperienceState.LocalTick);
+
+	// Now the time the game should be rendering graphics for can be calculated as an offset
+	// from the smooth audio rendering time...
+	LookBehindLagMs = FHarmonixModule::Get().GetMeasuredVideoToAudioRenderOffsetMs();
+	VideoRenderState.TempoMapMs = AudioRenderState.TempoMapMs - LookBehindLagMs;
+	VideoRenderState.TempoMapTick = CurrentMapChain->SongMaps->MsToTick(VideoRenderState.TempoMapMs);
+	ClockComponent->CurrentVideoRenderSongPos = CalculateSongPosAtMsForLoopingOrMonotonicClock(VideoRenderState.TempoMapMs, VideoRenderState.LocalTick);
+}
+
 FMetasoundMusicClockDriver::EHistoryFailureType FMetasoundMusicClockDriver::CalculateSmoothedTick(Metasound::FSampleCount ExpectedRenderPosSampleCount, Metasound::FSampleCount LastRenderPosSampleCount,
-	float& SmoothedTick, float& CurrentSpeed, HarmonixMetasound::Analysis::FMidiClockSongPositionHistory::FReadCursor& ReadCursor, float LookBehindSeconds)
+	float& SmoothedLocalTick, float& SmoothedTempoMapTick, float& CurrentSpeed, HarmonixMetasound::Analysis::FMidiClockSongPositionHistory::FReadCursor& ReadCursor, float LookBehindSeconds)
 {
 	using namespace HarmonixMetasound::Analysis;
 
@@ -545,7 +626,8 @@ FMetasoundMusicClockDriver::EHistoryFailureType FMetasoundMusicClockDriver::Calc
 
 	if (NumHistoryAvailable == 1 || PeekNextRef->SampleCount > LookingForSampleFrame)
 	{
-		SmoothedTick = (float)PeekNextRef->UpToTick;
+		SmoothedLocalTick = (float)PeekNextRef->UpToTick;
+		SmoothedTempoMapTick = (float)PeekNextRef->TempoMapTick;
 		CurrentSpeed = PeekNextRef->CurrentSpeed;
 		return EHistoryFailureType::None;
 	}
@@ -569,8 +651,8 @@ FMetasoundMusicClockDriver::EHistoryFailureType FMetasoundMusicClockDriver::Calc
 
 	// We now have enough to lerp!
 	float LerpAlpha = (float)(LookingForSampleFrame - PeekNextRef->SampleCount) / (float)(PeekOneAheadRef->SampleCount - PeekNextRef->SampleCount);
-	SmoothedTick = FMath::Lerp<float>((float)PeekNextRef->UpToTick, (float)PeekOneAheadRef->UpToTick, LerpAlpha);
-
+	SmoothedLocalTick = FMath::Lerp<float>((float)PeekNextRef->UpToTick, (float)PeekOneAheadRef->UpToTick, LerpAlpha);
+	SmoothedTempoMapTick = FMath::Lerp<float>((float)PeekNextRef->TempoMapTick, (float)PeekOneAheadRef->TempoMapTick, LerpAlpha);
 	CurrentSpeed = PeekNextRef->CurrentSpeed;
 
 	return EHistoryFailureType::None;
diff --git a/Engine/Plugins/Runtime/Harmonix/Source/HarmonixMetasound/Private/HarmonixMetasound/Components/MetasoundMusicClockDriver.h b/Engine/Plugins/Runtime/Harmonix/Source/HarmonixMetasound/Private/HarmonixMetasound/Components/MetasoundMusicClockDriver.h
--- a/Engine/Plugins/Runtime/Harmonix/Source/HarmonixMetasound/Private/HarmonixMetasound/Components/MetasoundMusicClockDriver.h
+++ b/Engine/Plugins/Runtime/Harmonix/Source/HarmonixMetasound/Private/HarmonixMetasound/Components/MetasoundMusicClockDriver.h
@@ -67,12 +67,27 @@ private:
 	FLocalMinimumMagnitudeTracker<double, kFramesOfErrorHistory> ErrorTracker;
 	double SyncSpeed = 1.0;
 
+	struct FPerTimebaseSmoothedClockState
+	{
+		float TempoMapMs = 0.0f;
+		float TempoMapTick = 0.0f;
+		float LocalTick = 0.0f;
+	};
+
+	FPerTimebaseSmoothedClockState AudioRenderState;
+	FPerTimebaseSmoothedClockState PlayerExperienceState;
+	FPerTimebaseSmoothedClockState VideoRenderState;
+
 	FDelegateHandle GeneratorAttachedCallbackHandle;
 	FDelegateHandle GeneratorDetachedCallbackHandle;
 	FDelegateHandle GeneratorIOUpdatedCallbackHandle;
 	FDelegateHandle GraphChangedCallbackHandle;
 
-	FMidiSongPos CalculateSongPosAtMs(float AbsoluteMs) const;
+	FMidiSongPos CalculateSongPosAtMsForLoopingOrMonotonicClock(float AbsoluteMs, float& PositionTick) const;
+	FMidiSongPos CalculateSongPosAtMsForOffsetClock(float PositionMs, float ClockTickOffsetFromDrivingClock, float& PositionTick) const;
+
+	void UpdateCurrentTicksForOffsetClock(float SmoothedTick, float SmoothedTempoMapTick);
+	void UpdateCurrentTicksForLoopingOrMonotonicClock(float SmoothedTick, float SmoothedTempoMapTick);
 
 	bool AttemptToConnectToAudioComponentsMetasound();
 	void DetachAllCallbacks();
@@ -89,7 +104,7 @@ private:
 	};
 
 	EHistoryFailureType CalculateSmoothedTick(Metasound::FSampleCount ExpectedRenderPosSampleCount, Metasound::FSampleCount LastRenderPosSampleCount,
-		float& SmoothedTick, float& CurrentSpeed, HarmonixMetasound::Analysis::FMidiClockSongPositionHistory::FReadCursor& ReadCursor,
+		float& SmoothedTick, float& SmoothedTempoMapTick, float& CurrentSpeed, HarmonixMetasound::Analysis::FMidiClockSongPositionHistory::FReadCursor& ReadCursor,
 		float LookBehindSeconds);
 
 	static FString HistoryFailureTypeToString(FMetasoundMusicClockDriver::EHistoryFailureType Error);
diff --git a/Engine/Plugins/Runtime/Harmonix/Source/HarmonixMetasound/Private/HarmonixMetasound/DataTypes/MidiClock.cpp b/Engine/Plugins/Runtime/Harmonix/Source/HarmonixMetasound/Private/HarmonixMetasound/DataTypes/MidiClock.cpp
--- a/Engine/Plugins/Runtime/Harmonix/Source/HarmonixMetasound/Private/HarmonixMetasound/DataTypes/MidiClock.cpp
+++ b/Engine/Plugins/Runtime/Harmonix/Source/HarmonixMetasound/Private/HarmonixMetasound/DataTypes/MidiClock.cpp
@@ -49,6 +49,7 @@ namespace HarmonixMetasound
 		, FirstTickProcessedThisBlock(-1)
 		, LastProcessedMidiTick(-1)
 		, NextMidiTickToProcess(0)
+		, NextTempoMapTickToProcess(0)
 		, SampleRate(InSettings.GetSampleRate())
 		, SampleCount(0)
 		, FramesUntilNextProcess(0)
@@ -87,6 +88,7 @@ namespace HarmonixMetasound
 		, FirstTickProcessedThisBlock(Other.FirstTickProcessedThisBlock)
 		, LastProcessedMidiTick(Other.LastProcessedMidiTick)
 		, NextMidiTickToProcess(Other.NextMidiTickToProcess)
+		, NextTempoMapTickToProcess(Other.NextTempoMapTickToProcess)
 		, SampleRate(Other.SampleRate)
 		, SampleCount(Other.SampleCount)
 		, FramesUntilNextProcess(Other.FramesUntilNextProcess)
@@ -132,6 +134,7 @@ namespace HarmonixMetasound
 			FirstTickProcessedThisBlock = Other.FirstTickProcessedThisBlock;
 			LastProcessedMidiTick = Other.LastProcessedMidiTick;
 			NextMidiTickToProcess = Other.NextMidiTickToProcess;
+			NextTempoMapTickToProcess = Other.NextTempoMapTickToProcess;
 			SampleRate = Other.SampleRate;
 			SampleCount = Other.SampleCount;
 			FramesUntilNextProcess = Other.FramesUntilNextProcess;
@@ -181,7 +184,7 @@ namespace HarmonixMetasound
 
 		if (ResetToStart)
 		{
-			SeekTo(CurrentBlockFrameIndex, 0);
+			SeekTo(CurrentBlockFrameIndex, 0, 0);
 		}
 		PostTempoOrTimeSignatureEventsIfNeeded();
 	}
@@ -249,14 +252,14 @@ namespace HarmonixMetasound
 		AddSpeedChangeToBlock(BlockFrameIndex, Speed, true);
 	}
 
-	void FMidiClock::SetTempo(int32 BlockFrameIndex, int32 Tick, float Bpm)
+	void FMidiClock::SetTempo(int32 BlockFrameIndex, int32 Tick, float Bpm, int32 TempoMapTick)
 	{
-		AddTempoChangeToBlock(BlockFrameIndex, Tick, Bpm);
+		AddTempoChangeToBlock(BlockFrameIndex, Tick, Bpm, TempoMapTick);
 	}
 
-	void HarmonixMetasound::FMidiClock::SetTimeSignature(int32 BlockFrameIndex, int32 Tick, const FTimeSignature& TimeSignature)
+	void HarmonixMetasound::FMidiClock::SetTimeSignature(int32 BlockFrameIndex, int32 Tick, const FTimeSignature& TimeSignature, int32 TempoMapTick)
 	{
-		AddTimeSignatureChangeToBlock(BlockFrameIndex, Tick, TimeSignature);
+		AddTimeSignatureChangeToBlock(BlockFrameIndex, Tick, TimeSignature, TempoMapTick);
 	}
 
 	void FMidiClock::SeekTo(int32 BlockFrameIndex, const FMusicSeekTarget& InTarget)
@@ -275,21 +278,21 @@ namespace HarmonixMetasound
 			break;
 		}
 
-		SeekTo(BlockFrameIndex, Tick);
+		SeekTo(BlockFrameIndex, Tick, NextTempoMapTickToProcess);
 	}
 
-	void FMidiClock::SeekTo(int32 BlockFrameIndex, int32 Tick)
+	void FMidiClock::SeekTo(int32 BlockFrameIndex, int32 Tick, int32 TempoMapTick)
 	{
 		if (NextMidiTickToProcess != Tick)
 		{
-			AddSeekToBlock(BlockFrameIndex, Tick);
+			AddSeekToBlock(BlockFrameIndex, Tick, TempoMapTick);
 			SampleCount = FMath::Max<FSampleCount>(FSampleCount(SongMapEvaluator->TickToMs(NextMidiTickToProcess) / 1000.0f * SampleRate), 0);
 		}
 	}
 
 	void FMidiClock::AddTransientLoop(int32 BlockFrameIndex, int32 NewFirstTickInLoop, int32 NewLoopLengthTicks)
 	{
-		AddLoopToBlock(BlockFrameIndex, NewFirstTickInLoop, NewLoopLengthTicks);
+		AddLoopToBlock(BlockFrameIndex, NewFirstTickInLoop, NewLoopLengthTicks, NextTempoMapTickToProcess);
 	}
 
 	void FMidiClock::SetupPersistentLoop(int32 NewFirstTickInLoop, int32 NewLoopLengthTicks)
@@ -345,9 +348,10 @@ namespace HarmonixMetasound
 		if (ClockEvents.IsValidIndex(0) && NeedsSeekToDrivingClock)
 		{
 			NeedsSeekToDrivingClock = false;
-			int32 OurStartTick = DrivingClock.GetNextTickToProcessAtBlockFrame(StartFrame);
-			if (OurStartTick >= 0)
+			int32 DrivingClocksTick = DrivingClock.GetNextTickToProcessAtBlockFrame(StartFrame);
+			if (DrivingClocksTick >= 0)
 			{
+				int32 OurStartTick = DrivingClocksTick;
 				if (!FMath::IsNearlyEqual(CurrentLocalSpeed, 1.0f, 0.0001))
 				{
 					float FractionalToTick = (float)OurStartTick * CurrentLocalSpeed;
@@ -355,7 +359,7 @@ namespace HarmonixMetasound
 					TickResidualWhenDriven = FMath::Fractional(FractionalToTick);
 				}
 				OurStartTick = WrapTickIfLooping(OurStartTick);
-				SeekTo(ClockEvents[0].BlockFrameIndex, OurStartTick);
+				SeekTo(CurrentBlockFrameIndex, OurStartTick, DrivingClocksTick);
 			}
 		}
 
@@ -402,7 +406,7 @@ namespace HarmonixMetasound
 		FramesUntilNextProcess -= NumFrames;
 	}
 
-	bool FMidiClock::AdvanceToTick(int32 BlockFrameIndex, int32 UpToTick)
+	bool FMidiClock::AdvanceToTick(int32 BlockFrameIndex, int32 UpToTick, int32 TempoMapTick)
 	{
 		using namespace MidiClockMessageTypes;
 
@@ -414,8 +418,10 @@ namespace HarmonixMetasound
 			// first we might need advance to the loop end...
 			if (NextMidiTickToProcess > FirstTickInLoop)
 			{
-				AddAdvanceToBlock(BlockFrameIndex, NextMidiTickToProcess, TickAfterLoop - NextMidiTickToProcess);
-				AddLoopToBlock(BlockFrameIndex, FirstTickInLoop, LoopLengthTicks);
+				int32 TicksToAdvance = TickAfterLoop - NextMidiTickToProcess;
+				AddAdvanceToBlock(BlockFrameIndex, NextMidiTickToProcess, TicksToAdvance, TempoMapTick);
+				TempoMapTick += TicksToAdvance;
+				AddLoopToBlock(BlockFrameIndex, FirstTickInLoop, LoopLengthTicks, TempoMapTick);
 				check(NextMidiTickToProcess == FirstTickInLoop);
 				bDidLoop = true;
 			}
@@ -425,10 +431,11 @@ namespace HarmonixMetasound
 			{
 				int32 TicksThisPass = FMath::Min(LoopLengthTicks, NumTicksLeftToProcess);
 				NumTicksLeftToProcess -= TicksThisPass;
-				AddAdvanceToBlock(BlockFrameIndex, NextMidiTickToProcess, TicksThisPass);
+				AddAdvanceToBlock(BlockFrameIndex, NextMidiTickToProcess, TicksThisPass, TempoMapTick);
+				TempoMapTick += TicksThisPass;
 				if (NumTicksLeftToProcess > 0)
 				{
-					AddLoopToBlock(BlockFrameIndex, FirstTickInLoop, LoopLengthTicks);
+					AddLoopToBlock(BlockFrameIndex, FirstTickInLoop, LoopLengthTicks, TempoMapTick);
 					check(NextMidiTickToProcess == FirstTickInLoop);
 					bDidLoop = true;
 				}
@@ -436,10 +443,12 @@ namespace HarmonixMetasound
 		}
 		else
 		{
-			AddAdvanceToBlock(BlockFrameIndex, NextMidiTickToProcess, UpToTick - NextMidiTickToProcess);
+			int32 NumTicks = UpToTick - NextMidiTickToProcess;
+			AddAdvanceToBlock(BlockFrameIndex, NextMidiTickToProcess, NumTicks, TempoMapTick);
 			if (HasPersistentLoop() && NextMidiTickToProcess >= (FirstTickInLoop + LoopLengthTicks))
 			{
-				AddLoopToBlock(BlockFrameIndex, FirstTickInLoop, LoopLengthTicks);
+				TempoMapTick += NumTicks;
+				AddLoopToBlock(BlockFrameIndex, FirstTickInLoop, LoopLengthTicks, TempoMapTick);
 				check(NextMidiTickToProcess == FirstTickInLoop);
 				bDidLoop = true;
 			}
@@ -450,6 +459,11 @@ namespace HarmonixMetasound
 
 	bool FMidiClock::AdvanceToMs(int32 BlockFrameIndex, float Ms)
 	{
+		if (!ensureMsgf(!ExternalClockDriver, TEXT("Can't Advance an FMidiClock by Ms when it is being driven by another clock!")))
+		{
+			return false;
+		}
+
 		using namespace MidiClockMessageTypes;
 
 		bool bDidLoop = false;
@@ -468,10 +482,10 @@ namespace HarmonixMetasound
 			// first we might need advance to the loop end...
 			if (NextMidiTickToProcess < TickAfterLoop)
 			{
-				AddAdvanceToBlock(BlockFrameIndex, NextMidiTickToProcess, TickAfterLoop - NextMidiTickToProcess);
+				AddAdvanceToBlock(BlockFrameIndex, NextMidiTickToProcess, TickAfterLoop - NextMidiTickToProcess, NextMidiTickToProcess);
 			}
 
-			AddLoopToBlock(BlockFrameIndex, FirstTickInLoop, LoopLengthTicks);
+			AddLoopToBlock(BlockFrameIndex, FirstTickInLoop, LoopLengthTicks, FirstTickInLoop);
 			check(NextMidiTickToProcess == FirstTickInLoop);
 			bDidLoop = true;
 
@@ -482,14 +496,14 @@ namespace HarmonixMetasound
 				ToFutureTick = SongMapEvaluator->MsToTick(GetLoopStartMs() + MsRemainingAfterProcessingToLoopEnd);
 				if (ToFutureTick > TickAfterLoop)
 				{
-					AddAdvanceToBlock(BlockFrameIndex, NextMidiTickToProcess, TickAfterLoop - NextMidiTickToProcess);
-					AddLoopToBlock(BlockFrameIndex, FirstTickInLoop, LoopLengthTicks);
+					AddAdvanceToBlock(BlockFrameIndex, NextMidiTickToProcess, TickAfterLoop - NextMidiTickToProcess, NextMidiTickToProcess);
+					AddLoopToBlock(BlockFrameIndex, FirstTickInLoop, LoopLengthTicks, FirstTickInLoop);
 					check(NextMidiTickToProcess == FirstTickInLoop);
 					MsRemainingAfterProcessingToLoopEnd -= GetLoopLengthMs();
 				}
 				else
 				{
-					AddAdvanceToBlock(BlockFrameIndex, NextMidiTickToProcess, ToFutureTick - NextMidiTickToProcess);
+					AddAdvanceToBlock(BlockFrameIndex, NextMidiTickToProcess, ToFutureTick - NextMidiTickToProcess, NextMidiTickToProcess);
 					MsRemainingAfterProcessingToLoopEnd = 0.0f;
 				}
 			}
@@ -501,10 +515,10 @@ namespace HarmonixMetasound
 		{
 			if (ToFutureTick > NextMidiTickToProcess)
 			{
-				AddAdvanceToBlock(BlockFrameIndex, NextMidiTickToProcess, ToFutureTick - NextMidiTickToProcess);
+				AddAdvanceToBlock(BlockFrameIndex, NextMidiTickToProcess, ToFutureTick - NextMidiTickToProcess, NextMidiTickToProcess);
 				if (HasPersistentLoop() && LastProcessedMidiTick == TickAfterLoop)
 				{
-					AddLoopToBlock(BlockFrameIndex, FirstTickInLoop, LoopLengthTicks);
+					AddLoopToBlock(BlockFrameIndex, FirstTickInLoop, LoopLengthTicks, FirstTickInLoop);
 					check(NextMidiTickToProcess == FirstTickInLoop);
 					bDidLoop = true;
 				}
@@ -538,6 +552,11 @@ namespace HarmonixMetasound
 
 	float FMidiClock::GetTempoAtBlockSampleFrame(int32 FrameIndex) const
 	{
+		if (ExternalClockDriver)
+		{
+			return ExternalClockDriver->GetTempoAtBlockSampleFrame(FrameIndex);
+		}
+
 		if (NumTempoChangeInBlock == 0 || FrameIndex == 0)
 		{
 			return TempoAtBlockStart;
@@ -619,15 +638,15 @@ namespace HarmonixMetasound
 				FoundNextTick = AsSeek->NewNextTick;
 				AtBlockIndex = Event.BlockFrameIndex;
 			}
-			else if (const FAdvance* AsProcess = Event.TryGet<FAdvance>())
+			else if (const FAdvance* AsAdvance = Event.TryGet<FAdvance>())
 			{
 				if (Event.BlockFrameIndex == BlockFrame)
 				{
-					FoundNextTick = AsProcess->FirstTickToProcess;
+					FoundNextTick = AsAdvance->FirstTickToProcess;
 				}
 				else
 				{
-					FoundNextTick = AsProcess->FirstTickToProcess + AsProcess->NumberOfTicksToProcess;
+					FoundNextTick = AsAdvance->FirstTickToProcess + AsAdvance->NumberOfTicksToProcess;
 				}
 				AtBlockIndex = Event.BlockFrameIndex;
 			}
@@ -668,13 +687,13 @@ namespace HarmonixMetasound
 		if (const FLoop* AsLoop = Event.TryGet<FLoop>())
 		{
 			const int32 Tick = WrapTickIfLooping(AsLoop->FirstTickInLoop);
-			SeekTo(Event.BlockFrameIndex, Tick);
+			SeekTo(Event.BlockFrameIndex, Tick, AsLoop->TempoMapTick);
 			TickResidualWhenDriven = 0.0f;
 		}
 		else if (const FSeek* AsSeek = Event.TryGet<FSeek>())
 		{
 			const int32 Tick = WrapTickIfLooping(AsSeek->NewNextTick);
-			SeekTo(Event.BlockFrameIndex, Tick);
+			SeekTo(Event.BlockFrameIndex, Tick, AsSeek->TempoMapTick);
 			TickResidualWhenDriven = 0.0f;
 		}
 		else if (const FAdvance* AsAdvance = Event.TryGet<FAdvance>())
@@ -692,11 +711,11 @@ namespace HarmonixMetasound
 				TickResidualWhenDriven = FMath::Fractional(FractionalToTick);
 			}
 			// No need to wrap the tick here because AdvanceToTick will handle that.
-			AdvanceToTick(Event.BlockFrameIndex, UpToTick);
+			AdvanceToTick(Event.BlockFrameIndex, UpToTick, AsAdvance->TempoMapTick);
 		}
 		else if (const FTempoChange* AsTempoChange = Event.TryGet<FTempoChange>())
 		{
-			AddTempoChangeToBlock(Event.BlockFrameIndex, NextMidiTickToProcess, AsTempoChange->Tempo);
+			AddTempoChangeToBlock(Event.BlockFrameIndex, NextMidiTickToProcess, AsTempoChange->Tempo, AsTempoChange->TempoMapTick);
 		}
 		else if (const FSpeedChange* AsSpeedChange = Event.TryGet<FSpeedChange>())
 		{
@@ -710,20 +729,20 @@ namespace HarmonixMetasound
 
 	void FMidiClock::PostTempoOrTimeSignatureEventsIfNeeded()
 	{
-		if (SongMapEvaluator->GetNumTempoChanges() == 0)
+		if (SongMapEvaluator->GetNumTempoChanges() == 0 || ExternalClockDriver)
 		{
 			CurrentTempoInfoPointIndex = -1;
-			AddTempoChangeToBlock(CurrentBlockFrameIndex, NextMidiTickToProcess, 120.0f);
+			if (!ExternalClockDriver)
+			{
+				AddTempoChangeToBlock(CurrentBlockFrameIndex, NextMidiTickToProcess, 120.0f, NextMidiTickToProcess);
+			}
 			NextTempoChangeTick = std::numeric_limits<int32>::max();
 		}
 		else
 		{
 			CurrentTempoInfoPointIndex = SongMapEvaluator->GetTempoPointIndexForTick(NextMidiTickToProcess > 0 ? NextMidiTickToProcess : 0);
 			check(CurrentTempoInfoPointIndex >= 0 && CurrentTempoInfoPointIndex < SongMapEvaluator->GetNumTempoChanges());
-			if (!ExternalClockDriver)
-			{
-				AddTempoChangeToBlock(CurrentBlockFrameIndex, NextMidiTickToProcess, SongMapEvaluator->GetTempoInfoPoint(CurrentTempoInfoPointIndex)->GetBPM());
-			}
+			AddTempoChangeToBlock(CurrentBlockFrameIndex, NextMidiTickToProcess, SongMapEvaluator->GetTempoInfoPoint(CurrentTempoInfoPointIndex)->GetBPM(), NextMidiTickToProcess);
 			if ((CurrentTempoInfoPointIndex + 1) < SongMapEvaluator->GetNumTempoChanges())
 			{
 				NextTempoChangeTick = SongMapEvaluator->GetTempoChangePointTick(CurrentTempoInfoPointIndex + 1);
@@ -737,14 +756,14 @@ namespace HarmonixMetasound
 		if (SongMapEvaluator->GetNumTimeSignatureChanges() == 0)
 		{
 			CurrentTimeSignaturePointIndex = -1;
-			AddTimeSignatureChangeToBlock(CurrentBlockFrameIndex, NextMidiTickToProcess, FTimeSignature(4,4));
+			AddTimeSignatureChangeToBlock(CurrentBlockFrameIndex, NextMidiTickToProcess, FTimeSignature(4,4), NextTempoMapTickToProcess);
 			NextTimeSigChangeTick = std::numeric_limits<int32>::max();
 		}
 		else
 		{
 			CurrentTimeSignaturePointIndex = SongMapEvaluator->GetTimeSignaturePointIndexForTick(NextMidiTickToProcess > 0 ? NextMidiTickToProcess : 0);
 			check(CurrentTimeSignaturePointIndex >= 0 && CurrentTimeSignaturePointIndex < SongMapEvaluator->GetNumTimeSignatureChanges());
-			AddTimeSignatureChangeToBlock(CurrentBlockFrameIndex, NextMidiTickToProcess, SongMapEvaluator->GetTimeSignaturePoint(CurrentTimeSignaturePointIndex)->TimeSignature);
+			AddTimeSignatureChangeToBlock(CurrentBlockFrameIndex, NextMidiTickToProcess, SongMapEvaluator->GetTimeSignaturePoint(CurrentTimeSignaturePointIndex)->TimeSignature, NextTempoMapTickToProcess);
 			if ((CurrentTimeSignaturePointIndex + 1) < SongMapEvaluator->GetNumTimeSignatureChanges())
 			{
 				NextTimeSigChangeTick = SongMapEvaluator->GetTimeSignatureChangePointTick(CurrentTempoInfoPointIndex + 1);
@@ -791,7 +810,7 @@ namespace HarmonixMetasound
 		TransportAtBlockEnd = NewTransportState;
 	}
 
-	void FMidiClock::AddTimeSignatureChangeToBlock(int32 BlockFrameIndex, int32 Tick, const FTimeSignature& TimeSignature)
+	void FMidiClock::AddTimeSignatureChangeToBlock(int32 BlockFrameIndex, int32 Tick, const FTimeSignature& TimeSignature, int32 TempoMapTick)
 	{
 		check(BlockFrameIndex >= CurrentBlockFrameIndex);
 
@@ -804,6 +823,7 @@ namespace HarmonixMetasound
 
 		CurrentBlockFrameIndex = BlockFrameIndex;
 		NextMidiTickToProcess = Tick;
+		NextTempoMapTickToProcess = TempoMapTick;
 
 		if (TimeSignatureAtBlockEnd == TimeSignature)
 		{
@@ -814,10 +834,11 @@ namespace HarmonixMetasound
 		if (MidiClockMessageTypes::FTimeSignatureChange* PreviousTimeSigChangeOnSameTick = LookForEventOnMidiTick<MidiClockMessageTypes::FTimeSignatureChange>(BlockFrameIndex))
 		{
 			PreviousTimeSigChangeOnSameTick->TimeSignature = TimeSignature;
+			PreviousTimeSigChangeOnSameTick->TempoMapTick = TempoMapTick;
 		}
 		else
 		{
-			AddEvent(FMidiClockEvent(BlockFrameIndex, MidiClockMessageTypes::FTimeSignatureChange(Tick, FTimeSignature(TimeSignature))));
+			AddEvent(FMidiClockEvent(BlockFrameIndex, MidiClockMessageTypes::FTimeSignatureChange(Tick, FTimeSignature(TimeSignature), TempoMapTick)));
 			NumTimeSignatureChangeInBlock++;
 		}
 
@@ -829,7 +850,7 @@ namespace HarmonixMetasound
 		TimeSignatureAtBlockEnd = TimeSignature;
 	}
 
-	void FMidiClock::AddTempoChangeToBlock(int32 BlockFrameIndex, int32 Tick, float NewTempo)
+	void FMidiClock::AddTempoChangeToBlock(int32 BlockFrameIndex, int32 Tick, float NewTempo, int32 TempoMapTick)
 	{
 		check(BlockFrameIndex >= CurrentBlockFrameIndex);
 
@@ -837,6 +858,7 @@ namespace HarmonixMetasound
 
 		CurrentBlockFrameIndex = BlockFrameIndex;
 		NextMidiTickToProcess = Tick;
+		NextTempoMapTickToProcess = TempoMapTick;
 
 		if (FirstTickProcessedThisBlock == -1)
 		{
@@ -852,10 +874,11 @@ namespace HarmonixMetasound
 		if (MidiClockMessageTypes::FTempoChange* PreviousTempoChangeOnSameTick = LookForEventOnMidiTick<MidiClockMessageTypes::FTempoChange>(BlockFrameIndex))
 		{
 			PreviousTempoChangeOnSameTick->Tempo = NewTempo;
+			PreviousTempoChangeOnSameTick->TempoMapTick = TempoMapTick;
 		}
 		else
 		{
-			AddEvent(FMidiClockEvent(BlockFrameIndex, MidiClockMessageTypes::FTempoChange(Tick, NewTempo)));
+			AddEvent(FMidiClockEvent(BlockFrameIndex, MidiClockMessageTypes::FTempoChange(Tick, NewTempo, TempoMapTick)));
 			NumTempoChangeInBlock++;
 		}
 
@@ -910,7 +933,7 @@ namespace HarmonixMetasound
 		SpeedAtBlockEnd = NewSpeed;
 	}
 
-	void FMidiClock::AddLoopToBlock(int32 BlockFrameIndex, int32 FirstTick, int32 LoopLength)
+	void FMidiClock::AddLoopToBlock(int32 BlockFrameIndex, int32 FirstTick, int32 LoopLength, int32 TempoMapTick)
 	{
 		check(BlockFrameIndex >= CurrentBlockFrameIndex);
 		check(LastProcessedMidiTick == FirstTick + LoopLength - 1);
@@ -926,19 +949,21 @@ namespace HarmonixMetasound
 		{
 			PreviousLoopOnSameTick->LengthInTicks = LoopLength;
 			PreviousLoopOnSameTick->FirstTickInLoop = FirstTick;
+			PreviousLoopOnSameTick->TempoMapTick = TempoMapTick;
 		}
 		else
 		{
-			AddEvent(FMidiClockEvent(BlockFrameIndex, MidiClockMessageTypes::FLoop(FirstTick, LoopLength)));
+			AddEvent(FMidiClockEvent(BlockFrameIndex, MidiClockMessageTypes::FLoop(FirstTick, LoopLength, TempoMapTick)));
 		}
 
-		NextMidiTickToProcess = FirstTickInLoop;
+		NextMidiTickToProcess = FirstTick;
+		NextTempoMapTickToProcess = TempoMapTick;
 
 		// The tempo and/or time signature may be different in the location where we are going... so... update...
 		PostTempoOrTimeSignatureEventsIfNeeded();
 	}
 
-	void FMidiClock::AddSeekToBlock(int32 BlockFrameIndex, int32 ToTick)
+	void FMidiClock::AddSeekToBlock(int32 BlockFrameIndex, int32 ToTick, int32 TempoMapTick)
 	{
 		check(BlockFrameIndex >= CurrentBlockFrameIndex);
 
@@ -953,19 +978,21 @@ namespace HarmonixMetasound
 		{
 			check(PreviousSeekOnSameTick->LastTickProcessedBeforeSeek == LastProcessedMidiTick);
 			PreviousSeekOnSameTick->NewNextTick = ToTick;
+			PreviousSeekOnSameTick->TempoMapTick = TempoMapTick;
 		}
 		else
 		{
-			AddEvent(FMidiClockEvent(BlockFrameIndex, MidiClockMessageTypes::FSeek(LastProcessedMidiTick, ToTick)));
+			AddEvent(FMidiClockEvent(BlockFrameIndex, MidiClockMessageTypes::FSeek(LastProcessedMidiTick, ToTick, TempoMapTick)));
 		}
 
 		NextMidiTickToProcess = ToTick;
+		NextTempoMapTickToProcess = TempoMapTick;
 
 		// The tempo and/or time signature may be different in the location where we are going... so... update...
 		PostTempoOrTimeSignatureEventsIfNeeded();
 	}
 
-	void FMidiClock::AddAdvanceToBlock(int32 BlockFrameIndex, int32 FirstTick, int32 NumTicks)
+	void FMidiClock::AddAdvanceToBlock(int32 BlockFrameIndex, int32 FirstTick, int32 NumTicks, int32 TempoMapTick)
 	{
 		check(BlockFrameIndex >= CurrentBlockFrameIndex);
 		check(FirstTick == NextMidiTickToProcess);
@@ -984,7 +1011,8 @@ namespace HarmonixMetasound
 			int32 SpanNumTicks = NextTempoOrTimeSigChangeTick - FirstTick;
 			if (SpanNumTicks > 0)
 			{
-				AddEvent(FMidiClockEvent(BlockFrameIndex, MidiClockMessageTypes::FAdvance(FirstTick, SpanNumTicks)));
+				AddEvent(FMidiClockEvent(BlockFrameIndex, MidiClockMessageTypes::FAdvance(FirstTick, SpanNumTicks, TempoMapTick)));
+				TempoMapTick += SpanNumTicks;
 				NextMidiTickToProcess = FirstTick + SpanNumTicks;
 				LastProcessedMidiTick = NextMidiTickToProcess - 1;
 				FirstTick = NextMidiTickToProcess;
@@ -1000,7 +1028,7 @@ namespace HarmonixMetasound
 				check(TempoPoint);
 				if (!ExternalClockDriver)
 				{
-					AddTempoChangeToBlock(BlockFrameIndex, TempoPoint->StartTick, TempoPoint->GetBPM());
+					AddTempoChangeToBlock(BlockFrameIndex, TempoPoint->StartTick, TempoPoint->GetBPM(), TempoMapTick);
 				}
 				int32 NextTempoInfoPointIndex = CurrentTempoInfoPointIndex + 1;
 				if (NextTempoInfoPointIndex < SongMapEvaluator->GetNumTempoChanges())
@@ -1020,7 +1048,7 @@ namespace HarmonixMetasound
 				CurrentTimeSignaturePointIndex++;
 				const FTimeSignaturePoint* TimeSignaturePoint = SongMapEvaluator->GetTimeSignaturePoint(CurrentTimeSignaturePointIndex);
 				check(TimeSignaturePoint);
-				AddTimeSignatureChangeToBlock(BlockFrameIndex, TimeSignaturePoint->StartTick, TimeSignaturePoint->TimeSignature);
+				AddTimeSignatureChangeToBlock(BlockFrameIndex, TimeSignaturePoint->StartTick, TimeSignaturePoint->TimeSignature, TempoMapTick);
 				int32 NextTimeSignaturePointIndex = CurrentTimeSignaturePointIndex + 1;
 				if (NextTimeSignaturePointIndex < SongMapEvaluator->GetNumTimeSignatureChanges())
 				{
@@ -1038,10 +1066,13 @@ namespace HarmonixMetasound
 
 		if (NumTicks > 0)
 		{
-			AddEvent(FMidiClockEvent(BlockFrameIndex, MidiClockMessageTypes::FAdvance(FirstTick, NumTicks)));
+			AddEvent(FMidiClockEvent(BlockFrameIndex, MidiClockMessageTypes::FAdvance(FirstTick, NumTicks, TempoMapTick)));
 			NextMidiTickToProcess = FirstTick + NumTicks;
+			TempoMapTick += NumTicks;
 			LastProcessedMidiTick = NextMidiTickToProcess - 1;
 		}
+
+		NextTempoMapTickToProcess = TempoMapTick;
 	}
 
 	void FMidiClock::RebuildSongMapEvaluator(const TSharedPtr<const ISongMapEvaluator>& MapWithTempo, const TSharedPtr<const ISongMapEvaluator>& MapWithOtherMaps)
diff --git a/Engine/Plugins/Runtime/Harmonix/Source/HarmonixMetasound/Private/HarmonixMetasound/Nodes/MetronomeNode.cpp b/Engine/Plugins/Runtime/Harmonix/Source/HarmonixMetasound/Private/HarmonixMetasound/Nodes/MetronomeNode.cpp
--- a/Engine/Plugins/Runtime/Harmonix/Source/HarmonixMetasound/Private/HarmonixMetasound/Nodes/MetronomeNode.cpp
+++ b/Engine/Plugins/Runtime/Harmonix/Source/HarmonixMetasound/Private/HarmonixMetasound/Nodes/MetronomeNode.cpp
@@ -243,10 +243,10 @@ namespace HarmonixMetasound
 			switch (CurrentState)
 			{
 			case EMusicPlayerTransportState::Starting:
-				MidiClockOutPin->SeekTo(0, 0);
+				MidiClockOutPin->SeekTo(0, 0, 0);
 				if (LoopInPin)
 				{
-					MonotonicallyIncreasingClock->SeekTo(0, 0);
+					MonotonicallyIncreasingClock->SeekTo(0, 0, 0);
 				}
 				break;
 			}
@@ -304,7 +304,7 @@ namespace HarmonixMetasound
 				}
 				if (!ReceivedSeekWhileStopped())
 				{
-					DrivingMidiClock.SeekTo(StartFrameIndex, 0);
+					DrivingMidiClock.SeekTo(StartFrameIndex, 0, 0);
 				}
 				return EMusicPlayerTransportState::Playing;
 				
diff --git a/Engine/Plugins/Runtime/Harmonix/Source/HarmonixMetasound/Private/HarmonixMetasound/Nodes/MidiClockOffsetNode.cpp b/Engine/Plugins/Runtime/Harmonix/Source/HarmonixMetasound/Private/HarmonixMetasound/Nodes/MidiClockOffsetNode.cpp
--- a/Engine/Plugins/Runtime/Harmonix/Source/HarmonixMetasound/Private/HarmonixMetasound/Nodes/MidiClockOffsetNode.cpp
+++ b/Engine/Plugins/Runtime/Harmonix/Source/HarmonixMetasound/Private/HarmonixMetasound/Nodes/MidiClockOffsetNode.cpp
@@ -132,7 +132,7 @@ namespace HarmonixMetasound::Nodes::MidiClockOffset
 		{
 			BlockSize = Params.OperatorSettings.GetNumFramesPerBlock();
 			CurrentBlockSpanStart = 0;
-			MidiClockOut->SeekTo(0,0);
+			MidiClockOut->SeekTo(0,0,0);
 			MidiClockOut->SetTransportState(0, EMusicPlayerTransportState::Playing);
 			bClockOutNeedsPrepare = false;
 
@@ -181,6 +181,12 @@ namespace HarmonixMetasound::Nodes::MidiClockOffset
 			{
 				MidiClockOut->PrepareBlock();
 			}
+
+			if (MidiClockIn->GetSongMapsChangedInBlock())
+			{
+				MidiClockOut->SongMapsChanged();
+			}
+
 			// next time we definitely want to prepare the block.
 			bClockOutNeedsPrepare = true;
 
@@ -188,16 +194,16 @@ namespace HarmonixMetasound::Nodes::MidiClockOffset
 			for (int32 EventIndex = 0; EventIndex < MidiClockEvents.Num(); ++EventIndex)
 			{
 				const FMidiClockEvent& Event = MidiClockEvents[EventIndex];
-				if (Event.Msg.IsType<FSeek>())
+				if (const FSeek* AsSeek = Event.TryGet<FSeek>())
 				{
-					const int32 NewNextTick = GetTickWithOffset(Event.Msg.Get<FSeek>().NewNextTick, OffsetBars, OffsetBeats, OffsetMs);
-					MidiClockOut->SeekTo(Event.BlockFrameIndex, NewNextTick);
+					const int32 NewNextTick = GetTickWithOffset(AsSeek->NewNextTick, OffsetBars, OffsetBeats, OffsetMs);
+					MidiClockOut->SeekTo(Event.BlockFrameIndex, NewNextTick, AsSeek->TempoMapTick);
 				}
 				else if (const FLoop* AsLoop = Event.TryGet<FLoop>())
 				{
 					const int32 FirstTickInLoop = GetTickWithOffset(AsLoop->FirstTickInLoop, OffsetBars, OffsetBeats, OffsetMs);
 					const int32 LastTickAfterLoop = GetTickWithOffset(AsLoop->FirstTickInLoop + AsLoop->LengthInTicks, OffsetBars, OffsetBeats, OffsetMs);
-					MidiClockOut->AdvanceToTick(Event.BlockFrameIndex, LastTickAfterLoop);
+					MidiClockOut->AdvanceToTick(Event.BlockFrameIndex, LastTickAfterLoop, AsLoop->TempoMapTick);
 					MidiClockOut->AddTransientLoop(Event.BlockFrameIndex, FirstTickInLoop, LastTickAfterLoop - FirstTickInLoop);
 				}
 				else if (const FAdvance* AsAdvance = Event.TryGet<FAdvance>())
@@ -215,19 +221,19 @@ namespace HarmonixMetasound::Nodes::MidiClockOffset
 					// just want to advance from where we sit to the appropriate destination.
 					if (!bAdvancedSinceTransportChange && FirstickToProcess != MidiClockOut->GetNextMidiTickToProcess())
 					{
-						MidiClockOut->SeekTo(Event.BlockFrameIndex, FirstickToProcess);
+						MidiClockOut->SeekTo(Event.BlockFrameIndex, FirstickToProcess, AsAdvance->TempoMapTick);
 					}
 					const int32 ProcessUpToTick = GetTickWithOffset(AsAdvance->FirstTickToProcess + AsAdvance->NumberOfTicksToProcess, OffsetBars, OffsetBeats, OffsetMs);
-					MidiClockOut->AdvanceToTick(Event.BlockFrameIndex, ProcessUpToTick);
+					MidiClockOut->AdvanceToTick(Event.BlockFrameIndex, ProcessUpToTick, AsAdvance->TempoMapTick);
 					bAdvancedSinceTransportChange = true;
 				}
 				else if (const MidiClockMessageTypes::FTempoChange* AsTempoChange = Event.TryGet<MidiClockMessageTypes::FTempoChange>())
 				{
-					MidiClockOut->SetTempo(Event.BlockFrameIndex, MidiClockOut->GetNextMidiTickToProcess(), AsTempoChange->Tempo);
+					MidiClockOut->SetTempo(Event.BlockFrameIndex, MidiClockOut->GetNextMidiTickToProcess(), AsTempoChange->Tempo, AsTempoChange->TempoMapTick);
 				}
 				else if (const MidiClockMessageTypes::FTimeSignatureChange* AsTimeSigChange = Event.TryGet<MidiClockMessageTypes::FTimeSignatureChange>())
 				{
-					MidiClockOut->SetTimeSignature(Event.BlockFrameIndex, MidiClockOut->GetNextMidiTickToProcess(), AsTimeSigChange->TimeSignature);
+					MidiClockOut->SetTimeSignature(Event.BlockFrameIndex, MidiClockOut->GetNextMidiTickToProcess(), AsTimeSigChange->TimeSignature, AsTimeSigChange->TempoMapTick);
 				}
 				else if (const MidiClockMessageTypes::FSpeedChange* AsSpeedChange = Event.TryGet<MidiClockMessageTypes::FSpeedChange>())
 				{
diff --git a/Engine/Plugins/Runtime/Harmonix/Source/HarmonixMetasound/Private/HarmonixMetasound/Nodes/MidiPlayerNode.cpp b/Engine/Plugins/Runtime/Harmonix/Source/HarmonixMetasound/Private/HarmonixMetasound/Nodes/MidiPlayerNode.cpp
--- a/Engine/Plugins/Runtime/Harmonix/Source/HarmonixMetasound/Private/HarmonixMetasound/Nodes/MidiPlayerNode.cpp
+++ b/Engine/Plugins/Runtime/Harmonix/Source/HarmonixMetasound/Private/HarmonixMetasound/Nodes/MidiPlayerNode.cpp
@@ -335,7 +335,7 @@ namespace HarmonixMetasound
 				MidiClockOut->SetTransportState(0, EMusicPlayerTransportState::Playing);
 				if (!ReceivedSeekWhileStopped())
 				{
-					MidiClockOut->SeekTo(0, 0);
+					MidiClockOut->SeekTo(0, 0, 0);
 				}
 				return EMusicPlayerTransportState::Playing;
 
@@ -382,6 +382,11 @@ namespace HarmonixMetasound
 	{
 		FMidiPlayerOperator::Execute();
 
+		if (MidiClockIn->GetSongMapsChangedInBlock())
+		{
+			MidiClockOut->SongMapsChanged();
+		}
+
 		TransportSpanPostProcessor HandleMidiClockEventsInBlock = [this](int32 StartFrameIndex, int32 EndFrameIndex, EMusicPlayerTransportState CurrentState)
 		{
 			// clock should always process in post processor
@@ -401,7 +406,7 @@ namespace HarmonixMetasound
 			case EMusicPlayerTransportState::Starting:
 				if (!ReceivedSeekWhileStopped())
 				{
-					MidiClockOut->SeekTo(StartFrameIndex, 0);
+					MidiClockOut->SeekTo(StartFrameIndex, 0, MidiClockIn->GetNextTickToProcessAtBlockFrame(StartFrameIndex));
 				}
 				break;
 			case EMusicPlayerTransportState::Stopping:
@@ -464,7 +469,7 @@ namespace HarmonixMetasound
 			case EMusicPlayerTransportState::Starting:
 				if (!ReceivedSeekWhileStopped())
 				{
-					MidiClockOut->SeekTo(StartFrameIndex, 0);
+					MidiClockOut->SeekTo(StartFrameIndex, 0, 0);
 				}
 				break;
 			case EMusicPlayerTransportState::Seeking:
@@ -506,7 +511,7 @@ namespace HarmonixMetasound
 		CurrentBlockSpanStart = 0;
 
 		MidiOutPin->SetClock(*MidiClockOut);
-		MidiClockOut->SeekTo(0, 0);
+		MidiClockOut->SeekTo(0, 0, 0);
 		MidiClockOut->SetTransportState(0, EMusicPlayerTransportState::Prepared);
 
 		NeedsTransportInit = true;
@@ -605,7 +610,11 @@ namespace HarmonixMetasound
 				// remap our current tick based on the looping behavior
 				// maybe this should happen automatically when resetting a loop or changing midi files?
 				int32 NewTick = MidiClockOut->WrapTickIfLooping(MidiClockOut->GetNextMidiTickToProcess());
-				MidiClockOut->SeekTo(0, NewTick);
+				
+				// Hmmmm. Sending in 0 for the DrivingClock's tick because in this base clase we don't know
+				// if we have a driving clock. This may need to be changed (eg moved in to a subclass 'OnNewMidi' 
+				// callback function.
+				MidiClockOut->SeekTo(0, NewTick, 0);
 			}
 			else
 			{
diff --git a/Engine/Plugins/Runtime/Harmonix/Source/HarmonixMetasound/Public/HarmonixMetasound/Analysis/MidiClockSongPos.h b/Engine/Plugins/Runtime/Harmonix/Source/HarmonixMetasound/Public/HarmonixMetasound/Analysis/MidiClockSongPos.h
--- a/Engine/Plugins/Runtime/Harmonix/Source/HarmonixMetasound/Public/HarmonixMetasound/Analysis/MidiClockSongPos.h
+++ b/Engine/Plugins/Runtime/Harmonix/Source/HarmonixMetasound/Public/HarmonixMetasound/Analysis/MidiClockSongPos.h
@@ -19,6 +19,7 @@ namespace HarmonixMetasound::Analysis
 		};
 		Metasound::FSampleCount SampleCount {0};
 		int32 UpToTick;
+		int32 TempoMapTick;
 		float CurrentSpeed {1.0f};
 		EMusicPlayerTransportState CurrentTransportState {EMusicPlayerTransportState::Invalid};
 		EMarkerType MarkerType {EMarkerType::None};
diff --git a/Engine/Plugins/Runtime/Harmonix/Source/HarmonixMetasound/Public/HarmonixMetasound/Analysis/MidiSongPosVertexAnalyzer.h b/Engine/Plugins/Runtime/Harmonix/Source/HarmonixMetasound/Public/HarmonixMetasound/Analysis/MidiSongPosVertexAnalyzer.h
--- a/Engine/Plugins/Runtime/Harmonix/Source/HarmonixMetasound/Public/HarmonixMetasound/Analysis/MidiSongPosVertexAnalyzer.h
+++ b/Engine/Plugins/Runtime/Harmonix/Source/HarmonixMetasound/Public/HarmonixMetasound/Analysis/MidiSongPosVertexAnalyzer.h
@@ -48,6 +48,7 @@ namespace HarmonixMetasound::Analysis
 		int32 BlockSize {0};
 		Metasound::FSampleCount SampleCount {0};
 		int32 LastTickProcessed {-1};
+		int32 LastTempoMapTickProcessed {-1};
 		int32 LastAdvanceUpToTick {-1};
 		float CurrentFramesPerTick {0.0f};
 		const ISongMapEvaluator* CurrentSongMapEvaluator {nullptr};
diff --git a/Engine/Plugins/Runtime/Harmonix/Source/HarmonixMetasound/Public/HarmonixMetasound/DataTypes/MidiClock.h b/Engine/Plugins/Runtime/Harmonix/Source/HarmonixMetasound/Public/HarmonixMetasound/DataTypes/MidiClock.h
--- a/Engine/Plugins/Runtime/Harmonix/Source/HarmonixMetasound/Public/HarmonixMetasound/DataTypes/MidiClock.h
+++ b/Engine/Plugins/Runtime/Harmonix/Source/HarmonixMetasound/Public/HarmonixMetasound/DataTypes/MidiClock.h
@@ -59,12 +59,12 @@ namespace HarmonixMetasound
 
 		void SetTransportState(int32 BlockFrameIndex, EMusicPlayerTransportState TransportState);
 		void SetSpeed(int32 BlockFrameIndex, float Speed);
-		void SetTempo(int32 BlockFrameIndex, int32 Tick, float Bpm);
-		void SetTimeSignature(int32 BlockFrameIndex, int32 Tick, const FTimeSignature& TimeSignature);
+		void SetTempo(int32 BlockFrameIndex, int32 Tick, float Bpm, int32 TempoMapTick);
+		void SetTimeSignature(int32 BlockFrameIndex, int32 Tick, const FTimeSignature& TimeSignature, int32 TempoMapTick);
 
 		// directly seek this clock with a musical seek target or a specific tick
 		void SeekTo(int32 BlockFrameIndex, const FMusicSeekTarget& InTarget);
-		void SeekTo(int32 BlockFrameIndex, int32 Tick);
+		void SeekTo(int32 BlockFrameIndex, int32 Tick, int32 TempoMapTick);
 
 		// This will add a loop event to the clock event stream WITHOUT having 
 		// to set this clock to looping. This is used when this clock is being 
@@ -87,7 +87,7 @@ namespace HarmonixMetasound
 		// process and advance the clock normally based on the given sample frames
 		void Advance(int32 StartFrame, int32 NumFrames);
 
-		bool AdvanceToTick(int32 BlockFrameIndex, int32 UpToTick);
+		bool AdvanceToTick(int32 BlockFrameIndex, int32 UpToTick, int32 TempoMapTick);
 		bool AdvanceToMs(int32 BlockFrameIndex, float Ms);
 
 		bool  HasTransportStateChangesInBlock() const    { return NumTransportChangeInBlock > 0; }
@@ -106,9 +106,9 @@ namespace HarmonixMetasound
 		float GetSpeedAtBlockSampleFrame(int32 FrameIndex) const;
 		float GetSpeedAtEndOfBlock() const { return SpeedAtBlockEnd; }
 
-		float GetTempoAtStartOfBlock() const { return TempoAtBlockStart; }
+		float GetTempoAtStartOfBlock() const { return ExternalClockDriver ? ExternalClockDriver->GetTempoAtStartOfBlock() : TempoAtBlockStart; }
 		float GetTempoAtBlockSampleFrame(int32 FrameIndex) const;
-		float GetTempoAtEndOfBlock() const { return TempoAtBlockEnd; }
+		float GetTempoAtEndOfBlock() const { return ExternalClockDriver ? ExternalClockDriver->GetTempoAtEndOfBlock() : TempoAtBlockEnd; }
 
 		int32 GetLastProcessedMidiTick() const { return LastProcessedMidiTick; }
 		int32 GetNextMidiTickToProcess() const { return NextMidiTickToProcess; }
@@ -165,8 +165,9 @@ namespace HarmonixMetasound
 		
 		bool GetSongMapsChangedInBlock() const { return MidiDataChangedInBlock; }
 
-	private:
 		int32 GetNextTickToProcessAtBlockFrame(int32 BlockFrame) const;
+
+	private:
 		void AddEvent(const FMidiClockEvent& InEvent, bool bRequireSequential = true);
 		void HandleClockEvent(const FMidiClock& DrivingClock, const FMidiClockEvent& Event);
 		void PostTempoOrTimeSignatureEventsIfNeeded();
@@ -177,12 +178,12 @@ namespace HarmonixMetasound
 		MSGTYPE* LookForEventOnBlockFrameIndex(int32 BlockFrameIndex);
 
 		void AddTransportStateChangeToBlock(int32 BlockFrameIndex, EMusicPlayerTransportState TransportState);
-		void AddTimeSignatureChangeToBlock(int32 BlockFrameIndex, int32 Tick, const FTimeSignature& TimeSignature);
-		void AddTempoChangeToBlock(int32 BlockFrameIndex, int32 Tick, float Tempo);
+		void AddTimeSignatureChangeToBlock(int32 BlockFrameIndex, int32 Tick, const FTimeSignature& TimeSignature, int32 TempoMapTick);
+		void AddTempoChangeToBlock(int32 BlockFrameIndex, int32 Tick, float Tempo, int32 TempoMapTick);
 		void AddSpeedChangeToBlock(int32 BlockFrameIndex, float Speed, bool bIsNewLocalSpeed);
-		void AddLoopToBlock(int32 BlockFrameIndex, int32 FirstTick, int32 LoopLength);
-		void AddSeekToBlock(int32 BlockFrameIndex, int32 ToTick);
-		void AddAdvanceToBlock(int32 BlockFrameIndex, int32 FirstTick, int32 NumTicks);
+		void AddLoopToBlock(int32 BlockFrameIndex, int32 FirstTick, int32 LoopLength, int32 TempoMapTick);
+		void AddSeekToBlock(int32 BlockFrameIndex, int32 ToTick, int32 TempoMapTick);
+		void AddAdvanceToBlock(int32 BlockFrameIndex, int32 FirstTick, int32 NumTicks, int32 TempoMapTick);
 		void RebuildSongMapEvaluator(const TSharedPtr<const ISongMapEvaluator>& MidiWithTempo, const TSharedPtr<const ISongMapEvaluator>& MidiWithOtherMaps);
 
 		TSharedPtr<FSongMapsWithAlternateTempoSource> SongMapEvaluator;
@@ -197,6 +198,7 @@ namespace HarmonixMetasound
 		int32 FirstTickProcessedThisBlock;
 		int32 LastProcessedMidiTick;
 		int32 NextMidiTickToProcess;
+		int32 NextTempoMapTickToProcess;
 		float SampleRate;
 		Metasound::FSampleCount SampleCount;
 		int32 FramesUntilNextProcess;
diff --git a/Engine/Plugins/Runtime/Harmonix/Source/HarmonixMetasound/Public/HarmonixMetasound/DataTypes/MidiClockEvent.h b/Engine/Plugins/Runtime/Harmonix/Source/HarmonixMetasound/Public/HarmonixMetasound/DataTypes/MidiClockEvent.h
--- a/Engine/Plugins/Runtime/Harmonix/Source/HarmonixMetasound/Public/HarmonixMetasound/DataTypes/MidiClockEvent.h
+++ b/Engine/Plugins/Runtime/Harmonix/Source/HarmonixMetasound/Public/HarmonixMetasound/DataTypes/MidiClockEvent.h
@@ -11,10 +11,12 @@ namespace HarmonixMetasound
 		{
 			int32 FirstTickInLoop;
 			int32 LengthInTicks;
+			int32 TempoMapTick;
 
-			FLoop(const int32 FirstTickInLoop, const int32 LengthInTicks)
+			FLoop(const int32 FirstTickInLoop, const int32 LengthInTicks, const int32 TempoMapTick)
 				: FirstTickInLoop(FirstTickInLoop)
 				, LengthInTicks(LengthInTicks)
+				, TempoMapTick(TempoMapTick)
 			{}
 		};
 
@@ -22,10 +24,12 @@ namespace HarmonixMetasound
 		{
 			int32 LastTickProcessedBeforeSeek;
 			int32 NewNextTick;
+			int32 TempoMapTick;
 
-			FSeek(const int32 LastTickProcessedBeforeSeek, const int32 NewNextTick)
+			FSeek(const int32 LastTickProcessedBeforeSeek, const int32 NewNextTick, const int32 TempoMapTick)
 				: LastTickProcessedBeforeSeek(LastTickProcessedBeforeSeek)
 				, NewNextTick(NewNextTick)
+				, TempoMapTick(TempoMapTick)
 			{}
 		};
 
@@ -33,10 +37,12 @@ namespace HarmonixMetasound
 		{
 			int32 FirstTickToProcess;
 			int32 NumberOfTicksToProcess;
+			int32 TempoMapTick;
 
-			FAdvance(const int32 FirstTickToProcess, const int32 NumberOfTicksToProcess)
+			FAdvance(const int32 FirstTickToProcess, const int32 NumberOfTicksToProcess, const int32 TempoMapTick)
 				: FirstTickToProcess(FirstTickToProcess)
 				, NumberOfTicksToProcess(NumberOfTicksToProcess)
+				, TempoMapTick(TempoMapTick)
 			{}
 
 			int32 LastTickToProcess() const { return FirstTickToProcess + NumberOfTicksToProcess - 1; }
@@ -48,10 +54,12 @@ namespace HarmonixMetasound
 		{
 			int32 Tick;
 			float Tempo;
+			int32 TempoMapTick;
 
-			FTempoChange(const int32 Tick, const float Tempo)
+			FTempoChange(const int32 Tick, const float Tempo, const int32 TempoMapTick)
 				: Tick(Tick)
 				, Tempo(Tempo)
+				, TempoMapTick(TempoMapTick)
 			{}
 
 			bool ContainsTick(int32 InTick) const { return InTick == Tick; }
@@ -61,10 +69,12 @@ namespace HarmonixMetasound
 		{
 			int32 Tick;
 			FTimeSignature TimeSignature;
+			int32 TempoMapTick;
 
-			FTimeSignatureChange(const int32 Tick, FTimeSignature&& TimeSignature)
+			FTimeSignatureChange(const int32 Tick, FTimeSignature&& TimeSignature, const int32 TempoMapTick)
 				: Tick(Tick)
 				, TimeSignature(MoveTemp(TimeSignature))
+				, TempoMapTick(TempoMapTick)
 			{}
 
 			bool ContainsTick(int32 InTick) const { return InTick == Tick; }
diff --git a/Engine/Plugins/Runtime/Harmonix/Source/HarmonixMetasoundTests/Private/Analysis/MidiClockVertexAnalyzerTests.cpp b/Engine/Plugins/Runtime/Harmonix/Source/HarmonixMetasoundTests/Private/Analysis/MidiClockVertexAnalyzerTests.cpp
--- a/Engine/Plugins/Runtime/Harmonix/Source/HarmonixMetasoundTests/Private/Analysis/MidiClockVertexAnalyzerTests.cpp
+++ b/Engine/Plugins/Runtime/Harmonix/Source/HarmonixMetasoundTests/Private/Analysis/MidiClockVertexAnalyzerTests.cpp
@@ -47,7 +47,7 @@ namespace HarmonixMetasoundTests::MidiClockVertexAnalyzer
 		SongMaps->SetSongLengthTicks(std::numeric_limits<int32>::max());
 
 		ClockInput->AttachToSongMapEvaluator(SongMaps);
-		ClockInput->SeekTo(0,0);
+		ClockInput->SeekTo(0,0,0);
 		ClockInput->SetSpeed(0, Speed);
 		ClockInput->SetTransportState(0, HarmonixMetasound::EMusicPlayerTransportState::Playing);
 	}
diff --git a/Engine/Plugins/Runtime/Harmonix/Source/HarmonixMetasoundTests/Private/Analysis/MidiSongPosVertexAnalyzerTests.cpp b/Engine/Plugins/Runtime/Harmonix/Source/HarmonixMetasoundTests/Private/Analysis/MidiSongPosVertexAnalyzerTests.cpp
--- a/Engine/Plugins/Runtime/Harmonix/Source/HarmonixMetasoundTests/Private/Analysis/MidiSongPosVertexAnalyzerTests.cpp
+++ b/Engine/Plugins/Runtime/Harmonix/Source/HarmonixMetasoundTests/Private/Analysis/MidiSongPosVertexAnalyzerTests.cpp
@@ -49,7 +49,7 @@ namespace HarmonixMetasoundTests::MidiSongPosVertexAnalyzer
 		SongMaps->SetSongLengthTicks(std::numeric_limits<int32>::max());
 
 		ClockInput->AttachToSongMapEvaluator(SongMaps);
-		ClockInput->SeekTo(0,0);
+		ClockInput->SeekTo(0,0,0);
 		ClockInput->SetSpeed(0, 1.0f);
 		ClockInput->SetTransportState(0, HarmonixMetasound::EMusicPlayerTransportState::Playing);
 		int32 Bar1Tick = SongMaps->BarBeatTickIncludingCountInToTick(1, 1, 0);
diff --git a/Engine/Plugins/Runtime/Harmonix/Source/HarmonixMetasoundTests/Private/DataTypes/MidiClock.spec.cpp b/Engine/Plugins/Runtime/Harmonix/Source/HarmonixMetasoundTests/Private/DataTypes/MidiClock.spec.cpp
--- a/Engine/Plugins/Runtime/Harmonix/Source/HarmonixMetasoundTests/Private/DataTypes/MidiClock.spec.cpp
+++ b/Engine/Plugins/Runtime/Harmonix/Source/HarmonixMetasoundTests/Private/DataTypes/MidiClock.spec.cpp
@@ -28,7 +28,7 @@ namespace HarmonixMetasoundTests::MidiClock
 	void ExecuteWriteAdvance(int32 StartFrameIndex, int32 EndFrameIndex, float Speed)
 	{
 		TestClock->PrepareBlock();
-		TestClock->SeekTo(0, 0);
+		TestClock->SeekTo(0, 0, 0);
 		TestClock->SetTransportState(0, HarmonixMetasound::EMusicPlayerTransportState::Playing);
 
 		TestEqual("Clock.GetCurrentSongPosMs()", TestClock->GetCurrentSongPosMs(), 0.0f);
@@ -313,7 +313,7 @@ namespace HarmonixMetasoundTests::MidiClock
 
 			It("EventType::AdvanceThru.NonLooping", [&, this]()
 			{
-				DrivingClock->SeekTo(0,0);
+				DrivingClock->SeekTo(0,0,0);
 				DrivingClock->SetTransportState(0, HarmonixMetasound::EMusicPlayerTransportState::Playing);
 
 				float DrivingClockSpeed = 1.0f;
@@ -322,7 +322,7 @@ namespace HarmonixMetasoundTests::MidiClock
 				float DeltaMs = DeltaFrames * 1000.0f / OperatorSettings.GetSampleRate();
 				int32 Tick = (int32)(TestClock->GetSongMapEvaluator().MsToTick(DeltaMs) + 0.5f);
 
-				DrivingClock->AdvanceToTick(0, Tick);
+				DrivingClock->AdvanceToTick(0, Tick, Tick);
 
 				TestClock->PrepareBlock();
 				int32 OldEventsNum = TestClock->GetMidiClockEventsInBlock().Num();
@@ -356,7 +356,7 @@ namespace HarmonixMetasoundTests::MidiClock
 
 				int32 TestStartTick = LoopEndTick - 1;;
 
-				DrivingClock->SeekTo(0, TestStartTick);
+				DrivingClock->SeekTo(0, TestStartTick, TestStartTick);
 
 				float DeltaMs = FMidiClock::kMidiGranularity * 1000.0f / OperatorSettings.GetSampleRate();
 				int32 DeltaTicks = (int32)(DrivingClock->GetSongMapEvaluator().MsToTick(DeltaMs) + 0.5f);
@@ -389,7 +389,7 @@ namespace HarmonixMetasoundTests::MidiClock
 				TestClock->PrepareBlock();
 				int32 OldEventsNum = TestClock->GetMidiClockEventsInBlock().Num();
 
-				TestClock->SeekTo(StartFrame, 1000);
+				TestClock->SeekTo(StartFrame, 1000, 1000);
 
 				if (!TestTrue("Clock has new clock events", TestClock->GetMidiClockEventsInBlock().Num() > OldEventsNum))
 				{
diff --git a/Engine/Plugins/Runtime/Harmonix/Source/HarmonixMetasoundTests/Private/MidiOps/PulseGeneratorTests.cpp b/Engine/Plugins/Runtime/Harmonix/Source/HarmonixMetasoundTests/Private/MidiOps/PulseGeneratorTests.cpp
--- a/Engine/Plugins/Runtime/Harmonix/Source/HarmonixMetasoundTests/Private/MidiOps/PulseGeneratorTests.cpp
+++ b/Engine/Plugins/Runtime/Harmonix/Source/HarmonixMetasoundTests/Private/MidiOps/PulseGeneratorTests.cpp
@@ -32,7 +32,7 @@ namespace Harmonix::Midi::Ops::Tests
 			constexpr int32 NotesUntilWeAreSatisfiedThisWorks = 23;
 			int32 NumNotesReceived = 0;
 
-			Clock->SeekTo(0,0);
+			Clock->SeekTo(0,0,0);
 			Clock->SetTransportState(0, HarmonixMetasound::EMusicPlayerTransportState::Playing);
 
 			FTimeSignature TimeSignature(4, 4);
diff --git a/Engine/Plugins/Runtime/Harmonix/Source/HarmonixMidi/Public/HarmonixMidi/MidiSongPos.h b/Engine/Plugins/Runtime/Harmonix/Source/HarmonixMidi/Public/HarmonixMidi/MidiSongPos.h
--- a/Engine/Plugins/Runtime/Harmonix/Source/HarmonixMidi/Public/HarmonixMidi/MidiSongPos.h
+++ b/Engine/Plugins/Runtime/Harmonix/Source/HarmonixMidi/Public/HarmonixMidi/MidiSongPos.h
@@ -86,7 +86,6 @@ public:
 
 	static FMidiSongPos Lerp(const FMidiSongPos& A, const FMidiSongPos& B, float Alpha);
 
-private:
 	void SetByTimeAndTick(float Ms, float Tick, const ISongMapEvaluator& Map);
 };
 
