[Harmonix] Seek when the timing offset input changes on the midi clock offset node
#tests PIE

diff --git a/Engine/Plugins/Runtime/Harmonix/Source/HarmonixMetasound/Private/HarmonixMetasound/Nodes/MidiClockOffsetNode.cpp b/Engine/Plugins/Runtime/Harmonix/Source/HarmonixMetasound/Private/HarmonixMetasound/Nodes/MidiClockOffsetNode.cpp
--- a/Engine/Plugins/Runtime/Harmonix/Source/HarmonixMetasound/Private/HarmonixMetasound/Nodes/MidiClockOffsetNode.cpp
+++ b/Engine/Plugins/Runtime/Harmonix/Source/HarmonixMetasound/Private/HarmonixMetasound/Nodes/MidiClockOffsetNode.cpp
@@ -208,18 +208,18 @@ namespace HarmonixMetasound::Nodes::MidiClockOffset
 				}
 				else if (const FAdvance* AsAdvance = Event.TryGet<FAdvance>())
 				{
-					// if our offset changed while we were advancing, seek to the new offset first
-					if (!FMath::IsNearlyEqual(PrevOffsetMs, OffsetMs) || !FMath::IsNearlyEqual(PrevOffsetBeats, OffsetBeats) || PrevOffsetBars != OffsetBars)
+					bool bOffsetChanged = !FMath::IsNearlyEqual(PrevOffsetMs, OffsetMs) || !FMath::IsNearlyEqual(PrevOffsetBeats, OffsetBeats) || PrevOffsetBars != OffsetBars;
+
+					if (bOffsetChanged)
 					{
 						PrevOffsetMs = OffsetMs;
 						PrevOffsetBars = OffsetBars;
 						PrevOffsetBeats = OffsetBeats;
-
 					}
 					const int32 FirstickToProcess = GetTickWithOffset(AsAdvance->FirstTickToProcess, OffsetBars, OffsetBeats, OffsetMs);
-					// ONLY seek if there is a discontinuity AND the transport changed. Otherwise we 
-					// just want to advance from where we sit to the appropriate destination.
-					if (!bAdvancedSinceTransportChange && FirstickToProcess != MidiClockOut->GetNextMidiTickToProcess())
+					// ONLY seek if there is a discontinuity AND the transport changed OR the offset changed.
+					// Otherwise we just want to advance from where we sit to the appropriate destination.
+					if (bOffsetChanged || (!bAdvancedSinceTransportChange && FirstickToProcess != MidiClockOut->GetNextMidiTickToProcess()))
 					{
 						MidiClockOut->SeekTo(Event.BlockFrameIndex, FirstickToProcess, AsAdvance->TempoMapTick);
 					}
