- A couple of useful functions for the SPMC queue used by the midi clock system.

diff --git a/Engine/Plugins/Runtime/Harmonix/Source/HarmonixMetasound/Public/HarmonixMetasound/Analysis/SpmcAnalysisResultQueue.h b/Engine/Plugins/Runtime/Harmonix/Source/HarmonixMetasound/Public/HarmonixMetasound/Analysis/SpmcAnalysisResultQueue.h
--- a/Engine/Plugins/Runtime/Harmonix/Source/HarmonixMetasound/Public/HarmonixMetasound/Analysis/SpmcAnalysisResultQueue.h
+++ b/Engine/Plugins/Runtime/Harmonix/Source/HarmonixMetasound/Public/HarmonixMetasound/Analysis/SpmcAnalysisResultQueue.h
@@ -293,6 +293,16 @@ namespace HarmonixMetasound::Analysis
 				return Entry != nullptr;
 			}
 
+			bool operator==(const FEntry& Other)
+			{
+				return &Other == Entry;
+			}
+
+			bool operator!=(const FEntry& Other)
+			{
+				return &Other != Entry;
+			}
+
 			bool DiscontinuityDetectedInLastRead{ false };
 			const FEntry* Entry;
 			FReadCursor* ReadCursor;
@@ -391,6 +401,12 @@ namespace HarmonixMetasound::Analysis
 			Index--;
 		}
 
+		const FEntry& operator[](int32 Index)
+		{
+			check(Items.IsValidIndex(Index));
+			return Items[Index];
+		}
+
 	private:
 		int32 NextItemSerialNumber{ 0 };
 		int32 NextWritePos{ 0 };
