[Harmonix] Fix issue with notes getting stuck on when going through a merge node
#tests PIE

diff --git a/Engine/Plugins/Runtime/Harmonix/Source/HarmonixMetasound/Private/HarmonixMetasound/DataTypes/MidiStream.cpp b/Engine/Plugins/Runtime/Harmonix/Source/HarmonixMetasound/Private/HarmonixMetasound/DataTypes/MidiStream.cpp
--- a/Engine/Plugins/Runtime/Harmonix/Source/HarmonixMetasound/Private/HarmonixMetasound/DataTypes/MidiStream.cpp
+++ b/Engine/Plugins/Runtime/Harmonix/Source/HarmonixMetasound/Private/HarmonixMetasound/DataTypes/MidiStream.cpp
@@ -152,6 +152,7 @@ namespace HarmonixMetasound
 		}
 
 		// Merge in the active notes from the other stream
+		// Transform those notes so that the active note matches other events
 		for (const FMidiStreamEvent& ActiveNote : From.ActiveNotes)
 		{
 			check(ActiveNote.MidiMessage.IsNoteOn());
@@ -161,7 +162,7 @@ namespace HarmonixMetasound
 				return ActiveNote.TrackIndex == TrackedEvent.TrackIndex && ActiveNote.VoiceId == TrackedEvent.GetVoiceId();
 			}))
 			{
-				To.ActiveNotes.Add(ActiveNote);
+				To.ActiveNotes.Add(Transformer(ActiveNote));
 			}
 		}
 	}
@@ -173,6 +174,9 @@ namespace HarmonixMetasound
 		const FEventFilter& Filter,
 		const FEventTransformer& Transformer)
 	{
+		// Clear active notes. This unsticks notes when a midi stream is disconnected during live update.
+		To.ActiveNotes.Empty();
+
 		// Merge in stream A
 		Merge(FromA, To, Filter, Transformer);
 
