- Events when the music clock seeks or loops.

diff --git a/Engine/Plugins/Runtime/Harmonix/Source/HarmonixMetasound/Private/HarmonixMetasound/Components/MetasoundMusicClockDriver.cpp b/Engine/Plugins/Runtime/Harmonix/Source/HarmonixMetasound/Private/HarmonixMetasound/Components/MetasoundMusicClockDriver.cpp
--- a/Engine/Plugins/Runtime/Harmonix/Source/HarmonixMetasound/Private/HarmonixMetasound/Components/MetasoundMusicClockDriver.cpp
+++ b/Engine/Plugins/Runtime/Harmonix/Source/HarmonixMetasound/Private/HarmonixMetasound/Components/MetasoundMusicClockDriver.cpp
@@ -99,7 +99,7 @@ bool FMetasoundMusicClockDriver::CalculateSongPosWithOffset(float MsOffset, ECal
 	return true;
 }
 
-FMidiSongPos FMetasoundMusicClockDriver::CalculateSongPosAtMsForLoopingOrMonotonicClock(float AbsoluteMs, float& PositionTick) const
+FMidiSongPos FMetasoundMusicClockDriver::CalculateSongPosAtMsForLoopingOrMonotonicClock(float AbsoluteMs, float& PositionTick, bool& SeekDetected, bool& LoopDetected) const
 {
 	using namespace HarmonixMetasound::Analysis;
 	FMidiSongPos OutSongPos;
@@ -115,25 +115,52 @@ FMidiSongPos FMetasoundMusicClockDriver::CalculateSongPosAtMsForLoopingOrMonoton
 		return OutSongPos;
 	}
 
+	float NewPositionTick = 0.0f;
 	if (CurrentMapChain->LoopLengthTicks > 0)
 	{
 		float DrivingTick = CurrentMapChain->SongMaps->MsToTick(AbsoluteMs);
-		if (DrivingTick >= (float)(CurrentMapChain->FirstTickInLoop + CurrentMapChain->LoopLengthTicks))
+		float TickPastLoop = (float)(CurrentMapChain->FirstTickInLoop + CurrentMapChain->LoopLengthTicks);
+		if (DrivingTick >= TickPastLoop)
 		{
-			PositionTick = FMath::Fmod(DrivingTick - CurrentMapChain->FirstTickInLoop, (float)CurrentMapChain->LoopLengthTicks);
-			OutSongPos.SetByTick(PositionTick, *CurrentMapChain->SongMaps);
+			float WrappedTick = FMath::Fmod(DrivingTick - CurrentMapChain->FirstTickInLoop, (float)CurrentMapChain->LoopLengthTicks);
+			LoopDetected = (PositionTick - WrappedTick) > (float)(CurrentMapChain->LoopLengthTicks - 240);
+			if (LoopDetected)
+			{
+				UE_LOG(LogMusicClock, VeryVerbose, TEXT("Detected loop when calculating song pos (%f -> %f -> %f -> %d )"), PositionTick, WrappedTick, PositionTick - WrappedTick, CurrentMapChain->LoopLengthTicks);
+			}
+			OutSongPos.SetByTick(WrappedTick, *CurrentMapChain->SongMaps);
 			OutSongPos.Tempo = CurrentMapChain->SongMaps->GetTempoAtTick(FMath::FloorToInt32(DrivingTick));
+
+			if (!LoopDetected)
+			{
+				SeekDetected = CheckForSeek(PositionTick, WrappedTick, OutSongPos.Tempo, CurrentMapChain->SongMaps->GetTicksPerQuarterNote());
+				if (SeekDetected)
+				{
+					UE_LOG(LogMusicClock, VeryVerbose, TEXT("Detected seek when calculating song pos (%f -> %f)"), PositionTick, WrappedTick);
+				}
+			}
+
+			PositionTick = WrappedTick;
 			return OutSongPos;
 		}
+		NewPositionTick = DrivingTick;
+	}
+	else
+	{
+		NewPositionTick = CurrentMapChain->SongMaps->MsToTick(AbsoluteMs);
 	}
 
-	PositionTick = CurrentMapChain->SongMaps->MsToTick(AbsoluteMs);
-	OutSongPos.SetByTimeAndTick(AbsoluteMs, PositionTick, *(CurrentMapChain->SongMaps));
-
+	OutSongPos.SetByTimeAndTick(AbsoluteMs, NewPositionTick, *(CurrentMapChain->SongMaps));
+	SeekDetected = CheckForSeek(PositionTick, NewPositionTick, OutSongPos.Tempo, CurrentMapChain->SongMaps->GetTicksPerQuarterNote());
+	if (SeekDetected)
+	{
+		UE_LOG(LogMusicClock, VeryVerbose, TEXT("Detected seek when calculating song pos (%f -> %f)"), PositionTick, NewPositionTick);
+	}
+	PositionTick = NewPositionTick;
 	return OutSongPos;
 }
 
-FMidiSongPos FMetasoundMusicClockDriver::CalculateSongPosAtMsForOffsetClock(float PositionMs, float ClockTickOffsetFromDrivingClock, float& PositionTick) const
+FMidiSongPos FMetasoundMusicClockDriver::CalculateSongPosAtMsForOffsetClock(float PositionMs, float ClockTickOffsetFromDrivingClock, float& PositionTick, bool& SeekDetected) const
 {
 	using namespace HarmonixMetasound::Analysis;
 	
@@ -151,13 +178,30 @@ FMidiSongPos FMetasoundMusicClockDriver::CalculateSongPosAtMsForOffsetClock(floa
 		return OutSongPos;
 	}
 
-	PositionTick = CurrentMapChain->SongMaps->MsToTick(PositionMs);
-	OutSongPos.SetByTick(PositionTick, *CurrentMapChain->SongMaps);
-	OutSongPos.Tempo = CurrentMapChain->SongMaps->GetTempoAtTick(FMath::FloorToInt32(PositionTick - ClockTickOffsetFromDrivingClock));
+	float NewPositionTick = CurrentMapChain->SongMaps->MsToTick(PositionMs);
+
+	OutSongPos.SetByTick(NewPositionTick, *CurrentMapChain->SongMaps);
+	OutSongPos.Tempo = CurrentMapChain->SongMaps->GetTempoAtTick(FMath::FloorToInt32(NewPositionTick - ClockTickOffsetFromDrivingClock));
+
+	SeekDetected = CheckForSeek(PositionTick, NewPositionTick, OutSongPos.Tempo, CurrentMapChain->SongMaps->GetTicksPerQuarterNote());
+	if (SeekDetected)
+	{
+		UE_LOG(LogMusicClock, VeryVerbose, TEXT("Detected seek when calculating song pos (%f -> %f)"), PositionTick, NewPositionTick);
+	}
+
+	PositionTick = NewPositionTick;
 
 	return OutSongPos;
 }
 
+bool FMetasoundMusicClockDriver::CheckForSeek(float FirstTick, float NextTick, float CurrentTempo, int32 TicksPerQuarter) const
+{
+	float QuartersPerSecond = CurrentTempo / 60.0f;
+	float ExpectedDeltaQuarters = QuartersPerSecond * DeltaSecondsBetweenRefreshes;
+	float ExpectedDeltaTicks = ExpectedDeltaQuarters * (float)TicksPerQuarter;
+	return FMath::Abs(ExpectedDeltaTicks - (NextTick - FirstTick)) > (ExpectedDeltaTicks * 2.0f);
+}
+
 bool FMetasoundMusicClockDriver::RefreshCurrentSongPos()
 {
 	//	Only for use when on the game thread.
@@ -449,10 +493,13 @@ void FMetasoundMusicClockDriver::RefreshCurrentSongPosFromHistory()
 		RenderStartWallClockTimeSeconds = ClockComponent->GetWorld()->GetTimeSeconds() - (double)RenderStartSampleCount / (double)ClockHistory->SampleRate;
 		RenderSmoothingLagSeconds = MetasoundMusicClockDriver::SmoothedAudioRenderLagSeconds;
 		ErrorTracker.Reset();
+		LastRefreshWallClockTimeSeconds = RenderStartWallClockTimeSeconds;
 	}
 
 	double CurrentWallClockSeconds = ClockComponent->GetWorld()->GetTimeSeconds();
-		
+	DeltaSecondsBetweenRefreshes = CurrentWallClockSeconds - LastRefreshWallClockTimeSeconds;
+	LastRefreshWallClockTimeSeconds = CurrentWallClockSeconds;
+
 	double ExpectedRenderedSeconds = (CurrentWallClockSeconds - RenderStartWallClockTimeSeconds) * SyncSpeed * MetasoundMusicClockDriver::Fudge;
 	double RenderedSeconds = (double)LastRenderPosSampleCount / (double)ClockHistory->SampleRate;
 	double Error = RenderedSeconds - ExpectedRenderedSeconds;
@@ -542,7 +589,10 @@ void FMetasoundMusicClockDriver::UpdateCurrentTicksForOffsetClock(float Smoothed
 	SmoothedPositionMs += RenderSmoothingLagSeconds * 1000.0f;
 
 	// Calculate the song position AND get the "local tick" for the Audio Render timebase
-	ClockComponent->CurrentSmoothedAudioRenderSongPos = CalculateSongPosAtMsForOffsetClock(SmoothedPositionMs, SmoothedTick - SmoothedTempoMapTick, AudioRenderState.LocalTick);
+	ClockComponent->CurrentSmoothedAudioRenderSongPos = CalculateSongPosAtMsForOffsetClock(SmoothedPositionMs,
+		SmoothedTick - SmoothedTempoMapTick,
+		AudioRenderState.LocalTick,
+		ClockComponent->AudioRenderSeekDetected);
 	float SmoothedTempoMapMs = CurrentMapChain->SongMaps->TickToMs(SmoothedTempoMapTick);
 	AudioRenderState.TempoMapMs = SmoothedTempoMapMs + (RenderSmoothingLagSeconds * 1000.0f);
 	AudioRenderState.TempoMapTick = CurrentMapChain->SongMaps->MsToTick(AudioRenderState.TempoMapMs);
@@ -553,7 +603,10 @@ void FMetasoundMusicClockDriver::UpdateCurrentTicksForOffsetClock(float Smoothed
 	float Ms = SmoothedPositionMs;
 	Ms -= LookBehindLagMs;
 	// Calculate the song position AND get the "local tick" for the Player Experience timebase
-	ClockComponent->CurrentPlayerExperiencedSongPos = CalculateSongPosAtMsForOffsetClock(Ms, SmoothedTick - SmoothedTempoMapTick, PlayerExperienceState.LocalTick);
+	ClockComponent->CurrentPlayerExperiencedSongPos = CalculateSongPosAtMsForOffsetClock(Ms,
+		SmoothedTick - SmoothedTempoMapTick,
+		PlayerExperienceState.LocalTick,
+		ClockComponent->PlayerExperiencedSeekDetected);
 	PlayerExperienceState.TempoMapMs = AudioRenderState.TempoMapMs - LookBehindLagMs;
 	PlayerExperienceState.TempoMapTick = CurrentMapChain->SongMaps->MsToTick(PlayerExperienceState.TempoMapMs);
 
@@ -563,7 +616,10 @@ void FMetasoundMusicClockDriver::UpdateCurrentTicksForOffsetClock(float Smoothed
 	Ms = SmoothedPositionMs;
 	Ms -= LookBehindLagMs;
 	// Calculate the song position AND get the "local tick" for the Video Render timebase
-	ClockComponent->CurrentVideoRenderSongPos = CalculateSongPosAtMsForOffsetClock(Ms, SmoothedTick - SmoothedTempoMapTick, VideoRenderState.LocalTick);
+	ClockComponent->CurrentVideoRenderSongPos = CalculateSongPosAtMsForOffsetClock(Ms,
+		SmoothedTick - SmoothedTempoMapTick,
+		VideoRenderState.LocalTick,
+		ClockComponent->VideoRenderSeekDetected);
 	VideoRenderState.TempoMapMs = AudioRenderState.TempoMapMs - LookBehindLagMs;
 	VideoRenderState.TempoMapTick = CurrentMapChain->SongMaps->MsToTick(VideoRenderState.TempoMapMs);
 }
@@ -577,21 +633,30 @@ void FMetasoundMusicClockDriver::UpdateCurrentTicksForLoopingOrMonotonicClock(fl
 
 	// Calculate the song position AND get the "local tick" for the Audio Render timebase
 	AudioRenderState.TempoMapTick = CurrentMapChain->SongMaps->MsToTick(AudioRenderState.TempoMapMs);
-	ClockComponent->CurrentSmoothedAudioRenderSongPos = CalculateSongPosAtMsForLoopingOrMonotonicClock(AudioRenderState.TempoMapMs, AudioRenderState.LocalTick);
+	ClockComponent->CurrentSmoothedAudioRenderSongPos = CalculateSongPosAtMsForLoopingOrMonotonicClock(AudioRenderState.TempoMapMs,
+		AudioRenderState.LocalTick,
+		ClockComponent->AudioRenderSeekDetected,
+		ClockComponent->AudioRenderLoopDetected);
 
 	// Now the time the user should actually be "experiencing" (ie "hearing") can be calculated as an offset
 	// from the smooth audio rendering time...
 	float LookBehindLagMs = FHarmonixModule::Get().GetMeasuredUserExperienceAndReactionToAudioRenderOffsetMs();
 	PlayerExperienceState.TempoMapMs = AudioRenderState.TempoMapMs - LookBehindLagMs;
 	PlayerExperienceState.TempoMapTick = CurrentMapChain->SongMaps->MsToTick(PlayerExperienceState.TempoMapMs);
-	ClockComponent->CurrentPlayerExperiencedSongPos = CalculateSongPosAtMsForLoopingOrMonotonicClock(PlayerExperienceState.TempoMapMs, PlayerExperienceState.LocalTick);
+	ClockComponent->CurrentPlayerExperiencedSongPos = CalculateSongPosAtMsForLoopingOrMonotonicClock(PlayerExperienceState.TempoMapMs,
+		PlayerExperienceState.LocalTick,
+		ClockComponent->PlayerExperiencedSeekDetected,
+		ClockComponent->PlayerExperiencedLoopDetected);
 
 	// Now the time the game should be rendering graphics for can be calculated as an offset
 	// from the smooth audio rendering time...
 	LookBehindLagMs = FHarmonixModule::Get().GetMeasuredVideoToAudioRenderOffsetMs();
 	VideoRenderState.TempoMapMs = AudioRenderState.TempoMapMs - LookBehindLagMs;
 	VideoRenderState.TempoMapTick = CurrentMapChain->SongMaps->MsToTick(VideoRenderState.TempoMapMs);
-	ClockComponent->CurrentVideoRenderSongPos = CalculateSongPosAtMsForLoopingOrMonotonicClock(VideoRenderState.TempoMapMs, VideoRenderState.LocalTick);
+	ClockComponent->CurrentVideoRenderSongPos = CalculateSongPosAtMsForLoopingOrMonotonicClock(VideoRenderState.TempoMapMs,
+		VideoRenderState.LocalTick,
+		ClockComponent->VideoRenderSeekDetected,
+		ClockComponent->VideoRenderLoopDetected);
 }
 
 FMetasoundMusicClockDriver::EHistoryFailureType FMetasoundMusicClockDriver::CalculateSmoothedTick(Metasound::FSampleCount ExpectedRenderPosSampleCount, Metasound::FSampleCount LastRenderPosSampleCount,
diff --git a/Engine/Plugins/Runtime/Harmonix/Source/HarmonixMetasound/Private/HarmonixMetasound/Components/MetasoundMusicClockDriver.h b/Engine/Plugins/Runtime/Harmonix/Source/HarmonixMetasound/Private/HarmonixMetasound/Components/MetasoundMusicClockDriver.h
--- a/Engine/Plugins/Runtime/Harmonix/Source/HarmonixMetasound/Private/HarmonixMetasound/Components/MetasoundMusicClockDriver.h
+++ b/Engine/Plugins/Runtime/Harmonix/Source/HarmonixMetasound/Private/HarmonixMetasound/Components/MetasoundMusicClockDriver.h
@@ -61,6 +61,9 @@ private:
 	float RenderSmoothingLagSeconds = 0.030f;
 
 	double RenderStartWallClockTimeSeconds = 0.0;
+	double LastRefreshWallClockTimeSeconds = 0.0;
+	double DeltaSecondsBetweenRefreshes = 0.0;
+
 	Metasound::FSampleCount RenderStartSampleCount {0};
 
 	static const int kFramesOfErrorHistory = 10;
@@ -83,8 +86,9 @@ private:
 	FDelegateHandle GeneratorIOUpdatedCallbackHandle;
 	FDelegateHandle GraphChangedCallbackHandle;
 
-	FMidiSongPos CalculateSongPosAtMsForLoopingOrMonotonicClock(float AbsoluteMs, float& PositionTick) const;
-	FMidiSongPos CalculateSongPosAtMsForOffsetClock(float PositionMs, float ClockTickOffsetFromDrivingClock, float& PositionTick) const;
+	FMidiSongPos CalculateSongPosAtMsForLoopingOrMonotonicClock(float AbsoluteMs, float& PositionTick, bool& SeekDetected, bool& LoopDetected) const;
+	FMidiSongPos CalculateSongPosAtMsForOffsetClock(float PositionMs, float ClockTickOffsetFromDrivingClock, float& PositionTick, bool& SeekDetected) const;
+	bool CheckForSeek(float FirstTick, float NextTick, float CurrentTempo, int32 TicksPerQuarter) const;
 
 	void UpdateCurrentTicksForOffsetClock(float SmoothedTick, float SmoothedTempoMapTick);
 	void UpdateCurrentTicksForLoopingOrMonotonicClock(float SmoothedTick, float SmoothedTempoMapTick);
diff --git a/Engine/Plugins/Runtime/Harmonix/Source/HarmonixMetasound/Private/HarmonixMetasound/Components/MusicClockComponent.cpp b/Engine/Plugins/Runtime/Harmonix/Source/HarmonixMetasound/Private/HarmonixMetasound/Components/MusicClockComponent.cpp
--- a/Engine/Plugins/Runtime/Harmonix/Source/HarmonixMetasound/Private/HarmonixMetasound/Components/MusicClockComponent.cpp
+++ b/Engine/Plugins/Runtime/Harmonix/Source/HarmonixMetasound/Private/HarmonixMetasound/Components/MusicClockComponent.cpp
@@ -152,11 +152,7 @@ bool UMusicClockComponent::ConnectToMetasound()
 	{
 		return false;
 	}
-	if (ClockDriver)
-	{
-		ClockDriver->Disconnect();
-		ClockDriver = nullptr;
-	}
+	DisconnectFromClockDriver();
 	TSharedPtr<FMetasoundMusicClockDriver> MetasoundClockDriver = MakeShared<FMetasoundMusicClockDriver>(this);
 	bool Connected = MetasoundClockDriver->ConnectToAudioComponentsMetasound(MetasoundsAudioComponent, MetasoundOutputName);
 	ClockDriver = MoveTemp(MetasoundClockDriver);
@@ -171,13 +167,26 @@ void UMusicClockComponent::ConnectToWallClock()
 {
 	// we don't 'check' that the driver mode is wall clock here because if the driver mode is metasound and we can't 
 	// connect for some reason we will fall back to this clock driver!
+	DisconnectFromClockDriver();
+	ClockDriver = MakeShared<FWallClockMusicClockDriver>(this, TempoMap);
+}
+
+void UMusicClockComponent::DisconnectFromClockDriver()
+{
 	if (ClockDriver)
 	{
 		ClockDriver->Disconnect();
+		ClockDriver = nullptr;
 	}
-	ClockDriver = MakeShared<FWallClockMusicClockDriver>(this, TempoMap);
+	AudioRenderSeekDetected = false;
+	AudioRenderLoopDetected = false;
+	PlayerExperiencedSeekDetected = false;
+	PlayerExperiencedLoopDetected = false;
+	VideoRenderSeekDetected = false;
+	VideoRenderLoopDetected = false;
 }
 
+
 FMidiSongPos UMusicClockComponent::CalculateSongPosWithOffset(float MsOffset, ECalibratedMusicTimebase Timebase) const
 {
 	FMidiSongPos Result;
@@ -227,11 +236,7 @@ void UMusicClockComponent::EndPlay(const EEndPlayReason::Type EndPlayReason)
 {
 	Super::EndPlay(EndPlayReason);
 	UMidiClockUpdateSubsystem::StopTrackingMusicClockComponent(this);
-	if (ClockDriver)
-	{
-		ClockDriver->Disconnect();
-		ClockDriver = nullptr;
-	}
+	DisconnectFromClockDriver();
 }
 
 void UMusicClockComponent::TickComponent(float DeltaTime, enum ELevelTick TickType, FActorComponentTickFunction* ThisTickFunction)
@@ -239,6 +244,7 @@ void UMusicClockComponent::TickComponent(float DeltaTime, enum ELevelTick TickTy
 	Super::TickComponent(DeltaTime, TickType, ThisTickFunction);
 	EnsureClockIsValidForGameFrame();
 	BroadcastSongPosChanges();
+	BroadcastSeekLoopDetections();
 }
 
 void UMusicClockComponent::Start()
@@ -561,6 +567,33 @@ float UMusicClockComponent::GetSongRemainingMs(ECalibratedMusicTimebase Timebase
 	return SongLengthMs <= 0.f ? 0.f : SongLengthMs - (GetSongPos(Timebase).SecondsIncludingCountIn * 1000.0f);
 }
 
+bool UMusicClockComponent::SeekedThisFrame(ECalibratedMusicTimebase Timebase) const
+{
+	EnsureClockIsValidForGameFrame();
+	switch (Timebase)
+	{
+	case ECalibratedMusicTimebase::RawAudioRenderTime:
+	case ECalibratedMusicTimebase::AudioRenderTime:    return AudioRenderSeekDetected;
+	case ECalibratedMusicTimebase::ExperiencedTime:	   return PlayerExperiencedSeekDetected;
+	case ECalibratedMusicTimebase::VideoRenderTime:
+	default:                                           return VideoRenderSeekDetected;
+	}
+}
+
+bool UMusicClockComponent::LoopedThisFrame(ECalibratedMusicTimebase Timebase) const
+{
+	EnsureClockIsValidForGameFrame();
+	switch (Timebase)
+	{
+	case ECalibratedMusicTimebase::RawAudioRenderTime:
+	case ECalibratedMusicTimebase::AudioRenderTime:    return AudioRenderLoopDetected;
+	case ECalibratedMusicTimebase::ExperiencedTime:	   return PlayerExperiencedLoopDetected;
+	case ECalibratedMusicTimebase::VideoRenderTime:
+	default:                                           return VideoRenderLoopDetected;
+	}
+}
+
+
 const ISongMapEvaluator& UMusicClockComponent::GetSongMaps() const
 {
 	const ISongMapEvaluator* SongMaps = ClockDriver ? ClockDriver->GetCurrentSongMapEvaluator() : nullptr;
@@ -585,21 +618,54 @@ const FMidiSongPos& UMusicClockComponent::GetSongPos(ECalibratedMusicTimebase Ti
 	}
 }
 
+const FMidiSongPos& UMusicClockComponent::GetPreviousSongPos(ECalibratedMusicTimebase Timebase) const
+{
+	EnsureClockIsValidForGameFrame();
+
+	switch (Timebase)
+	{
+	case ECalibratedMusicTimebase::RawAudioRenderTime:
+		return PrevRawAudioRenderSongPos;
+	case ECalibratedMusicTimebase::AudioRenderTime:
+		return PrevAudioRenderSongPos;
+	case ECalibratedMusicTimebase::ExperiencedTime:
+		return PrevPlayerExperiencedSongPos;
+	case ECalibratedMusicTimebase::VideoRenderTime:
+	default:
+		return PrevVideoRenderSongPos;
+	}
+}
+
 FMidiSongPos UMusicClockComponent::GetCurrentSmoothedAudioRenderSongPos() const
 {
 	return GetSongPos(ECalibratedMusicTimebase::AudioRenderTime);
 }
 
+FMidiSongPos UMusicClockComponent::GetPreviousSmoothedAudioRenderSongPos() const
+{
+	return GetPreviousSongPos(ECalibratedMusicTimebase::AudioRenderTime);
+}
+
 FMidiSongPos UMusicClockComponent::GetCurrentVideoRenderSongPos() const
 {
 	return GetSongPos(ECalibratedMusicTimebase::VideoRenderTime);
 }
 
+FMidiSongPos UMusicClockComponent::GetPreviousVideoRenderSongPos() const
+{
+	return GetPreviousSongPos(ECalibratedMusicTimebase::VideoRenderTime);
+}
+
 FMidiSongPos UMusicClockComponent::GetCurrentPlayerExperiencedSongPos() const
 {
 	return GetSongPos(ECalibratedMusicTimebase::ExperiencedTime);
 }
 
+FMidiSongPos UMusicClockComponent::GetPreviousPlayerExperiencedSongPos() const
+{
+	return GetPreviousSongPos(ECalibratedMusicTimebase::ExperiencedTime);
+}
+
 FMidiSongPos UMusicClockComponent::GetCurrentRawAudioRenderSongPos() const
 {
 	return GetSongPos(ECalibratedMusicTimebase::RawAudioRenderTime);
@@ -667,6 +733,34 @@ void UMusicClockComponent::BroadcastSongPosChanges()
 	}
 }
 
+void UMusicClockComponent::BroadcastSeekLoopDetections()
+{
+	if(AudioRenderSeekDetected)
+	{
+		AudioRenderMusicTimeDiscontinuityEvent.Broadcast(EMusicTimeDiscontinuityType::Seek, PrevAudioRenderSongPos, CurrentSmoothedAudioRenderSongPos);
+	}
+	if(AudioRenderLoopDetected)
+	{
+		AudioRenderMusicTimeDiscontinuityEvent.Broadcast(EMusicTimeDiscontinuityType::Loop, PrevAudioRenderSongPos, CurrentSmoothedAudioRenderSongPos);
+	}
+	if(PlayerExperiencedSeekDetected)
+	{
+		PlayerExperienceMusicTimeDiscontinuityEvent.Broadcast(EMusicTimeDiscontinuityType::Seek, PrevPlayerExperiencedSongPos, CurrentPlayerExperiencedSongPos);
+	}
+	if(PlayerExperiencedLoopDetected)
+	{
+		PlayerExperienceMusicTimeDiscontinuityEvent.Broadcast(EMusicTimeDiscontinuityType::Loop, PrevPlayerExperiencedSongPos, CurrentPlayerExperiencedSongPos);
+	}
+	if(VideoRenderSeekDetected)
+	{
+		VideoRenderMusicTimeDiscontinuityEvent.Broadcast(EMusicTimeDiscontinuityType::Seek, PrevVideoRenderSongPos, CurrentVideoRenderSongPos);
+	}
+	if(VideoRenderLoopDetected)
+	{
+		VideoRenderMusicTimeDiscontinuityEvent.Broadcast(EMusicTimeDiscontinuityType::Loop, PrevVideoRenderSongPos, CurrentVideoRenderSongPos);
+	}
+}
+
 void UMusicClockComponent::MakeDefaultSongMap()
 {
 	DefaultMaps.EmptyAllMaps();
@@ -686,6 +780,13 @@ void FMusicClockDriverBase::EnsureClockIsValidForGameFrame()
 	ClockComponent->PrevPlayerExperiencedSongPos = ClockComponent->CurrentPlayerExperiencedSongPos;
 	ClockComponent->PrevVideoRenderSongPos = ClockComponent->CurrentVideoRenderSongPos;
 
+	ClockComponent->AudioRenderSeekDetected = false;
+	ClockComponent->AudioRenderLoopDetected = false;
+	ClockComponent->PlayerExperiencedSeekDetected = false;
+	ClockComponent->PlayerExperiencedLoopDetected = false;
+	ClockComponent->VideoRenderSeekDetected = false;
+	ClockComponent->VideoRenderLoopDetected = false;
+
 	if (RefreshCurrentSongPos())
 	{
 		ClockComponent->RawAudioRenderDeltaBarF = ClockComponent->CurrentRawAudioRenderSongPos.BarsIncludingCountIn - ClockComponent->PrevRawAudioRenderSongPos.BarsIncludingCountIn;
diff --git a/Engine/Plugins/Runtime/Harmonix/Source/HarmonixMetasound/Public/HarmonixMetasound/Components/MusicClockComponent.h b/Engine/Plugins/Runtime/Harmonix/Source/HarmonixMetasound/Public/HarmonixMetasound/Components/MusicClockComponent.h
--- a/Engine/Plugins/Runtime/Harmonix/Source/HarmonixMetasound/Public/HarmonixMetasound/Components/MusicClockComponent.h
+++ b/Engine/Plugins/Runtime/Harmonix/Source/HarmonixMetasound/Public/HarmonixMetasound/Components/MusicClockComponent.h
@@ -32,10 +32,18 @@ enum class EMusicClockDriveMethod : uint8
 	MetaSound,
 };
 
+UENUM(BlueprintType)
+enum class EMusicTimeDiscontinuityType : uint8
+{
+	Loop,
+	Seek,
+};
+
 DECLARE_DYNAMIC_MULTICAST_DELEGATE_TwoParams(FBeatEvent, int, BeatNumber, int, BeatInBar);
 DECLARE_DYNAMIC_MULTICAST_DELEGATE_OneParam(FBarEvent, int, BarNumber);
 DECLARE_DYNAMIC_MULTICAST_DELEGATE_ThreeParams(FSectionEvent, const FString&, SectionName, float, SectionStartMs, float, SectionLengthMs);
 DECLARE_DYNAMIC_MULTICAST_DELEGATE_OneParam(FPlayStateEvent, EMusicClockState, State);
+DECLARE_DYNAMIC_MULTICAST_DELEGATE_ThreeParams(FMusicTimeDiscontinuityEvent, EMusicTimeDiscontinuityType, Type, FMidiSongPos, PreviousPos, FMidiSongPos, NewPos);
 
 DECLARE_DYNAMIC_MULTICAST_DELEGATE(FMusicClockConnected);
 DECLARE_DYNAMIC_MULTICAST_DELEGATE(FMusicClockDisconnected);
@@ -205,11 +213,24 @@ public:
 	UFUNCTION(BlueprintPure, Category = "MusicClock")
 	const FMidiSongPos& GetSongPos(ECalibratedMusicTimebase Timebase = ECalibratedMusicTimebase::VideoRenderTime) const;
 
+	UFUNCTION(BlueprintPure, Category = "MusicClock")
+	const FMidiSongPos& GetPreviousSongPos(ECalibratedMusicTimebase Timebase = ECalibratedMusicTimebase::VideoRenderTime) const;
+
 	/** Returns the remaining time until the end of the MIDI in milliseconds based on the timestamp corresponding to the passed Timebase */
 	// Note: Not const as it might cause the clock to update from its source.
 	UFUNCTION(BlueprintPure, Category = "MusicClock")
 	float GetSongRemainingMs(ECalibratedMusicTimebase Timebase = ECalibratedMusicTimebase::VideoRenderTime) const;
 
+	/** Returns true if there was a seek in the specified timebase */
+	// Note: Not const as it might cause the clock to update from its source.
+	UFUNCTION(BlueprintPure, Category = "MusicClock")
+	bool SeekedThisFrame(ECalibratedMusicTimebase Timebase = ECalibratedMusicTimebase::VideoRenderTime) const;
+
+	/** Returns true if there was a seek in the specified timebase */
+	// Note: Not const as it might cause the clock to update from its source.
+	UFUNCTION(BlueprintPure, Category = "MusicClock")
+	bool LoopedThisFrame(ECalibratedMusicTimebase Timebase = ECalibratedMusicTimebase::VideoRenderTime) const;
+
 	UFUNCTION(BlueprintPure, Category = "Count In")
 	float GetCountInSeconds() const;
 
@@ -281,6 +302,15 @@ public:
 	UPROPERTY(BlueprintAssignable, Category = "MusicClock")
 	FMusicClockDisconnected MusicClockDisconnectedEvent;
 
+	UPROPERTY(BlueprintAssignable, Category = "MusicClock")
+	FMusicTimeDiscontinuityEvent AudioRenderMusicTimeDiscontinuityEvent;
+
+	UPROPERTY(BlueprintAssignable, Category = "MusicClock")
+	FMusicTimeDiscontinuityEvent PlayerExperienceMusicTimeDiscontinuityEvent;
+
+	UPROPERTY(BlueprintAssignable, Category = "MusicClock")
+	FMusicTimeDiscontinuityEvent VideoRenderMusicTimeDiscontinuityEvent;
+
 private:
 	// Don't let C++ access these directly! They are a blueprint convenience and only work because 
 	// they specify getter functions!
@@ -298,12 +328,21 @@ public:
 	UFUNCTION(BlueprintCallable, BlueprintPure, Category = "MusicClock")
 	FMidiSongPos GetCurrentSmoothedAudioRenderSongPos() const;
 
+	UFUNCTION(BlueprintCallable, BlueprintPure, Category = "MusicClock")
+	FMidiSongPos GetPreviousSmoothedAudioRenderSongPos() const;
+
 	UFUNCTION(BlueprintCallable, BlueprintPure, Category = "MusicClock")
 	FMidiSongPos GetCurrentVideoRenderSongPos() const;
 
+	UFUNCTION(BlueprintCallable, BlueprintPure, Category = "MusicClock")
+	FMidiSongPos GetPreviousVideoRenderSongPos() const;
+
 	UFUNCTION(BlueprintCallable, BlueprintPure, Category = "MusicClock")
 	FMidiSongPos GetCurrentPlayerExperiencedSongPos() const;
 
+	UFUNCTION(BlueprintCallable, BlueprintPure, Category = "MusicClock")
+	FMidiSongPos GetPreviousPlayerExperiencedSongPos() const;
+
 	UFUNCTION(BlueprintCallable, BlueprintPure, Category = "MusicClock")
 	FMidiSongPos GetCurrentRawAudioRenderSongPos() const;
 
@@ -355,6 +394,13 @@ private:
 	int32 LastBroadcastBeat = -1;
 	FSongSection LastBroadcastSongSection;
 
+	bool AudioRenderSeekDetected = false;
+	bool AudioRenderLoopDetected = false;
+	bool PlayerExperiencedSeekDetected = false;
+	bool PlayerExperiencedLoopDetected = false;
+	bool VideoRenderSeekDetected = false;
+	bool VideoRenderLoopDetected = false;
+
 	FMidiSongPos PrevRawAudioRenderSongPos;
 	FMidiSongPos PrevAudioRenderSongPos;
 	FMidiSongPos PrevPlayerExperiencedSongPos;
@@ -364,9 +410,11 @@ private:
 
 	void CreateClockDriver();
 	void BroadcastSongPosChanges();
+	void BroadcastSeekLoopDetections();
 	void MakeDefaultSongMap();
 	bool ConnectToMetasound();
 	void ConnectToWallClock();
+	void DisconnectFromClockDriver();
 
 	// Ensures the clock will be updated once per frame.  Should only get called on the game thread.
 	void EnsureClockIsValidForGameFrame() const;// TODO: Cleanup task - UE-205069 - If we find we are able to use the new 
