[Fabric] Fix underflow issue with midi notes initialized by non-uint numbers

#tests PIE

diff --git a/Engine/Plugins/Runtime/Harmonix/Source/HarmonixMidi/Public/HarmonixMidi/Blueprint/MidiNote.h b/Engine/Plugins/Runtime/Harmonix/Source/HarmonixMidi/Public/HarmonixMidi/Blueprint/MidiNote.h
--- a/Engine/Plugins/Runtime/Harmonix/Source/HarmonixMidi/Public/HarmonixMidi/Blueprint/MidiNote.h
+++ b/Engine/Plugins/Runtime/Harmonix/Source/HarmonixMidi/Public/HarmonixMidi/Blueprint/MidiNote.h
@@ -20,8 +20,8 @@ struct HARMONIXMIDI_API FMidiNote
 	FMidiNote() : NoteNumber(60) { };
 
 	FMidiNote(uint8 inValue) : NoteNumber(inValue) { };
-	FMidiNote(int inValue) : NoteNumber(inValue) { };
-	FMidiNote(int8 inValue) : NoteNumber(inValue) { };
+	FMidiNote(int inValue) : NoteNumber(inValue < 0 ? 0 : inValue) { };
+	FMidiNote(int8 inValue) : NoteNumber(inValue < 0 ? 0 : inValue) { };
 
 	operator uint8() { return NoteNumber; }
 	operator int8() { return NoteNumber; }
@@ -36,7 +36,7 @@ struct HARMONIXMIDI_API FMidiNote
 	
 	
 
-	UPROPERTY(EditAnywhere, BlueprintReadWrite, Category = "Midi")
+	UPROPERTY(EditAnywhere, BlueprintReadWrite, Category = "Midi", meta=(ClampMin="0", ClampMax="127"))
 	uint8 NoteNumber = 60;
 
 	FText GetDisplayText() const;
