Fix issue where smoothing music clock that ran for a long time would start to stutter.

diff --git a/Engine/Plugins/Runtime/Harmonix/Source/HarmonixMetasound/Private/HarmonixMetasound/Components/MetasoundMusicClockDriver.cpp b/Engine/Plugins/Runtime/Harmonix/Source/HarmonixMetasound/Private/HarmonixMetasound/Components/MetasoundMusicClockDriver.cpp
--- a/Engine/Plugins/Runtime/Harmonix/Source/HarmonixMetasound/Private/HarmonixMetasound/Components/MetasoundMusicClockDriver.cpp
+++ b/Engine/Plugins/Runtime/Harmonix/Source/HarmonixMetasound/Private/HarmonixMetasound/Components/MetasoundMusicClockDriver.cpp
@@ -13,7 +13,7 @@
 namespace MetasoundMusicClockDriver
 {
 	float Fudge = 1.00;
-	float kP = 0.001;
+	float kP = 0.18;
 	float HistoricSmoothedAudioRenderLagSeconds = 0.030; // this used to be baked-in/hardcoded into the smoothing of the audio render time. 
 	float SmoothedAudioRenderLagSeconds = 0.030f; // 30 ms
 	float MaxErrorSecondsBeforeJump = 0.060f; // 60ms
@@ -521,7 +521,10 @@ void FMetasoundMusicClockDriver::RefreshCurrentSongPosFromHistory()
 		}
 
 		// Use proportional part of error to adjust speed ever so slightly...
-		SyncSpeed += MetasoundMusicClockDriver::kP * ErrorTracker.Min();
+		if (ExpectedRenderedSeconds > 0.0)
+		{
+			SyncSpeed += MetasoundMusicClockDriver::kP * ErrorTracker.Min() / ExpectedRenderedSeconds;
+		}
 		SyncSpeed = FMath::Clamp(SyncSpeed, MetasoundMusicClockDriver::SlowestCorrectionSpeed, MetasoundMusicClockDriver::FastestCorrectionSpeed);
 	}
 
