Use transposed notes when calling All Notes Off in StepSequencePlayerNode
#tests PIE

diff --git a/Engine/Plugins/Runtime/Harmonix/Source/HarmonixMetasound/Private/HarmonixMetasound/Nodes/StepSequencePlayerNode.cpp b/Engine/Plugins/Runtime/Harmonix/Source/HarmonixMetasound/Private/HarmonixMetasound/Nodes/StepSequencePlayerNode.cpp
--- a/Engine/Plugins/Runtime/Harmonix/Source/HarmonixMetasound/Private/HarmonixMetasound/Nodes/StepSequencePlayerNode.cpp
+++ b/Engine/Plugins/Runtime/Harmonix/Source/HarmonixMetasound/Private/HarmonixMetasound/Nodes/StepSequencePlayerNode.cpp
@@ -891,15 +891,21 @@ namespace HarmonixMetasound::Nodes::StepSequencePlayer
 
 		if (SequenceTable->Notes.Num() < CurrentCellNotes.Num())
 		{
+			const int32 AdditionalOctaveNotes = (int32)*AdditionalOctavesInPin * 12;
+
 			// We may have existing notes that need to be stopped.
 			for (int32 i = SequenceTable->Notes.Num(); i < CurrentCellNotes.Num(); ++i)
 			{
 				if (CurrentCellNotes[i])
 				{
+					const int32 OriginalNote = SequenceTable->Notes[i].NoteNumber;
+					const int32 TransposedNote = FMath::Clamp(OriginalNote + AdditionalOctaveNotes, 0, 127);
+
 					uint8 MidiCh;
 					uint8 MidiNote;
 					CurrentCellNotes[i].GetChannelAndNote(MidiCh, MidiNote);
 					FMidiStreamEvent MidiEvent(CurrentCellNotes[i].GetGeneratorId(), FMidiMsg::CreateNoteOff(MidiCh, MidiNote));
+					MidiEvent.MidiMessage.Data1 = TransposedNote;
 					MidiEvent.BlockSampleFrameIndex  = CurrentBlockSpanStart;
 					MidiEvent.AuthoredMidiTick       = 0;
 					MidiEvent.CurrentMidiTick        = 0;
@@ -921,14 +927,19 @@ namespace HarmonixMetasound::Nodes::StepSequencePlayer
 
 	void FStepSequencePlayerOperator::AllNotesOff(int32 AtFrameIndex, int32 AbsMidiTick, bool ResetCellIndex)
 	{
+		const int32 AdditionalOctaveNotes = (int32)*AdditionalOctavesInPin * 12;
 		for (int i = 0; i < CurrentCellNotes.Num(); ++i)
 		{
 			if (CurrentCellNotes[i])
 			{
+				const int32 OriginalNote = SequenceTable->Notes[i].NoteNumber;
+				const int32 TransposedNote = FMath::Clamp(OriginalNote + AdditionalOctaveNotes, 0, 127);
+
 				uint8 MidiCh;
 				uint8 MidiNote;
 				CurrentCellNotes[i].GetChannelAndNote(MidiCh, MidiNote);
 				FMidiStreamEvent MidiEvent(CurrentCellNotes[i].GetGeneratorId(), FMidiMsg::CreateNoteOff(MidiCh, MidiNote));
+				MidiEvent.MidiMessage.Data1 = TransposedNote;
 				MidiEvent.BlockSampleFrameIndex  = AtFrameIndex;
 				MidiEvent.AuthoredMidiTick       = AbsMidiTick;
 				MidiEvent.CurrentMidiTick        = AbsMidiTick;
