[Harmonix] Make subdivision trigger and pulse generator work the same way, and fix a couple issues with each.

#virtualized

diff --git a/Engine/Plugins/Runtime/Harmonix/Source/HarmonixMetasound/Private/HarmonixMetasound/MidiOps/PulseGenerator.cpp b/Engine/Plugins/Runtime/Harmonix/Source/HarmonixMetasound/Private/HarmonixMetasound/MidiOps/PulseGenerator.cpp
--- a/Engine/Plugins/Runtime/Harmonix/Source/HarmonixMetasound/Private/HarmonixMetasound/MidiOps/PulseGenerator.cpp
+++ b/Engine/Plugins/Runtime/Harmonix/Source/HarmonixMetasound/Private/HarmonixMetasound/MidiOps/PulseGenerator.cpp
@@ -4,7 +4,7 @@
 
 namespace Harmonix::Midi::Ops
 {
-	void FPulseGenerator::Enable(bool bEnable)
+	void FPulseGenerator::Enable(const bool bEnable)
 	{
 		Enabled = bEnable;
 	}
@@ -12,6 +12,7 @@ namespace Harmonix::Midi::Ops
 	void FPulseGenerator::SetClock(const TSharedPtr<const HarmonixMetasound::FMidiClock, ESPMode::NotThreadSafe>& NewClock)
 	{
 		Clock = NewClock;
+		CurrentTimeSignature = *NewClock->GetSongMapEvaluator().GetTimeSignatureAtTick(NewClock->GetLastProcessedMidiTick());
 	}
 
 	void FPulseGenerator::SetInterval(const FMusicTimeInterval& NewInterval)
@@ -22,13 +23,14 @@ namespace Harmonix::Midi::Ops
 		Interval.IntervalMultiplier = FMath::Max(Interval.IntervalMultiplier, static_cast<uint16>(1));
 	}
 
-	void FPulseGenerator::Process(HarmonixMetasound::FMidiStream& OutStream)
+	void FPulseGenerator::Reset()
 	{
-		using namespace HarmonixMetasound;
-		using namespace HarmonixMetasound::MidiClockMessageTypes;
+		NextPulseTimestamp = { -1, -1 };
+	}
 
-		OutStream.PrepareBlock();
-		
+	void FPulseGenerator::Process(const TFunctionRef<void(const FPulseInfo&)>& OnPulse)
+	{
+		check(Clock.IsValid()); // if we're here, and we don't have a clock, we have problems
 		const auto PinnedClock = Clock.Pin();
 
 		if (!PinnedClock)
@@ -36,53 +38,65 @@ namespace Harmonix::Midi::Ops
 			return;
 		}
 
+		using namespace HarmonixMetasound;
+		using namespace HarmonixMetasound::MidiClockMessageTypes;
+
 		for (const FMidiClockEvent& ClockEvent : PinnedClock->GetMidiClockEventsInBlock())
 		{
 			if (const FAdvance* AsAdvance = ClockEvent.TryGet<FAdvance>())
 			{
-				HandleAdvance(ClockEvent.BlockFrameIndex, AsAdvance, OutStream);
+				if (!NextPulseTimestamp.IsValid())
+				{
+					return;
+				}
+				
+				int32 NextPulseTick = PinnedClock->GetSongMapEvaluator().MusicTimestampToTick(NextPulseTimestamp);
+
+				while (AsAdvance->LastTickToProcess() >= NextPulseTick)
+				{
+					OnPulse({ ClockEvent.BlockFrameIndex, NextPulseTick });
+
+					IncrementTimestampByInterval(NextPulseTimestamp, Interval, CurrentTimeSignature);
+
+					NextPulseTick = PinnedClock->GetSongMapEvaluator().MusicTimestampToTick(NextPulseTimestamp);
+				}
 			}
 			else if (const FTimeSignatureChange* AsTimeSigChange = ClockEvent.TryGet<FTimeSignatureChange>())
 			{
-				HandleTimeSignatureChange(AsTimeSigChange);
+				CurrentTimeSignature = AsTimeSigChange->TimeSignature;
+				
+				// Time sig changes will come on the downbeat, and if we change time signature,
+				// we want to reset the pulse, so the next pulse is now plus the offset
+				NextPulseTimestamp = PinnedClock->GetSongMapEvaluator().TickToMusicTimestamp(AsTimeSigChange->Tick);
+				IncrementTimestampByOffset(NextPulseTimestamp, Interval, CurrentTimeSignature);
+			}
+			else if (const FLoop* AsLoop = ClockEvent.TryGet<FLoop>())
+			{
+				// We assume the pulse should be reset on loop, and the loop start should imply the phase of the pulse
+				NextPulseTimestamp = PinnedClock->GetSongMapEvaluator().TickToMusicTimestamp(AsLoop->FirstTickInLoop);
+				IncrementTimestampByOffset(NextPulseTimestamp, Interval, CurrentTimeSignature);
 			}
 		}
 	}
 
-	void FPulseGenerator::HandleAdvance(const int32 BlockFrameIndex, const HarmonixMetasound::MidiClockMessageTypes::FAdvance* Advance, HarmonixMetasound::FMidiStream& OutStream)
+	void FMidiPulseGenerator::Reset()
 	{
-		if (!NextPulseTimestamp.IsValid())
-		{
-			return;
-		}
-
-		check(Clock.IsValid()); // if we're here and we don't have a clock, we have problems
-		const auto PinnedClock = Clock.Pin();
-		int32 NextPulseTick = PinnedClock->GetSongMapEvaluator().MusicTimestampToTick(NextPulseTimestamp);
-
-		while (Advance->LastTickToProcess() >= NextPulseTick)
-		{
-			DoPulse(BlockFrameIndex, NextPulseTick, OutStream);
-
-			IncrementTimestampByInterval(NextPulseTimestamp, Interval, CurrentTimeSignature);
+		FPulseGenerator::Reset();
 
-			NextPulseTick = PinnedClock->GetSongMapEvaluator().MusicTimestampToTick(NextPulseTimestamp);
-		}
+		LastNoteOn.Reset();
 	}
 
-	void FPulseGenerator::HandleTimeSignatureChange(const HarmonixMetasound::MidiClockMessageTypes::FTimeSignatureChange* TimeSigChange)
+	void FMidiPulseGenerator::Process(HarmonixMetasound::FMidiStream& OutStream)
 	{
-		CurrentTimeSignature = TimeSigChange->TimeSignature;
+		OutStream.PrepareBlock();
 
-		check(Clock.IsValid()); // if we're here and we don't have a clock, we have problems
-		const auto PinnedClock = Clock.Pin();
-		// Time sig changes will come on the downbeat, and if we change time signature,
-		// we want to reset the pulse, so the next pulse is now plus the offset
-		NextPulseTimestamp = PinnedClock->GetSongMapEvaluator().TickToMusicTimestamp(TimeSigChange->Tick);
-		IncrementTimestampByOffset(NextPulseTimestamp, Interval, CurrentTimeSignature);
+		FPulseGenerator::Process([this, &OutStream](const FPulseInfo& Pulse)
+		{
+			AddPulseNote(Pulse.BlockFrameIndex, Pulse.Tick, OutStream);
+		});
 	}
 
-	void FPulseGenerator::DoPulse(const int32 BlockFrameIndex, const int32 PulseTick, HarmonixMetasound::FMidiStream& OutStream)
+	void FMidiPulseGenerator::AddPulseNote(const int32 BlockFrameIndex, const int32 PulseTick, HarmonixMetasound::FMidiStream& OutStream)
 	{
 		int32 NoteOnSample = BlockFrameIndex;
 
diff --git a/Engine/Plugins/Runtime/Harmonix/Source/HarmonixMetasound/Private/HarmonixMetasound/Nodes/MidiClockSubdivisionTriggerNode.cpp b/Engine/Plugins/Runtime/Harmonix/Source/HarmonixMetasound/Private/HarmonixMetasound/Nodes/MidiClockSubdivisionTriggerNode.cpp
--- a/Engine/Plugins/Runtime/Harmonix/Source/HarmonixMetasound/Private/HarmonixMetasound/Nodes/MidiClockSubdivisionTriggerNode.cpp
+++ b/Engine/Plugins/Runtime/Harmonix/Source/HarmonixMetasound/Private/HarmonixMetasound/Nodes/MidiClockSubdivisionTriggerNode.cpp
@@ -17,15 +17,25 @@
 
 #include "HarmonixMetasound/Common.h"
 #include "HarmonixMetasound/DataTypes/MidiClock.h"
-
-DEFINE_LOG_CATEGORY_STATIC(LogMidiClockSubdivisionTrigger, Log, All);
+#include "HarmonixMetasound/MidiOps/PulseGenerator.h"
+#include "HarmonixMetasound/Nodes/MidiClockSubdivisionTriggerNode.h"
 
 #define LOCTEXT_NAMESPACE "HarmonixMetaSound_MidiClockSubdivisionTriggerNode"
 
-namespace HarmonixMetasound
+namespace HarmonixMetasound::Nodes::MidiClockSubdivisionTriggerNode
 {
 	using namespace Metasound;
 
+	const FNodeClassName& GetClassName()
+	{
+		static const FNodeClassName ClassName { HarmonixNodeNamespace, "MidiClockSubdivisionTrigger", "" };
+		return ClassName;
+	}
+
+	int32 GetCurrentMajorVersion()
+	{
+		return 0;
+	}
 
 	class FMidiClockSubdivisionTriggerOperator : public TExecutableOperator<FMidiClockSubdivisionTriggerOperator>
 	{
@@ -63,13 +73,7 @@ namespace HarmonixMetasound
 		FTriggerWriteRef   TriggerOutPin;
 
  		//** DATA (current state)
-		EMidiClockSubdivisionQuantization GridSizeUnits;
-		int32 GridSizeMultiplier = 1;
-		EMidiClockSubdivisionQuantization GridOffsetUnits;
-		int32 GridOffsetMultiplier = 0;
-
-		int32 GridOffsetTicks = 0;
-		int32 GridSizeTicks	  = 0;
+		Harmonix::Midi::Ops::FPulseGenerator PulseGenerator;
 	};
 
 	class FMidiClockSubdivisionTriggerNode : public FNodeFacade
@@ -78,7 +82,7 @@ namespace HarmonixMetasound
 		FMidiClockSubdivisionTriggerNode(const FNodeInitData& InInitData)
 			: FNodeFacade(InInitData.InstanceName, InInitData.InstanceID, TFacadeOperatorClass<FMidiClockSubdivisionTriggerOperator>())
 		{}
-		virtual ~FMidiClockSubdivisionTriggerNode() = default;
+		virtual ~FMidiClockSubdivisionTriggerNode() override = default;
 	};
 
 	METASOUND_REGISTER_NODE(FMidiClockSubdivisionTriggerNode)
@@ -88,8 +92,8 @@ namespace HarmonixMetasound
 		auto InitNodeInfo = []() -> FNodeClassMetadata
 		{
 			FNodeClassMetadata Info;
-			Info.ClassName        = { HarmonixNodeNamespace, TEXT("MidiClockSubdivisionTrigger"), TEXT("")};
-			Info.MajorVersion     = 0;
+			Info.ClassName        = GetClassName();
+			Info.MajorVersion     = GetCurrentMajorVersion();
 			Info.MinorVersion     = 1;
 			Info.DisplayName      = METASOUND_LOCTEXT("MIDIClockSubdivisionTriggerNode_DisplayName", "MIDI Clock Subdivision Trigger");
 			Info.Description      = METASOUND_LOCTEXT("MIDIClockSubdivisionTriggerNode_Description", "Watches a MIDI clock and outputs triggers at musical subdivisions.");
@@ -105,16 +109,23 @@ namespace HarmonixMetasound
 		return Info;
 	}
 
-	namespace MidiClockSubdivisionTriggerPinNames
+	namespace Inputs
 	{
-		METASOUND_PARAM(TriggerOutput, "Trigger Out", "A series of triggers at the specified subdivision grid.")
+		DEFINE_METASOUND_PARAM_ALIAS(MidiClock, CommonPinNames::Inputs::MidiClock);
+		DEFINE_METASOUND_PARAM_ALIAS(Enable, CommonPinNames::Inputs::Enable);
+		DEFINE_METASOUND_PARAM_ALIAS(GridSizeUnits, CommonPinNames::Inputs::GridSizeUnits);
+		DEFINE_METASOUND_PARAM_ALIAS(GridSizeMult, CommonPinNames::Inputs::GridSizeMult);
+		DEFINE_METASOUND_PARAM_ALIAS(OffsetUnits, CommonPinNames::Inputs::OffsetUnits);
+		DEFINE_METASOUND_PARAM_ALIAS(OffsetMult, CommonPinNames::Inputs::OffsetMult);
 	}
 
-	const FVertexInterface& FMidiClockSubdivisionTriggerOperator::GetVertexInterface()
+	namespace Outputs
 	{
-		using namespace MidiClockSubdivisionTriggerPinNames;
-		using namespace CommonPinNames;
+		DEFINE_OUTPUT_METASOUND_PARAM(TriggerOutput, "Trigger Out", "A series of triggers at the specified subdivision grid.")
+	}
 
+	const FVertexInterface& FMidiClockSubdivisionTriggerOperator::GetVertexInterface()
+	{
 		static const FVertexInterface Interface(
 			FInputVertexInterface(
 				TInputDataVertex<FMidiClock>(METASOUND_GET_PARAM_NAME_AND_METADATA(Inputs::MidiClock)),
@@ -125,7 +136,7 @@ namespace HarmonixMetasound
 				TInputDataVertex<int32>(METASOUND_GET_PARAM_NAME_AND_METADATA(Inputs::OffsetMult), 0)
 				),
 			FOutputVertexInterface(
-				TOutputDataVertex<FTrigger>(METASOUND_GET_PARAM_NAME_AND_METADATA(TriggerOutput))
+				TOutputDataVertex<FTrigger>(METASOUND_GET_PARAM_NAME_AND_METADATA(Outputs::TriggerOutput))
 			)
 		);
 
@@ -134,9 +145,6 @@ namespace HarmonixMetasound
 
 	TUniquePtr<IOperator> FMidiClockSubdivisionTriggerOperator::CreateOperator(const FBuildOperatorParams& InParams, FBuildResults& OutResults)
 	{
-		using namespace MidiClockSubdivisionTriggerPinNames;
-		using namespace CommonPinNames;
-
 		const FMidiClockSubdivisionTriggerNode& TheNode = static_cast<const FMidiClockSubdivisionTriggerNode&>(InParams.Node);
 
 		const FInputVertexInterfaceData& InputData = InParams.InputData;
@@ -180,112 +188,46 @@ namespace HarmonixMetasound
 
 	void FMidiClockSubdivisionTriggerOperator::BindInputs(FInputVertexInterfaceData& InVertexData)
 	{
-		using namespace CommonPinNames;
-
 		InVertexData.BindReadVertex(METASOUND_GET_PARAM_NAME(Inputs::Enable),    EnableInPin);
 		InVertexData.BindReadVertex(METASOUND_GET_PARAM_NAME(Inputs::MidiClock), MidiClockInPin);
 		InVertexData.BindReadVertex(METASOUND_GET_PARAM_NAME(Inputs::GridSizeMult), GridSizeMultInPin);
 		InVertexData.BindReadVertex(METASOUND_GET_PARAM_NAME(Inputs::GridSizeUnits), GridSizeUnitsInPin);
 		InVertexData.BindReadVertex(METASOUND_GET_PARAM_NAME(Inputs::OffsetMult), GridOffsetMultInPin);
 		InVertexData.BindReadVertex(METASOUND_GET_PARAM_NAME(Inputs::OffsetUnits), GridOffsetUnitsInPin);
+
+		PulseGenerator.Reset();
+		PulseGenerator.SetClock(MidiClockInPin->AsShared());
 	}
 
 	void FMidiClockSubdivisionTriggerOperator::BindOutputs(FOutputVertexInterfaceData& InVertexData)
 	{
-		InVertexData.BindReadVertex(METASOUND_GET_PARAM_NAME(MidiClockSubdivisionTriggerPinNames::TriggerOutput), TriggerOutPin);
+		InVertexData.BindReadVertex(METASOUND_GET_PARAM_NAME(Outputs::TriggerOutput), TriggerOutPin);
 	}
 	
 	void FMidiClockSubdivisionTriggerOperator::Reset(const FResetParams& ResetParams)
 	{
 		TriggerOutPin->Reset();
 		
-		GridSizeUnits = *GridSizeUnitsInPin;
-		GridSizeMultiplier = *GridSizeMultInPin;
-		if (GridSizeMultiplier < 1)
-		{
-			GridSizeMultiplier = 1;
-		}
-		GridOffsetUnits = *GridOffsetUnitsInPin;
-		GridOffsetMultiplier = *GridOffsetMultInPin;
-
-		const ISongMapEvaluator& Map = MidiClockInPin->GetSongMapEvaluator();
-		GridOffsetTicks = Map.SubdivisionToMidiTicks(GridOffsetUnits, 0) * GridOffsetMultiplier;
-		GridSizeTicks = Map.SubdivisionToMidiTicks(GridSizeUnits, 0) * GridSizeMultiplier;
+		PulseGenerator.Reset();
+		PulseGenerator.SetClock(MidiClockInPin->AsShared());
 	}
 
 	void FMidiClockSubdivisionTriggerOperator::Execute()
 	{
+		// Update the pulse generator
+		PulseGenerator.Enable(*EnableInPin);
+		PulseGenerator.SetInterval({
+			*GridSizeUnitsInPin,
+			*GridOffsetUnitsInPin,
+			static_cast<uint16>(*GridSizeMultInPin),
+			static_cast<uint16>(*GridOffsetMultInPin) });
+
+		// If there were pulses, trigger the output
 		TriggerOutPin->AdvanceBlock();
-			
-		// first let's see if our configuration has changed at all...
-		if (*GridSizeUnitsInPin != GridSizeUnits || *GridSizeMultInPin != GridSizeMultiplier)
-		{
-			GridSizeUnits = *GridSizeUnitsInPin;
-			GridSizeMultiplier = *GridSizeMultInPin;
-			if (GridSizeMultiplier < 1)
-			{
-				GridSizeMultiplier = 1;
-			}
-		}
-
-		if (*GridOffsetUnitsInPin != GridOffsetUnits || *GridOffsetMultInPin != GridOffsetMultiplier)
-		{
-			GridOffsetUnits = *GridOffsetUnitsInPin;
-			GridOffsetMultiplier = *GridOffsetMultInPin;
-			const ISongMapEvaluator& Map = MidiClockInPin->GetSongMapEvaluator();
-			GridOffsetTicks = Map.SubdivisionToMidiTicks(GridOffsetUnits, 0) * GridOffsetMultiplier;
-		}
-
-		if (*EnableInPin)
+		PulseGenerator.Process([this](const Harmonix::Midi::Ops::FPulseGenerator::FPulseInfo& Pulse)
 		{
-			using namespace MidiClockMessageTypes;
-
-			for (const FMidiClockEvent& ClockEvent : MidiClockInPin->GetMidiClockEventsInBlock())
-			{
-				if (const FAdvance* AsAdvance = ClockEvent.TryGet<FAdvance>())
-				{
-					int32 TickPreceedingThisAdvance = (AsAdvance->FirstTickToProcess - 1) - GridOffsetTicks;
-					int32 LastTickProcessed = AsAdvance->LastTickToProcess() - GridOffsetTicks;
-					const ISongMapEvaluator& Map = MidiClockInPin->GetSongMapEvaluator();
-					
-					// calculate new GridSizeTicks and GridOffsetTicks at the tick we're evaluating based on any time signature changes
-					GridSizeTicks = Map.SubdivisionToMidiTicks(GridSizeUnits, AsAdvance->LastTickToProcess()) * GridSizeMultiplier;
-					
-					// if this is one of the quantization types that is variable in size we have to do it a little different...
-					if (GridSizeUnits == EMidiClockSubdivisionQuantization::Bar ||
-						GridSizeUnits == EMidiClockSubdivisionQuantization::Beat)
-					{
-						
-						FMusicTimestamp Start = Map.TickToMusicTimestamp(TickPreceedingThisAdvance);
-						FMusicTimestamp End = Map.TickToMusicTimestamp(LastTickProcessed);
-						if (GridSizeUnits == EMidiClockSubdivisionQuantization::Bar && Start.Bar != End.Bar)
-						{
-							if (End.Bar % GridSizeMultiplier == 0)
-							{
-								TriggerOutPin->TriggerFrame(ClockEvent.BlockFrameIndex);
-							}
-						}
-						else if (GridSizeUnits == EMidiClockSubdivisionQuantization::Beat && FMath::FloorToInt32(Start.Beat) != FMath::FloorToInt32(End.Beat))
-						{
-							if (FMath::FloorToInt32(End.Beat-1.0f) % GridSizeMultiplier == 0)
-							{
-								TriggerOutPin->TriggerFrame(ClockEvent.BlockFrameIndex);
-							}
-						}
-					}
-					else
-					{
-						
-						int32 StartingGridSquare = TickPreceedingThisAdvance / GridSizeTicks;
-						int32 EndingGridSquare = LastTickProcessed / GridSizeTicks;
-						if (StartingGridSquare != EndingGridSquare)
-						{
-							TriggerOutPin->TriggerFrame(ClockEvent.BlockFrameIndex);
-						}
-					}
-				}
-			}
-		}
+			TriggerOutPin->TriggerFrame(Pulse.BlockFrameIndex);
+		});
 	}
 }
 
diff --git a/Engine/Plugins/Runtime/Harmonix/Source/HarmonixMetasound/Private/HarmonixMetasound/Nodes/MidiPulseGeneratorNode.cpp b/Engine/Plugins/Runtime/Harmonix/Source/HarmonixMetasound/Private/HarmonixMetasound/Nodes/MidiPulseGeneratorNode.cpp
--- a/Engine/Plugins/Runtime/Harmonix/Source/HarmonixMetasound/Private/HarmonixMetasound/Nodes/MidiPulseGeneratorNode.cpp
+++ b/Engine/Plugins/Runtime/Harmonix/Source/HarmonixMetasound/Private/HarmonixMetasound/Nodes/MidiPulseGeneratorNode.cpp
@@ -74,7 +74,7 @@ namespace HarmonixMetasound::Nodes::MidiPulseGeneratorNode
 		{
 			using namespace Metasound;
 
-			static const Harmonix::Midi::Ops::FPulseGenerator PulseGeneratorForDefaults;
+			static const Harmonix::Midi::Ops::FMidiPulseGenerator PulseGeneratorForDefaults;
 			const auto DefaultInterval = PulseGeneratorForDefaults.GetInterval();
 
 			static const FVertexInterface Interface
@@ -159,6 +159,7 @@ namespace HarmonixMetasound::Nodes::MidiPulseGeneratorNode
 
 		void Reset(const FResetParams&)
 		{
+			PulseGenerator.Reset();
 			PulseGenerator.SetClock(Inputs.Clock->AsShared());
 			ApplyParameters();
 		}
@@ -175,6 +176,7 @@ namespace HarmonixMetasound::Nodes::MidiPulseGeneratorNode
 			InVertexData.BindReadVertex(Inputs::MidiNoteNumberName, Inputs.NoteNumber);
 			InVertexData.BindReadVertex(Inputs::MidiVelocityName, Inputs.Velocity);
 
+			PulseGenerator.Reset();
 			PulseGenerator.SetClock(Inputs.Clock->AsShared());
 			ApplyParameters();
 		}
@@ -208,7 +210,7 @@ namespace HarmonixMetasound::Nodes::MidiPulseGeneratorNode
 		
 		FInputs Inputs;
 		FOutputs Outputs;
-		Harmonix::Midi::Ops::FPulseGenerator PulseGenerator;
+		Harmonix::Midi::Ops::FMidiPulseGenerator PulseGenerator;
 	};
 
 	class FMidiPulseGeneratorNode final : public Metasound::FNodeFacade
diff --git a/Engine/Plugins/Runtime/Harmonix/Source/HarmonixMetasound/Public/HarmonixMetasound/MidiOps/PulseGenerator.h b/Engine/Plugins/Runtime/Harmonix/Source/HarmonixMetasound/Public/HarmonixMetasound/MidiOps/PulseGenerator.h
--- a/Engine/Plugins/Runtime/Harmonix/Source/HarmonixMetasound/Public/HarmonixMetasound/MidiOps/PulseGenerator.h
+++ b/Engine/Plugins/Runtime/Harmonix/Source/HarmonixMetasound/Public/HarmonixMetasound/MidiOps/PulseGenerator.h
@@ -2,8 +2,6 @@
 
 #pragma once
 
-#include "Containers/SpscQueue.h"
-
 #include "HarmonixDsp/Parameters/Parameter.h"
 
 #include "HarmonixMetasound/DataTypes/MidiClock.h"
@@ -17,10 +15,38 @@ namespace Harmonix::Midi::Ops
 	class HARMONIXMETASOUND_API FPulseGenerator
 	{
 	public:
+		virtual ~FPulseGenerator() = default;
+		
 		void Enable(bool bEnable);
+
+		void SetClock(const TSharedPtr<const HarmonixMetasound::FMidiClock, ESPMode::NotThreadSafe>& NewClock);
 		
-		void SetClock(const TSharedPtr<const HarmonixMetasound::FMidiClock, ESPMode::NotThreadSafe>& Clock);
+		void SetInterval(const FMusicTimeInterval& NewInterval);
+		FMusicTimeInterval GetInterval() const { return Interval; }
 
+		virtual void Reset();
+		
+		struct FPulseInfo
+		{
+			int32 BlockFrameIndex;
+			int32 Tick;
+		};
+		
+		void Process(const TFunctionRef<void(const FPulseInfo&)>& OnPulse);
+
+	protected:
+		bool Enabled{ true };
+		TWeakPtr<const HarmonixMetasound::FMidiClock, ESPMode::NotThreadSafe> Clock;
+		FMusicTimeInterval Interval{};
+		FTimeSignature CurrentTimeSignature{};
+		FMusicTimestamp NextPulseTimestamp{ -1, -1 };
+	};
+	
+	class HARMONIXMETASOUND_API FMidiPulseGenerator : public FPulseGenerator
+	{
+	public:
+		virtual ~FMidiPulseGenerator() override = default;
+		
 		Dsp::Parameters::TParameter<uint8> Channel{ 1, 16, 1 };
 
 		Dsp::Parameters::TParameter<uint16> Track{ 1, UINT16_MAX, 1 };
@@ -28,24 +54,15 @@ namespace Harmonix::Midi::Ops
 		Dsp::Parameters::TParameter<uint8> NoteNumber{ 0, 127, 60 };
 		
 		Dsp::Parameters::TParameter<uint8> Velocity{ 0, 127, 127 };
-		
-		void SetInterval(const FMusicTimeInterval& NewInterval);
 
-		FMusicTimeInterval GetInterval() const { return Interval; }
+		virtual void Reset() override;
 
 		void Process(HarmonixMetasound::FMidiStream& OutStream);
 
 	private:
-		void HandleAdvance(const int32 BlockFrameIndex, const HarmonixMetasound::MidiClockMessageTypes::FAdvance* Advance, HarmonixMetasound::FMidiStream& OutStream);
-		void HandleTimeSignatureChange(const HarmonixMetasound::MidiClockMessageTypes::FTimeSignatureChange* TimeSigChange);
-		void DoPulse(const int32 BlockFrameIndex, const int32 PulseTick, HarmonixMetasound::FMidiStream& OutStream);
+		void AddPulseNote(const int32 BlockFrameIndex, const int32 PulseTick, HarmonixMetasound::FMidiStream& OutStream);
 
 		FMidiVoiceGeneratorBase VoiceGenerator{};
 		TOptional<HarmonixMetasound::FMidiStreamEvent> LastNoteOn;
-		bool Enabled{ true };
-		TWeakPtr<const HarmonixMetasound::FMidiClock, ESPMode::NotThreadSafe> Clock;
-		FMusicTimeInterval Interval{};
-		FTimeSignature CurrentTimeSignature{};
-		FMusicTimestamp NextPulseTimestamp{ -1, -1 };
 	};
 }
diff --git a/Engine/Plugins/Runtime/Harmonix/Source/HarmonixMetasound/Public/HarmonixMetasound/Nodes/MidiClockSubdivisionTriggerNode.h b/Engine/Plugins/Runtime/Harmonix/Source/HarmonixMetasound/Public/HarmonixMetasound/Nodes/MidiClockSubdivisionTriggerNode.h
new file mode 100644
--- /dev/null
+++ b/Engine/Plugins/Runtime/Harmonix/Source/HarmonixMetasound/Public/HarmonixMetasound/Nodes/MidiClockSubdivisionTriggerNode.h
@@ -0,0 +1,28 @@
+// Copyright Epic Games, Inc. All Rights Reserved.
+
+#pragma once
+
+#include "MetasoundNodeInterface.h"
+
+#include "HarmonixMetasound/Common.h"
+
+namespace HarmonixMetasound::Nodes::MidiClockSubdivisionTriggerNode
+{
+	const HARMONIXMETASOUND_API Metasound::FNodeClassName& GetClassName();
+	HARMONIXMETASOUND_API int32 GetCurrentMajorVersion();
+
+	namespace Inputs
+	{
+		DECLARE_METASOUND_PARAM_ALIAS(MidiClock);
+		DECLARE_METASOUND_PARAM_ALIAS(Enable);
+		DECLARE_METASOUND_PARAM_ALIAS(GridSizeUnits);
+		DECLARE_METASOUND_PARAM_ALIAS(GridSizeMult);
+		DECLARE_METASOUND_PARAM_ALIAS(OffsetUnits);
+		DECLARE_METASOUND_PARAM_ALIAS(OffsetMult);
+	}
+
+	namespace Outputs
+	{
+		DECLARE_METASOUND_PARAM_EXTERN(TriggerOutput);
+	}
+}
\ No newline at end of file
diff --git a/Engine/Plugins/Runtime/Harmonix/Source/HarmonixMetasoundTests/Private/MidiOps/PulseGeneratorTests.cpp b/Engine/Plugins/Runtime/Harmonix/Source/HarmonixMetasoundTests/Private/MidiOps/PulseGeneratorTests.cpp
--- a/Engine/Plugins/Runtime/Harmonix/Source/HarmonixMetasoundTests/Private/MidiOps/PulseGeneratorTests.cpp
+++ b/Engine/Plugins/Runtime/Harmonix/Source/HarmonixMetasoundTests/Private/MidiOps/PulseGeneratorTests.cpp
@@ -13,7 +13,7 @@ namespace Harmonix::Midi::Ops::Tests
 		EAutomationTestFlags::EditorContext | EAutomationTestFlags::EngineFilter)
 	bool FMidiPulseGeneratorBasicTest::RunTest(const FString&)
 	{
-		FPulseGenerator PulseGenerator;
+		FMidiPulseGenerator PulseGenerator;
 		
 		Metasound::FOperatorSettings OperatorSettings{ 48000, 100 };
 
diff --git a/Engine/Plugins/Runtime/Harmonix/Source/HarmonixMetasoundTests/Private/Nodes/MidiClockSubdivisionTriggerNodeTests.cpp b/Engine/Plugins/Runtime/Harmonix/Source/HarmonixMetasoundTests/Private/Nodes/MidiClockSubdivisionTriggerNodeTests.cpp
new file mode 100644
--- /dev/null
+++ b/Engine/Plugins/Runtime/Harmonix/Source/HarmonixMetasoundTests/Private/Nodes/MidiClockSubdivisionTriggerNodeTests.cpp
@@ -0,0 +1,64 @@
+// Copyright Epic Games, Inc. All Rights Reserved.
+
+#include "NodeTestGraphBuilder.h"
+#include "HarmonixMetasound/MidiOps/PulseGenerator.h"
+#include "HarmonixMetasound/Nodes/MidiClockSubdivisionTriggerNode.h"
+#include "Misc/AutomationTest.h"
+
+#if WITH_DEV_AUTOMATION_TESTS
+
+namespace HarmonixMetasound::Nodes::MidiClockSubdivisionTriggerNode::Tests
+{
+	IMPLEMENT_SIMPLE_AUTOMATION_TEST(
+	FMidiClockSubdivisionTriggerNodeParityTest,
+	"Harmonix.Metasound.Nodes.MidiClockSubdivisionTrigger.PulseGeneratorParity",
+	EAutomationTestFlags::EditorContext | EAutomationTestFlags::EngineFilter)
+	bool FMidiClockSubdivisionTriggerNodeParityTest::RunTest(const FString&)
+	{
+		const TUniquePtr<Metasound::FMetasoundGenerator> Generator =
+			Metasound::Test::FNodeTestGraphBuilder::MakeSingleNodeGraph(GetClassName(), GetCurrentMajorVersion());
+		TOptional<Metasound::TDataReadReference<Metasound::FTrigger>> NodeTriggerOutput =
+			Generator->GetOutputReadReference<Metasound::FTrigger>(Outputs::TriggerOutputName);
+		UTEST_TRUE("Got node MIDI output", NodeTriggerOutput.IsSet());
+
+		TOptional<FMidiClockWriteRef> Clock = Generator->GetInputWriteReference<FMidiClock>(Inputs::MidiClockName);
+		UTEST_TRUE("Got clock", Clock.IsSet());
+		(*Clock)->SetTransportState(0, EMusicPlayerTransportState::Playing);
+
+		Harmonix::Midi::Ops::FPulseGenerator PulseGenerator;
+		PulseGenerator.SetClock((*Clock)->AsShared());
+
+		// Render for a bit and expect the same output from both the node and the raw processor
+		constexpr int32 NumBlocks = 1000;
+		Metasound::FAudioBuffer Buffer{ Generator->OperatorSettings };
+		
+		for (int32 BlockIdx = 0; BlockIdx < NumBlocks; ++BlockIdx)
+		{
+			// Advance the clock, which will advance the play cursor in the pulse generators
+			(*Clock)->PrepareBlock();
+			(*Clock)->Advance(0, Generator->OperatorSettings.GetNumFramesPerBlock());
+
+			// Process
+			Generator->OnGenerateAudio(Buffer.GetData(), Buffer.Num());
+
+			TArray<int32> PulseGeneratorTriggers;
+			PulseGenerator.Process([&PulseGeneratorTriggers](const Harmonix::Midi::Ops::FPulseGenerator::FPulseInfo& Pulse)
+			{
+				PulseGeneratorTriggers.Add(Pulse.BlockFrameIndex);
+			});
+
+			const TArray<int32>& NodeTriggers = (*NodeTriggerOutput)->GetTriggeredFrames();
+			
+			UTEST_EQUAL("Same number of triggers", (*NodeTriggerOutput)->NumTriggeredInBlock(), PulseGeneratorTriggers.Num());
+
+			for (const int32 Frame : NodeTriggers)
+			{
+				UTEST_TRUE("Trigger frame was in both arrays", PulseGeneratorTriggers.Contains(Frame));
+			}
+		}
+		
+		return true;
+	}
+}
+
+#endif
\ No newline at end of file
diff --git a/Engine/Plugins/Runtime/Harmonix/Source/HarmonixMetasoundTests/Private/Nodes/MidiPulseGeneratorNodeTests.cpp b/Engine/Plugins/Runtime/Harmonix/Source/HarmonixMetasoundTests/Private/Nodes/MidiPulseGeneratorNodeTests.cpp
--- a/Engine/Plugins/Runtime/Harmonix/Source/HarmonixMetasoundTests/Private/Nodes/MidiPulseGeneratorNodeTests.cpp
+++ b/Engine/Plugins/Runtime/Harmonix/Source/HarmonixMetasoundTests/Private/Nodes/MidiPulseGeneratorNodeTests.cpp
@@ -24,7 +24,7 @@ namespace HarmonixMetasound::Nodes::MidiPulseGeneratorNode::Tests
 		UTEST_TRUE("Got clock", Clock.IsSet());
 		(*Clock)->SetTransportState(0, EMusicPlayerTransportState::Playing);
 
-		Harmonix::Midi::Ops::FPulseGenerator PulseGenerator;
+		Harmonix::Midi::Ops::FMidiPulseGenerator PulseGenerator;
 		FMidiStream PulseGeneratorMidiOutput;
 		PulseGenerator.SetClock((*Clock)->AsShared());
 
