Set tempo and speed on input re-bind
#tests PIE to test changes

diff --git a/Engine/Plugins/Runtime/Harmonix/Source/HarmonixMetasound/Private/HarmonixMetasound/Nodes/DelayNode.cpp b/Engine/Plugins/Runtime/Harmonix/Source/HarmonixMetasound/Private/HarmonixMetasound/Nodes/DelayNode.cpp
--- a/Engine/Plugins/Runtime/Harmonix/Source/HarmonixMetasound/Private/HarmonixMetasound/Nodes/DelayNode.cpp
+++ b/Engine/Plugins/Runtime/Harmonix/Source/HarmonixMetasound/Private/HarmonixMetasound/Nodes/DelayNode.cpp
@@ -277,6 +277,13 @@ namespace HarmonixMetasound::DelayNode
 			if (Inputs.MidiClock.IsSet())
 			{
 				InVertexData.BindReadVertex(Inputs::MidiClockName, *Inputs.MidiClock);
+
+				// Set the tempo and speed after we bind here
+				const FMidiClockReadRef Clock = *Inputs.MidiClock;
+				LastTempo = Clock->GetTempoAtStartOfBlock();
+				LastSpeed = Clock->GetSpeedAtStartOfBlock();
+				Delay.SetTempo(LastTempo);
+				Delay.SetSpeed(LastSpeed);
 			}
 			InVertexData.BindReadVertex(Inputs::DelayTimeTypeName, Inputs.DelayTimeType);
 			InVertexData.BindReadVertex(Inputs::DelayTimeName, Inputs.DelayTime);
