[HarmonixMetasound] Support tempo changes after seeking to an earlier point in time

diff --git a/Engine/Plugins/Runtime/Harmonix/Source/HarmonixMetasound/Private/HarmonixMetasound/Nodes/MetronomeNode.cpp b/Engine/Plugins/Runtime/Harmonix/Source/HarmonixMetasound/Private/HarmonixMetasound/Nodes/MetronomeNode.cpp
--- a/Engine/Plugins/Runtime/Harmonix/Source/HarmonixMetasound/Private/HarmonixMetasound/Nodes/MetronomeNode.cpp
+++ b/Engine/Plugins/Runtime/Harmonix/Source/HarmonixMetasound/Private/HarmonixMetasound/Nodes/MetronomeNode.cpp
@@ -13,6 +13,7 @@
 #include "HarmonixMetasound/DataTypes/MusicTransport.h"
 #include <algorithm>
 
+#include "HAL/IConsoleManager.h"
 #include "HAL/PlatformMath.h"
 
 #define LOCTEXT_NAMESPACE "HarmonixMetaSound"
@@ -23,6 +24,13 @@ namespace HarmonixMetasound
 {
 	using namespace Metasound;
 
+	bool bSkipMetronomeLastProcessedClockTickCheck = true;
+	FAutoConsoleVariableRef CVarSkipMetronomeLastProcessedClockTickCheck(
+		TEXT("au.Metronome.SkipLastProcessedClockTickCheck"),
+		bSkipMetronomeLastProcessedClockTickCheck,
+		TEXT("Skip Last Processed Clock Tick conditions when executing a Metronome Metasound Node."),
+		ECVF_Default);
+
 	class FMetronomeOperator : public TExecutableOperator<FMetronomeOperator>, public FMusicTransportControllable
 	{
 	public:
@@ -273,7 +281,7 @@ namespace HarmonixMetasound
 		// only update our midi data if the clock is advancing
 		// update the tempo and time sig before we advance our clock 
 		int32 ClockTick = DrivingMidiClock.GetLastProcessedMidiTick();
-		if (ClockTick >= 0 && ClockTick > LastProcessedClockTick)
+		if ((ClockTick >= 0 && ClockTick > LastProcessedClockTick) || bSkipMetronomeLastProcessedClockTickCheck)
 		{
 			UpdateMidi();
 			LastProcessedClockTick = ClockTick;
