[HarmonixMetasound] Downgrading severity of assert/warnings for time signature changes on a looping metronome.

diff --git a/Engine/Plugins/Runtime/Harmonix/Source/HarmonixMetasound/Private/HarmonixMetasound/Nodes/MetronomeNode.cpp b/Engine/Plugins/Runtime/Harmonix/Source/HarmonixMetasound/Private/HarmonixMetasound/Nodes/MetronomeNode.cpp
--- a/Engine/Plugins/Runtime/Harmonix/Source/HarmonixMetasound/Private/HarmonixMetasound/Nodes/MetronomeNode.cpp
+++ b/Engine/Plugins/Runtime/Harmonix/Source/HarmonixMetasound/Private/HarmonixMetasound/Nodes/MetronomeNode.cpp
@@ -385,14 +385,14 @@ namespace HarmonixMetasound
 		{
 			CurrentTimeSigNum = InTimeSigNum;
 			CurrentTimeSigDenom = InTimeSigDenom;
-			if (ensure(!MidiClockOutPin->HasPersistentLoop()))
+			if (!MidiClockOutPin->HasPersistentLoop())
 			{
 				AddTimeSigChangeForMidi(InTimeSigNum, InTimeSigDenom);
 				HasMidiChanges = true;
 			}
 			else
 			{
-				UE_LOG(LogMetronomeNode, Warning, TEXT("Changing Time Sig. on looping metronome not supported." 
+				UE_LOG(LogMetronomeNode, Log, TEXT("Changing Time Sig. on looping metronome not supported." 
 					"Changing time signature will require changing loop length which is currently not supported."))
 			}
 		}
