[Harmonix] Reset pulse generator phase on new clock

diff --git a/Engine/Plugins/Runtime/Harmonix/Source/HarmonixMetasound/Private/HarmonixMetasound/MidiOps/PulseGenerator.cpp b/Engine/Plugins/Runtime/Harmonix/Source/HarmonixMetasound/Private/HarmonixMetasound/MidiOps/PulseGenerator.cpp
--- a/Engine/Plugins/Runtime/Harmonix/Source/HarmonixMetasound/Private/HarmonixMetasound/MidiOps/PulseGenerator.cpp
+++ b/Engine/Plugins/Runtime/Harmonix/Source/HarmonixMetasound/Private/HarmonixMetasound/MidiOps/PulseGenerator.cpp
@@ -13,6 +13,16 @@ namespace Harmonix::Midi::Ops
 	{
 		Clock = NewClock;
 		CurrentTimeSignature = *NewClock->GetSongMapEvaluator().GetTimeSignatureAtTick(NewClock->GetLastProcessedMidiTick());
+
+		// Find the next pulse and line up phase with the current bar
+		const FMusicTimestamp ClockCurrentTimestamp = NewClock->GetMusicTimestampAtBlockOffset(0);
+		NextPulseTimestamp.Bar = ClockCurrentTimestamp.Bar;
+		NextPulseTimestamp.Beat = 1;
+		IncrementTimestampByOffset(NextPulseTimestamp, Interval, CurrentTimeSignature);
+		while (NextPulseTimestamp < ClockCurrentTimestamp)
+		{
+			IncrementTimestampByInterval(NextPulseTimestamp, Interval, CurrentTimeSignature);
+		}
 	}
 
 	void FPulseGenerator::SetInterval(const FMusicTimeInterval& NewInterval)
