[Harmonix] Handle seek in pulse generator

diff --git a/Engine/Plugins/Runtime/Harmonix/Source/HarmonixMetasound/Private/HarmonixMetasound/MidiOps/PulseGenerator.cpp b/Engine/Plugins/Runtime/Harmonix/Source/HarmonixMetasound/Private/HarmonixMetasound/MidiOps/PulseGenerator.cpp
--- a/Engine/Plugins/Runtime/Harmonix/Source/HarmonixMetasound/Private/HarmonixMetasound/MidiOps/PulseGenerator.cpp
+++ b/Engine/Plugins/Runtime/Harmonix/Source/HarmonixMetasound/Private/HarmonixMetasound/MidiOps/PulseGenerator.cpp
@@ -86,6 +86,18 @@ namespace Harmonix::Midi::Ops
 				NextPulseTimestamp = PinnedClock->GetSongMapEvaluator().TickToMusicTimestamp(AsLoop->FirstTickInLoop);
 				IncrementTimestampByOffset(NextPulseTimestamp, Interval, CurrentTimeSignature);
 			}
+			else if (const FSeek* AsSeek = ClockEvent.TryGet<FSeek>())
+			{
+				// When we seek, reset the pulse phase to the current bar
+				const FMusicTimestamp ClockCurrentTimestamp = PinnedClock->GetSongMapEvaluator().TickToMusicTimestamp(AsSeek->NewNextTick);
+				NextPulseTimestamp.Bar = ClockCurrentTimestamp.Bar;
+				NextPulseTimestamp.Beat = 1;
+				IncrementTimestampByOffset(NextPulseTimestamp, Interval, CurrentTimeSignature);
+				while (NextPulseTimestamp < ClockCurrentTimestamp)
+				{
+					IncrementTimestampByInterval(NextPulseTimestamp, Interval, CurrentTimeSignature);
+				}
+			}
 		}
 	}
 
diff --git a/Engine/Plugins/Runtime/Harmonix/Source/HarmonixMetasoundTests/Private/MidiOps/PulseGeneratorTests.cpp b/Engine/Plugins/Runtime/Harmonix/Source/HarmonixMetasoundTests/Private/MidiOps/PulseGeneratorTests.cpp
--- a/Engine/Plugins/Runtime/Harmonix/Source/HarmonixMetasoundTests/Private/MidiOps/PulseGeneratorTests.cpp
+++ b/Engine/Plugins/Runtime/Harmonix/Source/HarmonixMetasoundTests/Private/MidiOps/PulseGeneratorTests.cpp
@@ -131,6 +131,104 @@ namespace Harmonix::Midi::Ops::Tests
 
 		return true;
 	}
+
+	IMPLEMENT_SIMPLE_AUTOMATION_TEST(
+	FMidiPulseGeneratorSeekTest,
+	"Harmonix.Midi.Ops.PulseGenerator.Seek",
+	EAutomationTestFlags::EditorContext | EAutomationTestFlags::EngineFilter)
+	bool FMidiPulseGeneratorSeekTest::RunTest(const FString&)
+	{
+		FPulseGenerator PulseGenerator;
+		
+		Metasound::FOperatorSettings OperatorSettings{ 48000, 100 };
+
+		constexpr float Tempo = 123;
+		const FTimeSignature TimeSignature(4, 4);
+		const auto Clock = MakeShared<HarmonixMetasound::FMidiClock, ESPMode::NotThreadSafe>(OperatorSettings);
+		Clock->AttachToSongMapEvaluator(MakeShared<FSongMaps>(Tempo, TimeSignature.Numerator, TimeSignature.Denominator));
+
+		PulseGenerator.SetClock(Clock);
+
+		// Advance forward a few pulses, then seek back and ensure we keep getting notes
+		{
+			constexpr int32 PulsesToDo = 10;
+			int32 NumPulsesReceived = 0;
+
+			Clock->SeekTo(0,0,0);
+			Clock->SetTransportState(0, HarmonixMetasound::EMusicPlayerTransportState::Playing);
+			const FMusicTimeInterval Interval = PulseGenerator.GetInterval();
+			FMusicTimestamp NextPulse{ 1, 1 };
+			IncrementTimestampByOffset(NextPulse, Interval, TimeSignature);
+
+			FMusicTimestamp EndTimestamp{ NextPulse };
+			for (int i = 0; i < PulsesToDo; ++i)
+			{
+				IncrementTimestampByInterval(EndTimestamp, Interval, TimeSignature);
+			}
+			
+			while (Clock->GetMusicTimestampAtBlockEnd() < EndTimestamp)
+			{
+				// Advance the clock, which will advance the play cursor in the pulse generator
+				Clock->Advance(0, OperatorSettings.GetNumFramesPerBlock());
+
+				// Process
+				TArray<FPulseGenerator::FPulseInfo> PulsesThisBlock;
+				PulseGenerator.Process([&PulsesThisBlock](const FPulseGenerator::FPulseInfo& Pulse)
+				{
+					PulsesThisBlock.Add(Pulse);
+				});
+				NumPulsesReceived += PulsesThisBlock.Num();
+				
+				// If this is a block where we should get a pulse, check that we got it
+				if (Clock->GetLastProcessedMidiTick() >= Clock->GetSongMapEvaluator().MusicTimestampToTick(NextPulse))
+				{
+					UTEST_EQUAL("Got the right number of pulses", PulsesThisBlock.Num(), 1);
+
+					IncrementTimestampByInterval(NextPulse, Interval, TimeSignature);
+				}
+
+				// we can prepare the clock for the next block...
+				Clock->PrepareBlock();
+			}
+			
+			UTEST_TRUE("Before seek: Got all the pulses at the right time", NumPulsesReceived >= PulsesToDo);
+
+			// Seek and reset the next pulse
+			Clock->SeekTo(0, 0, 0);
+			NumPulsesReceived = 0;
+			NextPulse = { 1, 1 };
+			IncrementTimestampByOffset(NextPulse, Interval, TimeSignature);
+
+			while (Clock->GetMusicTimestampAtBlockEnd() < EndTimestamp)
+			{
+				// Advance the clock, which will advance the play cursor in the pulse generator
+				Clock->Advance(0, OperatorSettings.GetNumFramesPerBlock());
+
+				// Process
+				TArray<FPulseGenerator::FPulseInfo> PulsesThisBlock;
+				PulseGenerator.Process([&PulsesThisBlock](const FPulseGenerator::FPulseInfo& Pulse)
+				{
+					PulsesThisBlock.Add(Pulse);
+				});
+				NumPulsesReceived += PulsesThisBlock.Num();
+				
+				// If this is a block where we should get a pulse, check that we got it
+				if (Clock->GetLastProcessedMidiTick() >= Clock->GetSongMapEvaluator().MusicTimestampToTick(NextPulse))
+				{
+					UTEST_EQUAL("Got the right number of pulses", PulsesThisBlock.Num(), 1);
+
+					IncrementTimestampByInterval(NextPulse, Interval, TimeSignature);
+				}
+
+				// we can prepare the clock for the next block...
+				Clock->PrepareBlock();
+			}
+			
+			UTEST_TRUE("After seek: Got all the notes at the right time", NumPulsesReceived >= PulsesToDo);
+		}
+
+		return true;
+	}
 }
 
 #endif
\ No newline at end of file
