[Harmonix/Fabric] Adjust the metronome node to allow for inheritance, add a fabric metronome node
#tests PIE

diff --git a/Engine/Plugins/Runtime/Harmonix/Source/HarmonixMetasound/Private/HarmonixMetasound/Nodes/MetronomeNode.cpp b/Engine/Plugins/Runtime/Harmonix/Source/HarmonixMetasound/Private/HarmonixMetasound/Nodes/MetronomeNode.cpp
--- a/Engine/Plugins/Runtime/Harmonix/Source/HarmonixMetasound/Private/HarmonixMetasound/Nodes/MetronomeNode.cpp
+++ b/Engine/Plugins/Runtime/Harmonix/Source/HarmonixMetasound/Private/HarmonixMetasound/Nodes/MetronomeNode.cpp
@@ -1,5 +1,7 @@
 // Copyright Epic Games, Inc. All Rights Reserved.
 
+#include "HarmonixMetasound/Nodes/MetronomeNode.h"
+
 #include "MetasoundExecutableOperator.h"
 #include "MetasoundFacade.h"
 #include "MetasoundNodeInterface.h"
@@ -8,9 +10,6 @@
 #include "MetasoundStandardNodesCategories.h"
 #include "MetasoundVertex.h"
 
-#include "HarmonixMetasound/Common.h"
-#include "HarmonixMetasound/DataTypes/MidiClock.h"
-#include "HarmonixMetasound/DataTypes/MusicTransport.h"
 #include <algorithm>
 
 #include "HAL/IConsoleManager.h"
@@ -20,7 +19,7 @@
 
 DEFINE_LOG_CATEGORY_STATIC(LogMetronomeNode, Log, All);
 
-namespace HarmonixMetasound
+namespace HarmonixMetasound::Nodes::MetronomeNode
 {
 	using namespace Metasound;
 
@@ -31,7 +30,7 @@ namespace HarmonixMetasound
 		TEXT("Skip Last Processed Clock Tick conditions when executing a Metronome Metasound Node."),
 		ECVF_Default);
 
-	class FMetronomeOperator : public TExecutableOperator<FMetronomeOperator>, public FMusicTransportControllable
+	class FMetronomeOperator : public FMetronomeOperatorBase
 	{
 	public:
 		static const FNodeClassMetadata& GetNodeInfo();
@@ -39,58 +38,14 @@ namespace HarmonixMetasound
 		static TUniquePtr<IOperator> CreateOperator(const FBuildOperatorParams& InParams, FBuildResults& OutResults);
 
 		FMetronomeOperator(const FBuildOperatorParams& InParams,
-						   const FMusicTransportEventStreamReadRef& InTransport,
-						   const bool  InLoop,
-						   const int32 InLoopLengthBars, 
-						   const FInt32ReadRef& InTimSigNumerator,
-						   const FInt32ReadRef& InTimeSigDenominator,
-						   const FFloatReadRef& InTempo,
-		                   const FFloatReadRef& InSpeedMultiplier,
-		                   const int32 InSeekPrerollBars);
-
-		virtual void BindInputs(FInputVertexInterfaceData& InVertexData) override;
-		virtual void BindOutputs(FOutputVertexInterfaceData& InVertexData) override;
-
-		void Reset(const FResetParams& Params);
-		void Execute();
-
-	private:
-		void Init();
-		
-		//** INPUTS
-		FMusicTransportEventStreamReadRef TransportInPin;
-		const bool  LoopInPin;
-		const int32 LoopLengthBarsInPin;
-		FInt32ReadRef TimeSigNumInPin;
-		FInt32ReadRef TimeSigDenomInPin;
-		FFloatReadRef TempoInPin;
-		FFloatReadRef SpeedMultInPin;
-		const int32 SeekPreRollBarsInPin;
-
-		//** OUTPUTS
-		FMidiClockWriteRef MidiClockOutPin;
-
-		//** DATA
-		TSharedPtr<FMidiClock, ESPMode::NotThreadSafe> MonotonicallyIncreasingClock;
-		TSharedPtr<FSongMaps> SongMaps;
-		FSampleCount BlockSize;
-		float        SampleRate;
-		float        CurrentTempo;
-		int32        CurrentTimeSigNum;
-		int32        CurrentTimeSigDenom;
-		int32		 LastProcessedClockTick = -1;
-		int32		 NextClockTickToProcess = 0;
-		bool		 bClocksArePreparedForExecute = true;
-
-		void BuildSongMaps(bool ResetToStart = true);
-		void UpdateMidi();
-		void AddTempoChangeForMidi(float TempoBPM);
-		void AddTimeSigChangeForMidi(int32 TimeSigNum, int32 TimeSigDenom);
-		void HandleTransportChange(int32 StartFrameIndex, EMusicPlayerTransportState NewTransportState);
-		void PrepareClocksForExecute();
-		void MarkClocksAsExecuted();
-
-		FMidiClock& GetDrivingMidiClock() { return LoopInPin ? *MonotonicallyIncreasingClock : (*MidiClockOutPin); }
+			const FMusicTransportEventStreamReadRef& InTransport,
+			const bool  InLoop,
+			const int32 InLoopLengthBars,
+			const FInt32ReadRef& InTimSigNumerator,
+			const FInt32ReadRef& InTimeSigDenominator,
+			const FFloatReadRef& InTempo,
+			const FFloatReadRef& InSpeedMultiplier,
+			const int32 InSeekPrerollBars);
 	};
 
 	class FMetronomeNode : public FNodeFacade
@@ -171,7 +126,7 @@ namespace HarmonixMetasound
 		return MakeUnique<FMetronomeOperator>(InParams, InTransport, InLoop, InLoopLengthBars, InTimeSigNumerator, InTimeSigDenominator, InTempo, InSpeed, InPreRollBars);
 	}
 
-	FMetronomeOperator::FMetronomeOperator(const FBuildOperatorParams& InParams, 
+	FMetronomeOperatorBase::FMetronomeOperatorBase(const FBuildOperatorParams& InParams,
 	                                       const FMusicTransportEventStreamReadRef& InTransport,
 										   const bool  InLoop,
 										   const int32 InLoopLengthBars,
@@ -201,7 +156,20 @@ namespace HarmonixMetasound
 		Init();
 	}
 
-	void FMetronomeOperator::BindInputs(FInputVertexInterfaceData& InVertexData)
+	FMetronomeOperator::FMetronomeOperator(const FBuildOperatorParams& InParams,
+		const FMusicTransportEventStreamReadRef& InTransport,
+		const bool  InLoop,
+		const int32 InLoopLengthBars,
+		const FInt32ReadRef& InTimSigNumerator,
+		const FInt32ReadRef& InTimeSigDenominator,
+		const FFloatReadRef& InTempo,
+		const FFloatReadRef& InSpeedMultiplier,
+		const int32 InSeekPrerollBars)
+		: FMetronomeOperatorBase(InParams, InTransport, InLoop, InLoopLengthBars, InTimSigNumerator, InTimeSigDenominator, InTempo, InSpeedMultiplier, InSeekPrerollBars)
+	{
+	}
+
+	void FMetronomeOperatorBase::BindInputs(FInputVertexInterfaceData& InVertexData)
 	{
 		using namespace CommonPinNames;
 		InVertexData.BindReadVertex(METASOUND_GET_PARAM_NAME(Inputs::Transport), TransportInPin);
@@ -216,13 +184,13 @@ namespace HarmonixMetasound
 		Init();
 	}
 
-	void FMetronomeOperator::BindOutputs(FOutputVertexInterfaceData& InVertexData)
+	void FMetronomeOperatorBase::BindOutputs(FOutputVertexInterfaceData& InVertexData)
 	{
 		using namespace CommonPinNames;
 		InVertexData.BindReadVertex(METASOUND_GET_PARAM_NAME(Outputs::MidiClock), MidiClockOutPin);
 	}
 
-	void FMetronomeOperator::Reset(const FResetParams& Params)
+	void FMetronomeOperatorBase::Reset(const FResetParams& Params)
 	{
 		BlockSize = Params.OperatorSettings.GetNumFramesPerBlock();
 		SampleRate = Params.OperatorSettings.GetSampleRate();
@@ -231,7 +199,7 @@ namespace HarmonixMetasound
 		NextClockTickToProcess = 0;
 	}
 
-	void FMetronomeOperator::Init()
+	void FMetronomeOperatorBase::Init()
 	{
 		bClocksArePreparedForExecute = false;
 		PrepareClocksForExecute();
@@ -267,7 +235,7 @@ namespace HarmonixMetasound
 		FMusicTransportControllable::Init(*TransportInPin, MoveTemp(InitFn));
 	}
 
-	void FMetronomeOperator::Execute()
+	void FMetronomeOperatorBase::Execute()
 	{
 		PrepareClocksForExecute();
 
@@ -338,7 +306,7 @@ namespace HarmonixMetasound
 		MarkClocksAsExecuted();
 	}
 
-	void FMetronomeOperator::BuildSongMaps(bool ResetToStart)
+	void FMetronomeOperatorBase::BuildSongMaps(bool ResetToStart)
 	{
 		// make sure we have valid values
 		CurrentTempo = FMath::Max(1.0f, *TempoInPin);
@@ -369,7 +337,7 @@ namespace HarmonixMetasound
 		}
 	}
 
-	void FMetronomeOperator::UpdateMidi()
+	void FMetronomeOperatorBase::UpdateMidi()
 	{
 		FMidiClock& DrivingMidiClock = GetDrivingMidiClock();
 		bool HasMidiChanges = false;
@@ -387,7 +355,7 @@ namespace HarmonixMetasound
 			CurrentTimeSigDenom = InTimeSigDenom;
 			if (!MidiClockOutPin->HasPersistentLoop())
 			{
-				AddTimeSigChangeForMidi(InTimeSigNum, InTimeSigDenom);
+				HandleTimeSigChangeForMidi(InTimeSigNum, InTimeSigDenom);
 				HasMidiChanges = true;
 			}
 			else
@@ -411,14 +379,14 @@ namespace HarmonixMetasound
 		}
 	}
 
-	void FMetronomeOperator::AddTempoChangeForMidi(float InTempoBPM)
+	void FMetronomeOperatorBase::AddTempoChangeForMidi(float InTempoBPM)
 	{
 		CurrentTempo = InTempoBPM;
 		int32 AtTick = GetDrivingMidiClock().GetNextMidiTickToProcess();
 		SongMaps->AddTempoChange(AtTick, CurrentTempo);
 	}
 
-	void FMetronomeOperator::AddTimeSigChangeForMidi(int32 InTimeSigNum, int32 InTimeSigDenom)
+	void FMetronomeOperatorBase::HandleTimeSigChangeForMidi(int32 InTimeSigNum, int32 InTimeSigDenom)
 	{
 		CurrentTimeSigNum = InTimeSigNum;
 		CurrentTimeSigDenom = InTimeSigDenom;
@@ -443,12 +411,12 @@ namespace HarmonixMetasound
 		}
 	}
 
-	void FMetronomeOperator::HandleTransportChange(int32 StartFrameIndex, EMusicPlayerTransportState NewTransportState)
+	void FMetronomeOperatorBase::HandleTransportChange(int32 StartFrameIndex, EMusicPlayerTransportState NewTransportState)
 	{
 		GetDrivingMidiClock().SetTransportState(StartFrameIndex, NewTransportState);
 	}
 
-	void FMetronomeOperator::PrepareClocksForExecute()
+	void FMetronomeOperatorBase::PrepareClocksForExecute()
 	{
 		if (bClocksArePreparedForExecute)
 		{
@@ -466,7 +434,7 @@ namespace HarmonixMetasound
 
 	}
 
-	void FMetronomeOperator::MarkClocksAsExecuted()
+	void FMetronomeOperatorBase::MarkClocksAsExecuted()
 	{
 		bClocksArePreparedForExecute = false;
 	}
diff --git a/Engine/Plugins/Runtime/Harmonix/Source/HarmonixMetasound/Public/HarmonixMetasound/Nodes/MetronomeNode.h b/Engine/Plugins/Runtime/Harmonix/Source/HarmonixMetasound/Public/HarmonixMetasound/Nodes/MetronomeNode.h
new file mode 100644
--- /dev/null
+++ b/Engine/Plugins/Runtime/Harmonix/Source/HarmonixMetasound/Public/HarmonixMetasound/Nodes/MetronomeNode.h
@@ -0,0 +1,84 @@
+// Copyright Epic Games, Inc. All Rights Reserved.
+
+#pragma once
+
+#include "MetasoundNodeInterface.h"
+
+#include "MetasoundExecutableOperator.h"
+#include "MetasoundFacade.h"
+#include "MetasoundNodeInterface.h"
+#include "MetasoundParamHelper.h"
+#include "MetasoundSampleCounter.h"
+#include "MetasoundStandardNodesCategories.h"
+#include "MetasoundVertex.h"
+
+#include "HarmonixMetasound/Common.h"
+#include "HarmonixMetasound/DataTypes/MidiClock.h"
+#include "HarmonixMetasound/DataTypes/MusicTransport.h"
+
+namespace HarmonixMetasound::Nodes::MetronomeNode
+{
+	using namespace Metasound;
+
+	//HARMONIXMETASOUND_API Metasound::FNodeClassName GetClassName();
+	//HARMONIXMETASOUND_API int32 GetCurrentMajorVersion();
+
+	class HARMONIXMETASOUND_API FMetronomeOperatorBase : public TExecutableOperator<FMetronomeOperatorBase>, public FMusicTransportControllable
+	{
+	public:
+
+		FMetronomeOperatorBase(const FBuildOperatorParams& InParams,
+			const FMusicTransportEventStreamReadRef& InTransport,
+			const bool  InLoop,
+			const int32 InLoopLengthBars,
+			const FInt32ReadRef& InTimSigNumerator,
+			const FInt32ReadRef& InTimeSigDenominator,
+			const FFloatReadRef& InTempo,
+			const FFloatReadRef& InSpeedMultiplier,
+			const int32 InSeekPrerollBars);
+
+		virtual void BindInputs(FInputVertexInterfaceData& InVertexData) override;
+		virtual void BindOutputs(FOutputVertexInterfaceData& InVertexData) override;
+
+		void Reset(const FResetParams& Params);
+		void Execute();
+
+	protected:
+		void Init();
+
+		//** INPUTS
+		FMusicTransportEventStreamReadRef TransportInPin;
+		const bool  LoopInPin;
+		const int32 LoopLengthBarsInPin;
+		FInt32ReadRef TimeSigNumInPin;
+		FInt32ReadRef TimeSigDenomInPin;
+		FFloatReadRef TempoInPin;
+		FFloatReadRef SpeedMultInPin;
+		const int32 SeekPreRollBarsInPin;
+
+		//** OUTPUTS
+		FMidiClockWriteRef MidiClockOutPin;
+
+		//** DATA
+		TSharedPtr<FMidiClock, ESPMode::NotThreadSafe> MonotonicallyIncreasingClock;
+		TSharedPtr<FSongMaps> SongMaps;
+		FSampleCount BlockSize;
+		float        SampleRate;
+		float        CurrentTempo;
+		int32        CurrentTimeSigNum;
+		int32        CurrentTimeSigDenom;
+		int32		 LastProcessedClockTick = -1;
+		int32		 NextClockTickToProcess = 0;
+		bool		 bClocksArePreparedForExecute = true;
+
+		void BuildSongMaps(bool ResetToStart = true);
+		void UpdateMidi();
+		void AddTempoChangeForMidi(float TempoBPM);
+		virtual void HandleTimeSigChangeForMidi(int32 TimeSigNum, int32 TimeSigDenom);
+		void HandleTransportChange(int32 StartFrameIndex, EMusicPlayerTransportState NewTransportState);
+		void PrepareClocksForExecute();
+		void MarkClocksAsExecuted();
+
+		FMidiClock& GetDrivingMidiClock() { return LoopInPin ? *MonotonicallyIncreasingClock : (*MidiClockOutPin); }
+	};
+}
\ No newline at end of file
