Fix infinite loop in FMidiClock::AddAdvanceToBlock caused by NextTimeSigChangeTick wrongly set (was using CurrentTempoInfoPointIndex instead of CurrentTimeSignaturePointIndex)

#rnx

diff --git a/Engine/Plugins/Runtime/Harmonix/Source/HarmonixMetasound/Private/HarmonixMetasound/DataTypes/MidiClock.cpp b/Engine/Plugins/Runtime/Harmonix/Source/HarmonixMetasound/Private/HarmonixMetasound/DataTypes/MidiClock.cpp
--- a/Engine/Plugins/Runtime/Harmonix/Source/HarmonixMetasound/Private/HarmonixMetasound/DataTypes/MidiClock.cpp
+++ b/Engine/Plugins/Runtime/Harmonix/Source/HarmonixMetasound/Private/HarmonixMetasound/DataTypes/MidiClock.cpp
@@ -766,7 +766,7 @@ namespace HarmonixMetasound
 			AddTimeSignatureChangeToBlock(CurrentBlockFrameIndex, NextMidiTickToProcess, SongMapEvaluator->GetTimeSignaturePoint(CurrentTimeSignaturePointIndex)->TimeSignature, NextTempoMapTickToProcess);
 			if ((CurrentTimeSignaturePointIndex + 1) < SongMapEvaluator->GetNumTimeSignatureChanges())
 			{
-				NextTimeSigChangeTick = SongMapEvaluator->GetTimeSignatureChangePointTick(CurrentTempoInfoPointIndex + 1);
+				NextTimeSigChangeTick = SongMapEvaluator->GetTimeSignatureChangePointTick(CurrentTimeSignaturePointIndex + 1);
 			}
 			else
 			{
