[Harmonix] Test voice id roundtrip and step sequencer stuck notes

#rnx

#virtualized

diff --git a/Engine/Plugins/Runtime/Harmonix/Source/HarmonixMetasound/Private/HarmonixMetasound/Nodes/StepSequencePlayerNode.cpp b/Engine/Plugins/Runtime/Harmonix/Source/HarmonixMetasound/Private/HarmonixMetasound/Nodes/StepSequencePlayerNode.cpp
--- a/Engine/Plugins/Runtime/Harmonix/Source/HarmonixMetasound/Private/HarmonixMetasound/Nodes/StepSequencePlayerNode.cpp
+++ b/Engine/Plugins/Runtime/Harmonix/Source/HarmonixMetasound/Private/HarmonixMetasound/Nodes/StepSequencePlayerNode.cpp
@@ -1,5 +1,7 @@
 // Copyright Epic Games, Inc. All Rights Reserved.
 
+#include "HarmonixMetasound/Nodes/StepSequencePlayerNode.h"
+
 #include "MetasoundExecutableOperator.h"
 #include "MetasoundFacade.h"
 #include "MetasoundNodeInterface.h"
@@ -337,6 +339,8 @@ namespace HarmonixMetasound::Nodes::StepSequencePlayer
 	{
 		MidiOutPin->PrepareBlock();
 
+		CheckForUpdatedSequenceTable();
+
 		// if we have no sequence table there is nothing to do. 
 		// Make sure the notes are all off and return.
 		if (!SequenceTable)
@@ -349,8 +353,6 @@ namespace HarmonixMetasound::Nodes::StepSequencePlayer
 			return;
 		}
 
-		CheckForUpdatedSequenceTable();
-
 		// We need to cache this to avoid avoid a crash if the value in the 
 		// sequence table asset changes while we are in the middle of rendering. 
 		CurrentStepSkipIndex = SequenceTable->StepSkipIndex;
diff --git a/Engine/Plugins/Runtime/Harmonix/Source/HarmonixMetasound/Public/HarmonixMetasound/Nodes/StepSequencePlayerNode.h b/Engine/Plugins/Runtime/Harmonix/Source/HarmonixMetasound/Public/HarmonixMetasound/Nodes/StepSequencePlayerNode.h
--- a/Engine/Plugins/Runtime/Harmonix/Source/HarmonixMetasound/Public/HarmonixMetasound/Nodes/StepSequencePlayerNode.h
+++ b/Engine/Plugins/Runtime/Harmonix/Source/HarmonixMetasound/Public/HarmonixMetasound/Nodes/StepSequencePlayerNode.h
@@ -8,19 +8,19 @@
 
 namespace HarmonixMetasound::Nodes::StepSequencePlayer
 {
-	HARMONIXMETASOUND_API Metasound::FNodeClassName GetClassName();
+	HARMONIXMETASOUND_API const Metasound::FNodeClassName& GetClassName();
 	HARMONIXMETASOUND_API int32 GetCurrentMajorVersion();
 
 	namespace Inputs
 	{
-		DECLARE_METASOUND_PARAM_EXTERN(SequenceAsset, "Step Sequence Asset", "Step sequence to play.");
-		DECLARE_METASOUND_PARAM_EXTERN(VelocityMultiplier, "Velocity Multiplier", "Multiplies the current note velocity by this number");
-		DECLARE_METASOUND_PARAM_EXTERN(MaxColumns, "Max Columns", "The Maximinum Number of cells to play per step sequence row.");
-		DECLARE_METASOUND_PARAM_EXTERN(AdditionalOctaves, "Additional Octaves", "The number of octaves to add to the authored step sequence note.");
-		DECLARE_METASOUND_PARAM_EXTERN(StepSizeQuarterNotes, "Step Size Quarter Notes", "The size, in quarter notes, of each step");
-		DECLARE_METASOUND_PARAM_EXTERN(ActivePage, "Active Page", "The page of the step sequence to play (1 indexed)");
-		DECLARE_METASOUND_PARAM_EXTERN(AutoPage, "Auto Page", "Whether to calculate the page of the step sequence based on current position");
-		DECLARE_METASOUND_PARAM_EXTERN(AutoPagePlaysBlankPages, "Auto Page Plays Blank Pages", "If autopaging, should blank pages be played?");
+		DECLARE_METASOUND_PARAM_EXTERN(SequenceAsset);
+		DECLARE_METASOUND_PARAM_EXTERN(VelocityMultiplier);
+		DECLARE_METASOUND_PARAM_EXTERN(MaxColumns);
+		DECLARE_METASOUND_PARAM_EXTERN(AdditionalOctaves);
+		DECLARE_METASOUND_PARAM_EXTERN(StepSizeQuarterNotes);
+		DECLARE_METASOUND_PARAM_EXTERN(ActivePage);
+		DECLARE_METASOUND_PARAM_EXTERN(AutoPage);
+		DECLARE_METASOUND_PARAM_EXTERN(AutoPagePlaysBlankPages);
 		DECLARE_METASOUND_PARAM_ALIAS(Transport);
 		DECLARE_METASOUND_PARAM_ALIAS(MidiClock);
 		DECLARE_METASOUND_PARAM_ALIAS(Speed);
diff --git a/Engine/Plugins/Runtime/Harmonix/Source/HarmonixMetasoundTests/Private/Nodes/StepSequencePlayerNodeTests.cpp b/Engine/Plugins/Runtime/Harmonix/Source/HarmonixMetasoundTests/Private/Nodes/StepSequencePlayerNodeTests.cpp
--- a/Engine/Plugins/Runtime/Harmonix/Source/HarmonixMetasoundTests/Private/Nodes/StepSequencePlayerNodeTests.cpp
+++ b/Engine/Plugins/Runtime/Harmonix/Source/HarmonixMetasoundTests/Private/Nodes/StepSequencePlayerNodeTests.cpp
@@ -1,8 +1,9 @@
 // Copyright Epic Games, Inc. All Rights Reserved.
 
 #include "NodeTestGraphBuilder.h"
-#include "HarmonixDsp/AudioBuffer.h"
-#include "HarmonixMetasound/Common.h"
+#include "HarmonixMetasound/DataTypes/MidiStepSequence.h"
+#include "HarmonixMetasound/DataTypes/MidiStream.h"
+#include "HarmonixMetasound/Nodes/StepSequencePlayerNode.h"
 #include "Misc/AutomationTest.h"
 
 #if WITH_DEV_AUTOMATION_TESTS
@@ -10,32 +11,103 @@
 namespace HarmonixMetasoundTests::StepSequencePlayerNode
 {
 	using GraphBuilder = Metasound::Test::FNodeTestGraphBuilder;
-	using namespace Metasound;
-	using namespace Metasound::Frontend;
-	using namespace HarmonixMetasound;
 
 	IMPLEMENT_SIMPLE_AUTOMATION_TEST(
-		FStepSequencePlayerCreateNodeTest,
-		"Harmonix.Metasound.Nodes.StepSequencePlayerNode.CreateNode",
+		FStepSequencePlayerNoStuckNotesOnTransposeTest,
+		"Harmonix.Metasound.Nodes.StepSequencePlayerNode.NoStuckNotesOnTranspose",
 		EAutomationTestFlags::EditorContext | EAutomationTestFlags::EngineFilter)
-		bool FStepSequencePlayerCreateNodeTest::RunTest(const FString&)
+		bool FStepSequencePlayerNoStuckNotesOnTransposeTest::RunTest(const FString&)
 	{
+		using namespace HarmonixMetasound::Nodes::StepSequencePlayer;
+		
 		// Build the graph.
 		constexpr int32 NumSamplesPerBlock = 256;
-		const TUniquePtr<FMetasoundGenerator> Generator = GraphBuilder::MakeSingleNodeGraph(
-			{ HarmonixMetasound::HarmonixNodeNamespace, "StepSequencePlayer", "" },
-			0,
+		const TUniquePtr<Metasound::FMetasoundGenerator> Generator = GraphBuilder::MakeSingleNodeGraph(
+			GetClassName(),
+			GetCurrentMajorVersion(),
 			48000,
 			NumSamplesPerBlock);
 		UTEST_TRUE("Graph successfully built", Generator.IsValid());
 
-		// execute a block
+		// Start the clock
+		auto ClockInput = Generator->GetInputWriteReference<HarmonixMetasound::FMidiClock>(Inputs::MidiClockName);
+		UTEST_TRUE("Got clock", ClockInput.IsSet());
+		(*ClockInput)->SetTransportState(0, HarmonixMetasound::EMusicPlayerTransportState::Playing);
+
+		// Start the transport
+		auto TransportInput = Generator->GetInputWriteReference<HarmonixMetasound::FMusicTransportEventStream>(Inputs::TransportName);
+		UTEST_TRUE("Got transport", TransportInput.IsSet());
+		(*TransportInput)->AddTransportRequest(HarmonixMetasound::EMusicPlayerTransportRequest::Play, 0);
+
+		// Create a sequence asset with a cell turned on
+		auto SequenceAssetInput = Generator->GetInputWriteReference<HarmonixMetasound::FMidiStepSequenceAsset>(Inputs::SequenceAssetName);
+		UTEST_TRUE("Got sequence asset", SequenceAssetInput.IsSet());
+
+		UMidiStepSequence* SequenceAsset = NewObject<UMidiStepSequence>();
+		SequenceAsset->SetCell(0, 0, true);
+		**SequenceAssetInput = SequenceAsset->CreateProxyData({});
+
+		// Render once, should get a note on in the output
+		(*ClockInput)->PrepareBlock();
+		(*ClockInput)->Advance(0, Generator->OperatorSettings.GetNumFramesPerBlock());
+		Metasound::FAudioBuffer Buffer{ NumSamplesPerBlock };
+		Generator->OnGenerateAudio(Buffer.GetData(), Buffer.Num());
+
+		auto MidiOutput = Generator->GetOutputReadReference<HarmonixMetasound::FMidiStream>(Outputs::MidiStreamName);
+		UTEST_TRUE("Got MIDI output", MidiOutput.IsSet());
+
+		uint8 Channel = Harmonix::Midi::Constants::GNumChannels;
+		uint8 NoteNumber = Harmonix::Midi::Constants::GMaxNote + 1;
+		FMidiVoiceId VoiceId = FMidiVoiceId::None();
+		
+		for (const auto& Event : (*MidiOutput)->GetEventsInBlock())
 		{
-			TAudioBuffer<float> Buffer{ Generator->GetNumChannels(), NumSamplesPerBlock, EAudioBufferCleanupMode::Delete};
-			Generator->OnGenerateAudio(Buffer.GetRawChannelData(0), Buffer.GetNumTotalValidSamples());
+			if (Event.MidiMessage.IsNoteOn())
+			{
+				Channel = Event.MidiMessage.GetStdChannel();
+				NoteNumber = Event.MidiMessage.GetStdData1();
+				VoiceId = Event.GetVoiceId();
+				break;
+			}
 		}
 
-		// No output to validate, finish test
+		UTEST_LESS("Note on: Channel was valid", Channel, Harmonix::Midi::Constants::GNumChannels);
+		UTEST_LESS("Note on: Note number was valid", NoteNumber, Harmonix::Midi::Constants::GMaxNote + 1);
+		UTEST_NOT_EQUAL("Note on: voice id was valid", VoiceId, FMidiVoiceId::None());
+
+		// Transpose and render until we get the note off, it should have the same voice id as the note on
+		auto AdditionalOctavesInput = Generator->GetInputWriteReference<float>(Inputs::AdditionalOctavesName);
+		UTEST_TRUE("Got additional octaves", AdditionalOctavesInput.IsSet());
+
+		constexpr float AdditionalOctaves = 2;
+		**AdditionalOctavesInput = AdditionalOctaves;
+
+		constexpr int32 MaxTries = 1000;
+		bool GotNoteOff = false;
+
+		for (int i = 0; i < MaxTries; ++i)
+		{
+			(*ClockInput)->PrepareBlock();
+			(*ClockInput)->Advance(0, Generator->OperatorSettings.GetNumFramesPerBlock());
+			Generator->OnGenerateAudio(Buffer.GetData(), Buffer.Num());
+
+			for (const auto& Event : (*MidiOutput)->GetEventsInBlock())
+			{
+				if (Event.MidiMessage.IsNoteOff())
+				{
+					UTEST_EQUAL("Correct channel", Event.MidiMessage.GetStdChannel(), Channel);
+					// NB: the actual message doesn't matter for note offs
+					UTEST_EQUAL("Correct voice id", Event.GetVoiceId(), VoiceId);
+					GotNoteOff = true;
+					break;
+				}
+			}
+
+			if (GotNoteOff)
+			{
+				break;
+			}
+		}
 
 		return true;
 	}
diff --git a/Engine/Plugins/Runtime/Harmonix/Source/HarmonixMidiTests/Private/MidiVoiceIdTests.cpp b/Engine/Plugins/Runtime/Harmonix/Source/HarmonixMidiTests/Private/MidiVoiceIdTests.cpp
new file mode 100644
--- /dev/null
+++ b/Engine/Plugins/Runtime/Harmonix/Source/HarmonixMidiTests/Private/MidiVoiceIdTests.cpp
@@ -0,0 +1,35 @@
+// Copyright Epic Games, Inc. All Rights Reserved.
+
+
+#include "HarmonixMidi/MidiVoiceId.h"
+#include "Misc/AutomationTest.h"
+
+#if WITH_DEV_AUTOMATION_TESTS
+
+namespace HarmonixMetasoundTests::StepSequencePlayerNode
+{
+	IMPLEMENT_SIMPLE_AUTOMATION_TEST(
+		FMidiVoiceIdRoundTripTest,
+		"Harmonix.Midi.VoiceId.RoundTrip",
+		EAutomationTestFlags::EditorContext | EAutomationTestFlags::EngineFilter)
+		bool FMidiVoiceIdRoundTripTest::RunTest(const FString&)
+	{
+		for (uint8 Channel = 0; Channel < Harmonix::Midi::Constants::GNumChannels; ++Channel)
+		{
+			for (uint8 NoteNumber = 0; NoteNumber < Harmonix::Midi::Constants::GMaxNumNotes; ++ NoteNumber)
+			{
+				FMidiVoiceGeneratorBase VoiceGenerator;
+				FMidiVoiceId VoiceId{ VoiceGenerator.GetIdBits(), FMidiMsg::CreateNoteOn(Channel, NoteNumber, 34) };
+				uint8 RetrievedChannel;
+				uint8 RetrievedNoteNumber;
+				VoiceId.GetChannelAndNote(RetrievedChannel, RetrievedNoteNumber);
+				UTEST_EQUAL("Channel is equal", RetrievedChannel, Channel);
+				UTEST_EQUAL("Note number is equal", RetrievedNoteNumber, NoteNumber);
+			}
+		}
+
+		return true;
+	}
+}
+
+#endif
