[Harmonix] Refactor to StepSequencePlayerNode to improve stability with non-looping StepSequences
#tests PIE, unit tests

diff --git a/Engine/Plugins/Runtime/Harmonix/Source/HarmonixMetasound/Private/HarmonixMetasound/Nodes/StepSequencePlayerNode.cpp b/Engine/Plugins/Runtime/Harmonix/Source/HarmonixMetasound/Private/HarmonixMetasound/Nodes/StepSequencePlayerNode.cpp
--- a/Engine/Plugins/Runtime/Harmonix/Source/HarmonixMetasound/Private/HarmonixMetasound/Nodes/StepSequencePlayerNode.cpp
+++ b/Engine/Plugins/Runtime/Harmonix/Source/HarmonixMetasound/Private/HarmonixMetasound/Nodes/StepSequencePlayerNode.cpp
@@ -46,6 +46,14 @@ namespace HarmonixMetasound::Nodes::StepSequencePlayer
 		return 0;
 	}
 
+	enum class EStepSequencePlayerState : uint8
+	{
+		NotPlaying,
+		PlayingLooping,
+		PlayingOneShot,
+		Finished
+	};
+
 	class FStepSequencePlayerOperator : public TExecutableOperator<FStepSequencePlayerOperator>, public FMusicTransportControllable, public FMidiVoiceGeneratorBase
 	{
 	public:
@@ -101,20 +109,19 @@ namespace HarmonixMetasound::Nodes::StepSequencePlayer
 
 		//** DATA
 		TSharedAudioRenderableDataPtr<FStepSequenceTable, TRefCountedAudioRenderableWithQueuedChanges<FStepSequenceTable>> SequenceTable;
-		FSampleCount BlockSize			= 0;
-		int32 CurrentBlockSpanStart		= 0;
-		int32 CurrentPageIndex			= -1;
-		int32 CurrentCellIndex			= -1;
-		int32 ProcessedThruTick			= -1;
-		int32 SequenceStartTick			= -1;
-		int32 CurrentStepSkipIndex		= 0;
-		bool  bAutoPage					= false;
-		bool  bPreviousAutoPage			= false;
-		bool  bAutoPagePlaysBlankPages	= false;
-		bool  bLoop						= true;
-		bool  bEnabled					= true;
-		bool  bPlaying					= false;
-		bool  bNeedsRebase				= false;
+		FSampleCount BlockSize			    = 0;
+		int32 CurrentBlockSpanStart		    = 0;
+		int32 CurrentPageIndex			    = -1;
+		int32 CurrentCellIndex			    = -1;
+		int32 ProcessedThruTick			    = -1;
+		int32 SequenceStartTick			    = -1;
+		int32 CurrentStepSkipIndex		    = 0;
+		bool  bAutoPage					    = false;
+		bool  bPreviousAutoPage			    = false;
+		bool  bAutoPagePlaysBlankPages	    = false;
+		bool  bLoop						    = true;
+		bool  bNeedsRebase				    = false;
+		EStepSequencePlayerState PlayState  = EStepSequencePlayerState::NotPlaying;
 		TArray<FMidiVoiceId> CurrentCellNotes;
 
 		void CheckForUpdatedSequenceTable();
@@ -125,6 +132,7 @@ namespace HarmonixMetasound::Nodes::StepSequencePlayer
 		void CalculatePageProperties(const FStepSequencePage& Page, int32 MaxColumns, float StepSizeQuarterNotes, int32& OutColumns, int32& OutTicksPerCell, int32& OutTableTickLength) const;
 		void RebaseSequenceStartTickForLoop(const int32 CurTick, const int32 TableTickLength);
 		void EnsureCurrentPageIndexIsValid();
+		bool IsPlaying() const;
 
 	protected:
 
@@ -329,9 +337,8 @@ namespace HarmonixMetasound::Nodes::StepSequencePlayer
 		bPreviousAutoPage = false;
 		bAutoPagePlaysBlankPages = false;
 		bLoop = true;
-		bEnabled = true;
-		bPlaying = false;
 		bNeedsRebase = false;
+		PlayState = EStepSequencePlayerState::NotPlaying;
 		CurrentCellNotes.Reset();
 	}
 
@@ -587,7 +594,7 @@ namespace HarmonixMetasound::Nodes::StepSequencePlayer
 			CurrentPageIndex = FMath::Abs(CurrentPageIndex);
 			if (CurrentPageIndex == 0)
 			{
-				// The loop ended in a previous tick, bPlaying should be false here.
+				// The loop ended in a previous tick, IsPlaying() should be false here.
 				// Reset the current index, and we'll wait to start playing again.
 				CurrentPageIndex = SequenceTable->GetFirstValidPage(bAutoPagePlaysBlankPages) + 1;
 			}
@@ -598,6 +605,11 @@ namespace HarmonixMetasound::Nodes::StepSequencePlayer
 		}
 	}
 
+	bool FStepSequencePlayerOperator::IsPlaying() const
+	{
+		return PlayState == EStepSequencePlayerState::PlayingLooping || PlayState == EStepSequencePlayerState::PlayingOneShot;
+	}
+
 	void FStepSequencePlayerOperator::AdvanceThruTick(int32 BlockFrameIndex, int32 Tick)
 	{
 		if (!SequenceTable || SequenceTable->Pages.IsEmpty())
@@ -639,66 +651,82 @@ namespace HarmonixMetasound::Nodes::StepSequencePlayer
 		int32 TableTickLength;
 		CalculatePageProperties(*CurrentPage, CurrentMaxColumns, CurrentStepSizeQuarterNotes, Columns, TicksPerCell, TableTickLength);
 		
-		bool bNewPlaying = bEnabled;
+		EStepSequencePlayerState LastPlayState = PlayState;
 
-		if (*EnabledInPin != bEnabled || (*LoopInPin && !bLoop))
+		// If this is enabled and not playing, transition it to playing
+		if (*EnabledInPin && PlayState == EStepSequencePlayerState::NotPlaying)
 		{
-			// If our enable state changes or we turn loop on, control playback
-			bEnabled = *EnabledInPin;
-			bNewPlaying = bEnabled;
+			PlayState = bLoop ? EStepSequencePlayerState::PlayingLooping : EStepSequencePlayerState::PlayingOneShot;
 		}
-		else if (*EnabledInPin && *LoopInPin)
+		// If this is not enabled and not in the "not playing" state, set that state
+		else if (!*EnabledInPin && PlayState != EStepSequencePlayerState::NotPlaying)
 		{
-			// Force playing to true if enabled and looping
-			bNewPlaying = true;
+			PlayState = EStepSequencePlayerState::NotPlaying;
 		}
 
-		if (bPlaying)
+		// Looping changed - adjust state if needed
+		if (*LoopInPin != bLoop)
 		{
-			if (!bNewPlaying)
+			bLoop = *LoopInPin;
+
+			if (!bLoop && PlayState == EStepSequencePlayerState::PlayingLooping)
 			{
-				// We're not playing anymore, make sure we're all off
-				AllNotesOff(BlockFrameIndex, Tick, true);
+				PlayState = EStepSequencePlayerState::PlayingOneShot;
 			}
-			else if (*LoopInPin != bLoop)
+			else if (bLoop && (PlayState == EStepSequencePlayerState::PlayingLooping || PlayState == EStepSequencePlayerState::Finished))
 			{
-				// If we're already playing and change the loop setting, need to re-sync to the global timeline
-				bNeedsRebase = true;
+				PlayState = EStepSequencePlayerState::PlayingLooping;
 			}
 		}
-		else if (bNewPlaying)
+		
+		if (LastPlayState != PlayState)
 		{
-			// We weren't playing before, but now we are
-			if (*LoopInPin)
+			switch (PlayState)
 			{
+			case EStepSequencePlayerState::PlayingLooping:
+			{
+				// Set that this needs a rebase - changing auto page may also need this same block to run
 				bNeedsRebase = true;
+				break;
 			}
-			else
+			case EStepSequencePlayerState::PlayingOneShot:
 			{
+				int32 TickInCell = (ProcessedThruTick % TicksPerCell);
+
+				// If the tick in the cell is close enough to a current division, 
+				// Move the ProcessedThruTick back to the previous division.
+				// That will become the SequenceStartTick, and that sequence will start immediately.
+				if (TickInCell < TicksPerCell / 16)
+				{
+					ProcessedThruTick -= TickInCell;
+					TickInCell = 0;
+				}
+
 				// Begin the oneshot on the next beat subdivision
-				SequenceStartTick = ProcessedThruTick + TicksPerCell - (ProcessedThruTick % TicksPerCell);
+				SequenceStartTick = ProcessedThruTick + (TickInCell > 0 ? TicksPerCell - TickInCell : 0);
 				if (bAutoPage)
 				{
 					// Nonlooping autopaging sequences start on the first valid page.
 					CurrentPageIndex = SequenceTable->GetFirstValidPage(bAutoPagePlaysBlankPages) + 1;
 				}
-				bNeedsRebase = false;
+				break;
+			}
+			case EStepSequencePlayerState::Finished:
+			case EStepSequencePlayerState::NotPlaying:
+			default:
+			{
+				// These states behave the same, they just have different transitions
+				// "Finished" cannot transition to "PlayingOneShot", and
+				// it won't become "NotPlaying" until the device is disabled
+				AllNotesOff(BlockFrameIndex, Tick, true);
+				break;
+			}
 			}
 		}
 
-		bPlaying = bNewPlaying;
-		bLoop = *LoopInPin;
-
-		// TODO When IsPreRoll we should probably be doing some of the handling of note ons/offs differently (more like FMidiPlayCursor::AdvanceThruTick), and therefore potentially our timeline sync as well
 		if (bNeedsRebase)
 		{
-			if (!bLoop)
-			{
-				// Temp: Don't rebase if loop is off, so we don't lose our start tick on graph rebuilds.
-				// NOTE: This introduces a bug where turning loop off in the middle of a sequence will cause a visual desync, but with loop in the Customize menu that's not a major concern.
-				bNeedsRebase = false;
-			}
-			else
+			if (PlayState == EStepSequencePlayerState::PlayingLooping)
 			{
 				RebaseSequenceStartTickForLoop(ProcessedThruTick, TableTickLength);
 				// Our page index may have changed, let's set it again
@@ -712,6 +740,10 @@ namespace HarmonixMetasound::Nodes::StepSequencePlayer
 				// Uncomment if there is ever a time where pages can have different tick lengths
 				// CalculatePageProperties(*CurrentPage, CurrentMaxColumns, CurrentStepSizeQuarterNotes, Columns, TicksPerCell, TableTickLength);
 			}
+			else
+			{
+				bNeedsRebase = false;
+			}
 		}
 
 		ProcessedThruTick = FMath::Max(-1, ProcessedThruTick);
@@ -720,7 +752,7 @@ namespace HarmonixMetasound::Nodes::StepSequencePlayer
 		{
 			ProcessedThruTick++;
 
-			if (!bPlaying || ProcessedThruTick < SequenceStartTick)
+			if (!IsPlaying() || ProcessedThruTick < SequenceStartTick)
 			{
 				// Loop has ended or hasn't started yet
 				continue;
@@ -735,7 +767,7 @@ namespace HarmonixMetasound::Nodes::StepSequencePlayer
 
 				if ((bAutoPage && CurrentPageIndex == 0) || (!bAutoPage && EffectiveTickForLoopPosition >= TableTickLength))
 				{
-					bPlaying = false;
+					PlayState = EStepSequencePlayerState::Finished;
 					continue;
 				}
 			}
@@ -768,7 +800,7 @@ namespace HarmonixMetasound::Nodes::StepSequencePlayer
 					if (CurrentPageIndex == 0)
 					{
 						// Loop has ended
-						bPlaying = false;
+						PlayState = EStepSequencePlayerState::Finished;
 						continue;
 					}
 					else if (!SequenceTable->Pages.IsValidIndex(CurrentPageIndex - 1))
@@ -853,7 +885,7 @@ namespace HarmonixMetasound::Nodes::StepSequencePlayer
 			}
 		}
 
-		if (bNewPlaying && !bPlaying)
+		if (PlayState == EStepSequencePlayerState::Finished && LastPlayState != PlayState)
 		{
 			// Loop ended above, make sure we turn everything off
 			AllNotesOff(BlockFrameIndex, Tick, true);
diff --git a/Engine/Plugins/Runtime/Harmonix/Source/HarmonixMetasound/Public/HarmonixMetasound/Nodes/StepSequencePlayerNode.h b/Engine/Plugins/Runtime/Harmonix/Source/HarmonixMetasound/Public/HarmonixMetasound/Nodes/StepSequencePlayerNode.h
--- a/Engine/Plugins/Runtime/Harmonix/Source/HarmonixMetasound/Public/HarmonixMetasound/Nodes/StepSequencePlayerNode.h
+++ b/Engine/Plugins/Runtime/Harmonix/Source/HarmonixMetasound/Public/HarmonixMetasound/Nodes/StepSequencePlayerNode.h
@@ -1,4 +1,4 @@
-﻿// Copyright Epic Games, Inc. All Rights Reserved.
+// Copyright Epic Games, Inc. All Rights Reserved.
 
 #pragma once
 
diff --git a/Engine/Plugins/Runtime/Harmonix/Source/HarmonixMetasoundTests/Private/Nodes/StepSequencePlayerNodeTests.cpp b/Engine/Plugins/Runtime/Harmonix/Source/HarmonixMetasoundTests/Private/Nodes/StepSequencePlayerNodeTests.cpp
--- a/Engine/Plugins/Runtime/Harmonix/Source/HarmonixMetasoundTests/Private/Nodes/StepSequencePlayerNodeTests.cpp
+++ b/Engine/Plugins/Runtime/Harmonix/Source/HarmonixMetasoundTests/Private/Nodes/StepSequencePlayerNodeTests.cpp
@@ -1,17 +1,16 @@
 // Copyright Epic Games, Inc. All Rights Reserved.
-
 #include "NodeTestGraphBuilder.h"
 #include "HarmonixMetasound/DataTypes/MidiStepSequence.h"
 #include "HarmonixMetasound/DataTypes/MidiStream.h"
 #include "HarmonixMetasound/Nodes/StepSequencePlayerNode.h"
 #include "Misc/AutomationTest.h"
-
 #if WITH_DEV_AUTOMATION_TESTS
-
 namespace HarmonixMetasoundTests::StepSequencePlayerNode
 {
 	using GraphBuilder = Metasound::Test::FNodeTestGraphBuilder;
 
+
+
 	IMPLEMENT_SIMPLE_AUTOMATION_TEST(
 		FStepSequencePlayerNoStuckNotesOnTransposeTest,
 		"Harmonix.Metasound.Nodes.StepSequencePlayerNode.NoStuckNotesOnTranspose",
@@ -111,6 +110,129 @@ namespace HarmonixMetasoundTests::StepSequencePlayerNode
 
 		return true;
 	}
-}
 
-#endif
+	IMPLEMENT_SIMPLE_AUTOMATION_TEST(
+		FStepSequencePlayerEnabledAndLooping,
+		"Harmonix.Metasound.Nodes.StepSequencePlayerNode.StepSequencePlayerEnabledAndLooping",
+		EAutomationTestFlags::EditorContext | EAutomationTestFlags::EngineFilter)
+		bool FStepSequencePlayerEnabledAndLooping::RunTest(const FString&)
+	{
+		using namespace HarmonixMetasound::Nodes::StepSequencePlayer;
+
+		// Build the graph.
+		constexpr int32 NumSamplesPerBlock = 256;
+		const TUniquePtr<Metasound::FMetasoundGenerator> Generator = GraphBuilder::MakeSingleNodeGraph(
+			GetClassName(),
+			GetCurrentMajorVersion(),
+			48000,
+			NumSamplesPerBlock);
+		UTEST_TRUE("Graph successfully built", Generator.IsValid());
+
+		// Start the clock
+		auto ClockInput = Generator->GetInputWriteReference<HarmonixMetasound::FMidiClock>(Inputs::MidiClockName);
+		UTEST_TRUE("Got clock", ClockInput.IsSet());
+		(*ClockInput)->SetTransportState(0, HarmonixMetasound::EMusicPlayerTransportState::Playing);
+		// Start the transport
+		auto TransportInput = Generator->GetInputWriteReference<HarmonixMetasound::FMusicTransportEventStream>(Inputs::TransportName);
+		UTEST_TRUE("Got transport", TransportInput.IsSet());
+		(*TransportInput)->AddTransportRequest(HarmonixMetasound::EMusicPlayerTransportRequest::Play, 0);
+		// Create a sequence asset with a cell turned on
+		auto SequenceAssetInput = Generator->GetInputWriteReference<HarmonixMetasound::FMidiStepSequenceAsset>(Inputs::SequenceAssetName);
+		UTEST_TRUE("Got sequence asset", SequenceAssetInput.IsSet());
+		UMidiStepSequence* SequenceAsset = NewObject<UMidiStepSequence>();
+		SequenceAsset->SetCell(0, 0, true);
+		**SequenceAssetInput = SequenceAsset->CreateProxyData({});
+
+		Metasound::FAudioBuffer Buffer{ NumSamplesPerBlock };
+		auto MidiOutput = Generator->GetOutputReadReference<HarmonixMetasound::FMidiStream>(Outputs::MidiStreamName);
+		UTEST_TRUE("Got MIDI output", MidiOutput.IsSet());
+
+		// Turn looping on
+		auto LoopInput = Generator->GetInputWriteReference<bool>(Inputs::LoopName);
+		UTEST_TRUE("Got loop input", LoopInput.IsSet());
+		**LoopInput = true;
+
+		constexpr int32 MaxTries = 10000;
+		int32 TotalNotesOn = 0;
+		for (int i = 0; i < MaxTries; ++i)
+		{
+			(*ClockInput)->PrepareBlock();
+			(*ClockInput)->Advance(0, Generator->OperatorSettings.GetNumFramesPerBlock());
+			Generator->OnGenerateAudio(Buffer.GetData(), Buffer.Num());
+			for (const auto& Event : (*MidiOutput)->GetEventsInBlock())
+			{
+				if (Event.MidiMessage.IsNoteOn())
+				{
+					TotalNotesOn++;
+					break;
+				}
+			}
+			if (TotalNotesOn > 1)
+			{
+				break;
+			}
+		}
+
+		UTEST_GREATER_EQUAL("Looping on: Sequencer Played At Least Two Notes", TotalNotesOn, 2);
+
+		// Disable and turn off looping
+		auto EnabledInput = Generator->GetInputWriteReference<bool>(Inputs::EnabledName);
+		UTEST_TRUE("Got loop input", EnabledInput.IsSet());
+		**EnabledInput = false;
+		**LoopInput = false;
+
+		(*ClockInput)->PrepareBlock();
+		(*ClockInput)->Advance(0, Generator->OperatorSettings.GetNumFramesPerBlock());
+
+		bool GotAnyNotes = false;
+		for (int i = 0; i < MaxTries; ++i)
+		{
+			(*ClockInput)->PrepareBlock();
+			(*ClockInput)->Advance(0, Generator->OperatorSettings.GetNumFramesPerBlock());
+			Generator->OnGenerateAudio(Buffer.GetData(), Buffer.Num());
+			for (const auto& Event : (*MidiOutput)->GetEventsInBlock())
+			{
+				if (Event.MidiMessage.IsNoteOn())
+				{
+					GotAnyNotes = true;
+					break;
+				}
+			}
+			if (GotAnyNotes)
+			{
+				break;
+			}
+		}
+
+		UTEST_FALSE("Enabled off: No notes received", GotAnyNotes);
+
+		// Re-enable
+		**EnabledInput = true;
+
+		TotalNotesOn = 0;
+		for (int i = 0; i < MaxTries; ++i)
+		{
+			(*ClockInput)->PrepareBlock();
+			(*ClockInput)->Advance(0, Generator->OperatorSettings.GetNumFramesPerBlock());
+			Generator->OnGenerateAudio(Buffer.GetData(), Buffer.Num());
+			for (const auto& Event : (*MidiOutput)->GetEventsInBlock())
+			{
+				if (Event.MidiMessage.IsNoteOn())
+				{
+					TotalNotesOn++;
+					break;
+				}
+			}
+			if (TotalNotesOn > 1)
+			{
+				break;
+			}
+		}
+
+		UTEST_EQUAL("Looping Off: One note received", TotalNotesOn, 1);
+
+		return true;
+	}
+
+}
+#endif
\ No newline at end of file
