[Harmonix] step sequence node test fix

#rnx

diff --git a/Engine/Plugins/Runtime/Harmonix/Source/HarmonixMetasoundTests/Private/Nodes/StepSequencePlayerNodeTests.cpp b/Engine/Plugins/Runtime/Harmonix/Source/HarmonixMetasoundTests/Private/Nodes/StepSequencePlayerNodeTests.cpp
--- a/Engine/Plugins/Runtime/Harmonix/Source/HarmonixMetasoundTests/Private/Nodes/StepSequencePlayerNodeTests.cpp
+++ b/Engine/Plugins/Runtime/Harmonix/Source/HarmonixMetasoundTests/Private/Nodes/StepSequencePlayerNodeTests.cpp
@@ -108,6 +108,8 @@ namespace HarmonixMetasoundTests::StepSequencePlayerNode
 			}
 		}
 
+		UTEST_TRUE("No stuck notes", GotNoteOff);
+
 		return true;
 	}
 
