Fix flawed logic in Step Sequencer node that very rarely would cause an assert.

diff --git a/Engine/Plugins/Runtime/Harmonix/Source/HarmonixMetasound/Private/HarmonixMetasound/Nodes/StepSequencePlayerNode.cpp b/Engine/Plugins/Runtime/Harmonix/Source/HarmonixMetasound/Private/HarmonixMetasound/Nodes/StepSequencePlayerNode.cpp
--- a/Engine/Plugins/Runtime/Harmonix/Source/HarmonixMetasound/Private/HarmonixMetasound/Nodes/StepSequencePlayerNode.cpp
+++ b/Engine/Plugins/Runtime/Harmonix/Source/HarmonixMetasound/Private/HarmonixMetasound/Nodes/StepSequencePlayerNode.cpp
@@ -932,13 +932,12 @@ namespace HarmonixMetasound::Nodes::StepSequencePlayer
 			{
 				if (CurrentCellNotes[i])
 				{
-					const int32 OriginalNote = SequenceTable->Notes[i].NoteNumber;
-					const int32 TransposedNote = FMath::Clamp(OriginalNote + AdditionalOctaveNotes, 0, 127);
-
 					uint8 MidiCh;
 					uint8 MidiNote;
 					CurrentCellNotes[i].GetChannelAndNote(MidiCh, MidiNote);
 					FMidiStreamEvent MidiEvent(CurrentCellNotes[i].GetGeneratorId(), FMidiMsg::CreateNoteOff(MidiCh, MidiNote));
+					// Since we can't get the note from the SequenceTable->Notes, let's just guess here and fix when we get an opportunity to remove the transposition
+					const int32 TransposedNote = FMath::Clamp(MidiNote + AdditionalOctaveNotes, 0, 127);
 					MidiEvent.MidiMessage.Data1 = TransposedNote;
 					MidiEvent.BlockSampleFrameIndex  = CurrentBlockSpanStart;
 					MidiEvent.AuthoredMidiTick       = 0;
