[Harmonix] MIDI Note Trigger node: set note number only on note on.

This will avoid issues with downstream patching that might behave badly when a note is transposed on the way in, but not on the way out.

Step Sequence Player node: don't worry about setting a transposed note number on note off messages. The voice id of the message will be enough for downstream code to respond appropriately.

#tests automated tests, test asset in CL

#virtualized

diff --git a/Engine/Plugins/Runtime/Harmonix/Source/HarmonixMetasound/Private/HarmonixMetasound/Nodes/MidiNoteTriggerNode.cpp b/Engine/Plugins/Runtime/Harmonix/Source/HarmonixMetasound/Private/HarmonixMetasound/Nodes/MidiNoteTriggerNode.cpp
--- a/Engine/Plugins/Runtime/Harmonix/Source/HarmonixMetasound/Private/HarmonixMetasound/Nodes/MidiNoteTriggerNode.cpp
+++ b/Engine/Plugins/Runtime/Harmonix/Source/HarmonixMetasound/Private/HarmonixMetasound/Nodes/MidiNoteTriggerNode.cpp
@@ -149,7 +149,7 @@ namespace HarmonixMetasound::Nodes::MidiNoteTriggerNode
 
 			StuckNoteGuard.UnstickNotes(*MidiStreamInPin, [this, &NoteOffTriggerFrame](const FMidiStreamEvent& Event)
 			{
-				TriggerNoteOff(0, Event.MidiMessage.GetStdData1());
+				TriggerNoteOff(0);
 				NoteOffTriggerFrame = 0;
 			});
 
@@ -157,7 +157,7 @@ namespace HarmonixMetasound::Nodes::MidiNoteTriggerNode
 			{
 				if (SoundingNote >= 0)
 				{
-					TriggerNoteOff(0, SoundingNote);
+					TriggerNoteOff(0);
 				}
 				return;
 			}
@@ -168,7 +168,7 @@ namespace HarmonixMetasound::Nodes::MidiNoteTriggerNode
 				const TSharedPtr<const FMidiClock, ESPMode::NotThreadSafe> Clock = MidiStreamInPin->GetClock();
 				if (Clock.IsValid() && Clock->GetTransportStateAtEndOfBlock() != EMusicPlayerTransportState::Playing)
 				{
-					TriggerNoteOff(0, SoundingNote);
+					TriggerNoteOff(0);
 					return;
 				}
 			}
@@ -180,7 +180,7 @@ namespace HarmonixMetasound::Nodes::MidiNoteTriggerNode
 					if (SoundingNote > -1)
 					{
 						// Stop sounding note
-						TriggerNoteOff(Event.BlockSampleFrameIndex, SoundingNote);
+						TriggerNoteOff(Event.BlockSampleFrameIndex);
 						NoteOffTriggerFrame = Event.BlockSampleFrameIndex;
 					}
 
@@ -195,7 +195,7 @@ namespace HarmonixMetasound::Nodes::MidiNoteTriggerNode
 				{
 					if (Event.GetVoiceId() == PlayingId)
 					{
-						TriggerNoteOff(Event.BlockSampleFrameIndex, Event.MidiMessage.GetStdData1());
+						TriggerNoteOff(Event.BlockSampleFrameIndex);
 						NoteOffTriggerFrame = Event.BlockSampleFrameIndex;
 						PlayingId = FMidiVoiceId::None();
 						SoundingNote = -1;
@@ -206,7 +206,7 @@ namespace HarmonixMetasound::Nodes::MidiNoteTriggerNode
 					if (SoundingNote > -1)
 					{
 						// Stop sounding note
-						TriggerNoteOff(Event.BlockSampleFrameIndex, SoundingNote);
+						TriggerNoteOff(Event.BlockSampleFrameIndex);
 						NoteOffTriggerFrame = Event.BlockSampleFrameIndex;
 						PlayingId = FMidiVoiceId::None();
 						SoundingNote = -1;
@@ -216,10 +216,9 @@ namespace HarmonixMetasound::Nodes::MidiNoteTriggerNode
 		}
 		
 	private:
-		void TriggerNoteOff(int32 BlockSampleFrameIndex, int32 NoteNumber)
+		void TriggerNoteOff(int32 BlockSampleFrameIndex)
 		{
 			*VelOutPin = 0;
-			*NoteNumOutPin = NoteNumber;
 			NoteOffOutPin->TriggerFrame(BlockSampleFrameIndex);
 			SoundingNote = -1;
 			PlayingId = FMidiVoiceId::None();
diff --git a/Engine/Plugins/Runtime/Harmonix/Source/HarmonixMetasound/Private/HarmonixMetasound/Nodes/StepSequencePlayerNode.cpp b/Engine/Plugins/Runtime/Harmonix/Source/HarmonixMetasound/Private/HarmonixMetasound/Nodes/StepSequencePlayerNode.cpp
--- a/Engine/Plugins/Runtime/Harmonix/Source/HarmonixMetasound/Private/HarmonixMetasound/Nodes/StepSequencePlayerNode.cpp
+++ b/Engine/Plugins/Runtime/Harmonix/Source/HarmonixMetasound/Private/HarmonixMetasound/Nodes/StepSequencePlayerNode.cpp
@@ -29,6 +29,8 @@ namespace HarmonixMetasound::Nodes::StepSequencePlayer
 {
 	using namespace Metasound;
 	using namespace Harmonix;
+
+	constexpr uint8 MidiChannel = 0;
 	
 	const Metasound::FNodeClassName& GetClassName()
 	{
@@ -814,70 +816,57 @@ namespace HarmonixMetasound::Nodes::StepSequencePlayer
 				}	
 				bPreviousAutoPage = bAutoPage;
 
-				for (int32 i = 0; i < CurrentPage->Rows.Num(); ++i)
+				for (int32 NoteIdx = 0; NoteIdx < CurrentPage->Rows.Num(); ++NoteIdx)
 				{
-					if (CurrentCellNotes[i])
+					if (CurrentCellNotes[NoteIdx])
 					{
-						int32 NoteIndex = i;
-						uint8 MidiCh;
-						uint8 MidiNote;
-						CurrentCellNotes[i].GetChannelAndNote(MidiCh, MidiNote);
-						
 						// If: 
 						// 1. The note that would play is not enabled (it is not a new note);
 						// 2. A note with this pitch would start up this tick;
 						// 3. The note is marked as a continuation note
 						// Then keep this note playing through the next cell
-						if (MidiNote == NoteIndex && !CurrentPage->Rows[i].Cells[NextCellInRow].bEnabled && CurrentPage->Rows[i].Cells[NextCellInRow].bContinuation)
+						if (!CurrentPage->Rows[NoteIdx].Cells[NextCellInRow].bEnabled && CurrentPage->Rows[NoteIdx].Cells[NextCellInRow].bContinuation)
 						{
 							continue;
 						}
 
 						// note off!
-						// use the "NoteIndex" as the note number for the midi event to track the midi event for our note off message
-						// the note number and transposition of the step sequencer can change out from under us, so we can't use those to id events
-						// we transpose the note later
-						int32 OriginalNote = SequenceTable->Notes[i].NoteNumber;
-						int32 TransposedNote = FMath::Clamp(OriginalNote + AdditionalOctaveNotes, 0, 127);
-						FMidiStreamEvent MidiEvent(CurrentCellNotes[i].GetGeneratorId(), FMidiMsg::CreateNoteOff(MidiCh, NoteIndex));
-						// and then assign the note directly to the midi message after
-						MidiEvent.MidiMessage.Data1 = TransposedNote;
+						FMidiStreamEvent MidiEvent(CurrentCellNotes[NoteIdx].GetGeneratorId(), FMidiMsg::CreateNoteOff(MidiChannel, NoteIdx));
 						MidiEvent.BlockSampleFrameIndex = BlockFrameIndex;
 						MidiEvent.AuthoredMidiTick = ProcessedThruTick;
 						MidiEvent.CurrentMidiTick = ProcessedThruTick;
 						MidiEvent.TrackIndex = 1;
 						MidiOutPin->AddNoteOffEventOrCancelPendingNoteOn(MidiEvent);
-						UE_LOG(LogStepSequencePlayer, Verbose, TEXT("0x%x Note-Off %d at %d"), (uint32)(size_t)this, SequenceTable->Notes[i].NoteNumber + AdditionalOctaveNotes, BlockFrameIndex);
-						CurrentCellNotes[i] = FMidiVoiceId::None();
+						UE_LOG(LogStepSequencePlayer, Verbose, TEXT("0x%x Note-Off %d at %d"), (uint32)(size_t)this, SequenceTable->Notes[NoteIdx].NoteNumber + AdditionalOctaveNotes, BlockFrameIndex);
+						CurrentCellNotes[NoteIdx] = FMidiVoiceId::None();
 					}
 				}
 			}
 
 			if (CellInRow != CurrentCellIndex || CurrentMaxColumns == 1)
 			{	
-				for (int32 i = 0; i < CurrentPage->Rows.Num(); ++i)
+				for (int32 NoteIdx = 0; NoteIdx < CurrentPage->Rows.Num(); ++NoteIdx)
 				{
-					if (CurrentPage->Rows[i].bRowEnabled && CurrentPage->Rows[i].Cells[CellInRow].bEnabled)
+					if (CurrentPage->Rows[NoteIdx].bRowEnabled && CurrentPage->Rows[NoteIdx].Cells[CellInRow].bEnabled)
 					{
-						if (!CurrentCellNotes[i])
+						if (!CurrentCellNotes[NoteIdx])
 						{
 							// note on!
-							int32 NoteIndex = i;
-							int32 OriginalNote = SequenceTable->Notes[i].NoteNumber;
-							int32 TransposedNote = FMath::Clamp(OriginalNote + AdditionalOctaveNotes, 0, 127);
 							// create the midi event with the original note to maintain voice ids.
-							FMidiStreamEvent MidiEvent(this, FMidiMsg::CreateNoteOn(0, NoteIndex, SequenceTable->Notes[i].Velocity));
+							FMidiStreamEvent MidiEvent(this, FMidiMsg::CreateNoteOn(MidiChannel, NoteIdx, SequenceTable->Notes[NoteIdx].Velocity));
 							// and then assign the note directly to the midi message after
+							int32 OriginalNote = SequenceTable->Notes[NoteIdx].NoteNumber;
+							int32 TransposedNote = FMath::Clamp(OriginalNote + AdditionalOctaveNotes, 0, 127);
 							MidiEvent.MidiMessage.Data1 = TransposedNote;
-							float NoteOnVelocity = FMath::Clamp(static_cast<float>(SequenceTable->Notes[i].Velocity) * CurrentVelocityMultiplierValue, 0.0f, 127.0f);
+							float NoteOnVelocity = FMath::Clamp(static_cast<float>(SequenceTable->Notes[NoteIdx].Velocity) * CurrentVelocityMultiplierValue, 0.0f, 127.0f);
 							MidiEvent.MidiMessage.SetNoteOnVelocity(static_cast<uint8>(NoteOnVelocity));
 							MidiEvent.BlockSampleFrameIndex = BlockFrameIndex;
 							MidiEvent.AuthoredMidiTick = ProcessedThruTick;
 							MidiEvent.CurrentMidiTick = ProcessedThruTick;
 							MidiEvent.TrackIndex = 1;
 							MidiOutPin->AddMidiEvent(MidiEvent);
-							UE_LOG(LogStepSequencePlayer, Verbose, TEXT("0x%x Note-On %d at %d"), (uint32)(size_t)this, SequenceTable->Notes[i].NoteNumber + AdditionalOctaveNotes, BlockFrameIndex);
-							CurrentCellNotes[i] = MidiEvent.GetVoiceId();
+							UE_LOG(LogStepSequencePlayer, Verbose, TEXT("0x%x Note-On %d at %d"), (uint32)(size_t)this, SequenceTable->Notes[NoteIdx].NoteNumber + AdditionalOctaveNotes, BlockFrameIndex);
+							CurrentCellNotes[NoteIdx] = MidiEvent.GetVoiceId();
 						}
 					}
 				}
@@ -925,25 +914,17 @@ namespace HarmonixMetasound::Nodes::StepSequencePlayer
 
 		if (SequenceTable->Notes.Num() < CurrentCellNotes.Num())
 		{
-			const int32 AdditionalOctaveNotes = (int32)*AdditionalOctavesInPin * 12;
-
 			// We may have existing notes that need to be stopped.
-			for (int32 i = SequenceTable->Notes.Num(); i < CurrentCellNotes.Num(); ++i)
+			for (int32 NoteIdx = SequenceTable->Notes.Num(); NoteIdx < CurrentCellNotes.Num(); ++NoteIdx)
 			{
-				if (CurrentCellNotes[i])
+				if (CurrentCellNotes[NoteIdx])
 				{
-					uint8 MidiCh;
-					uint8 MidiNote;
-					CurrentCellNotes[i].GetChannelAndNote(MidiCh, MidiNote);
-					FMidiStreamEvent MidiEvent(CurrentCellNotes[i].GetGeneratorId(), FMidiMsg::CreateNoteOff(MidiCh, MidiNote));
-					// Since we can't get the note from the SequenceTable->Notes, let's just guess here and fix when we get an opportunity to remove the transposition
-					const int32 TransposedNote = FMath::Clamp(MidiNote + AdditionalOctaveNotes, 0, 127);
-					MidiEvent.MidiMessage.Data1 = TransposedNote;
+					FMidiStreamEvent MidiEvent(CurrentCellNotes[NoteIdx].GetGeneratorId(), FMidiMsg::CreateNoteOff(MidiChannel, NoteIdx));
 					MidiEvent.BlockSampleFrameIndex  = CurrentBlockSpanStart;
 					MidiEvent.AuthoredMidiTick       = 0;
 					MidiEvent.CurrentMidiTick        = 0;
 					MidiEvent.TrackIndex             = 1;
-					UE_LOG(LogStepSequencePlayer, Verbose, TEXT("0x%x Note-Off %d (during resize)"), (uint32)(size_t)this, MidiNote);
+					UE_LOG(LogStepSequencePlayer, Verbose, TEXT("0x%x Note-Off %d (during resize)"), (uint32)(size_t)this, NoteIdx);
 					MidiOutPin->AddMidiEvent(MidiEvent);
 				}
 			}
@@ -960,26 +941,18 @@ namespace HarmonixMetasound::Nodes::StepSequencePlayer
 
 	void FStepSequencePlayerOperator::AllNotesOff(int32 AtFrameIndex, int32 AbsMidiTick, bool ResetCellIndex)
 	{
-		const int32 AdditionalOctaveNotes = (int32)*AdditionalOctavesInPin * 12;
-		for (int i = 0; i < CurrentCellNotes.Num(); ++i)
+		for (int NoteIdx = 0; NoteIdx < CurrentCellNotes.Num(); ++NoteIdx)
 		{
-			if (CurrentCellNotes[i])
+			if (CurrentCellNotes[NoteIdx])
 			{
-				const int32 OriginalNote = SequenceTable->Notes[i].NoteNumber;
-				const int32 TransposedNote = FMath::Clamp(OriginalNote + AdditionalOctaveNotes, 0, 127);
-
-				uint8 MidiCh;
-				uint8 MidiNote;
-				CurrentCellNotes[i].GetChannelAndNote(MidiCh, MidiNote);
-				FMidiStreamEvent MidiEvent(CurrentCellNotes[i].GetGeneratorId(), FMidiMsg::CreateNoteOff(MidiCh, MidiNote));
-				MidiEvent.MidiMessage.Data1 = TransposedNote;
+				FMidiStreamEvent MidiEvent(CurrentCellNotes[NoteIdx].GetGeneratorId(), FMidiMsg::CreateNoteOff(MidiChannel, NoteIdx));
 				MidiEvent.BlockSampleFrameIndex  = AtFrameIndex;
 				MidiEvent.AuthoredMidiTick       = AbsMidiTick;
 				MidiEvent.CurrentMidiTick        = AbsMidiTick;
 				MidiEvent.TrackIndex             = 1;
 				MidiOutPin->AddMidiEvent(MidiEvent);
-				UE_LOG(LogStepSequencePlayer, Verbose, TEXT("0x%x Note-Off %d (during all notes off)"), (uint32)(size_t)this, MidiNote);
-				CurrentCellNotes[i] = FMidiVoiceId::None();
+				UE_LOG(LogStepSequencePlayer, Verbose, TEXT("0x%x Note-Off %d (during all notes off)"), (uint32)(size_t)this, NoteIdx);
+				CurrentCellNotes[NoteIdx] = FMidiVoiceId::None();
 			}
 		}
 		if (ResetCellIndex)
