Fixed linker errors on some console platform when building modular builds. Fixes involves moving ctor/dtor and friends from header to cpp and dll import/export

diff --git a/Engine/Plugins/Experimental/MeshModelingToolsetExp/Source/ModelingUI/Public/ModelingWidgets/SDraggableBox.h b/Engine/Plugins/Experimental/MeshModelingToolsetExp/Source/ModelingUI/Public/ModelingWidgets/SDraggableBox.h
--- a/Engine/Plugins/Experimental/MeshModelingToolsetExp/Source/ModelingUI/Public/ModelingWidgets/SDraggableBox.h
+++ b/Engine/Plugins/Experimental/MeshModelingToolsetExp/Source/ModelingUI/Public/ModelingWidgets/SDraggableBox.h
@@ -11,7 +11,7 @@
  * Class that can be used to place a draggable box into a viewport or some other large widget as an
  * overlay. Just place the widget that you want to be draggable as the contents of SDraggableBoxOverlay.
  */
-class MODELINGUI_API SDraggableBoxOverlay : public SCompoundWidget
+class SDraggableBoxOverlay : public SCompoundWidget
 {
 public:
 
@@ -27,13 +27,13 @@ public:
 	SLATE_DEFAULT_SLOT(FArguments, Content)
 	SLATE_END_ARGS()
 
-	void Construct(const FArguments& InArgs);
+	MODELINGUI_API void Construct(const FArguments& InArgs);
 
 	/**
 	 * Sets the box position in the overlay. Horizontal position is distance from left, and vertical position
 	 * is distance from bottom or from top depending on whether bPositionRelativeToBottom == true.
 	 */
-	void SetBoxPosition(float HorizontalPosition, float VerticalPosition);
+	MODELINGUI_API void SetBoxPosition(float HorizontalPosition, float VerticalPosition);
 
 protected:
 
@@ -52,7 +52,7 @@ protected:
  * which will put its contents into a draggable box and properly handle the dragging without the
  * user having to set it up.
  */
-class MODELINGUI_API SDraggableBox : public SCompoundWidget
+class SDraggableBox : public SCompoundWidget
 {
 public:
 
@@ -63,11 +63,11 @@ public:
 	SLATE_DEFAULT_SLOT(FArguments, Content)
 	SLATE_END_ARGS()
 
-	void Construct(const FArguments& InArgs);
+	MODELINGUI_API void Construct(const FArguments& InArgs);
 
 	// SWidget overrides necessary for box to be draggable
-	FReply OnMouseButtonDown(const FGeometry& MyGeometry, const FPointerEvent& MouseEvent) override;
-	FReply OnDragDetected(const FGeometry& MyGeometry, const FPointerEvent& MouseEvent) override;
+	MODELINGUI_API FReply OnMouseButtonDown(const FGeometry& MyGeometry, const FPointerEvent& MouseEvent) override;
+	MODELINGUI_API FReply OnDragDetected(const FGeometry& MyGeometry, const FPointerEvent& MouseEvent) override;
 
 protected:
 	TSharedPtr<SWidget> InnerWidget;
@@ -85,15 +85,15 @@ protected:
  * OnDragComplete delegate, it just triggers on any drop. This makes it possible to use it in any
  * viewports, not just the level editor.
  */
-class MODELINGUI_API FDraggableBoxUIDragOperation : public FDragDropOperation
+class FDraggableBoxUIDragOperation : public FDragDropOperation
 {
 public:
 
 	DRAG_DROP_OPERATOR_TYPE(FDraggableBoxUIDragOperation, FDragDropOperation)
 
 	// FDragDropOperation
-	virtual void OnDrop(bool bDropWasHandled, const FPointerEvent& MouseEvent) override;
-	virtual void OnDragged(const FDragDropEvent& DragDropEvent) override;
+	MODELINGUI_API virtual void OnDrop(bool bDropWasHandled, const FPointerEvent& MouseEvent) override;
+	MODELINGUI_API virtual void OnDragged(const FDragDropEvent& DragDropEvent) override;
 
 	/**
 	 * Create this Drag and Drop Content
@@ -105,14 +105,14 @@ public:
 	 *
 	 * @return a new FDockingDragOperation
 	 */
-	static TSharedRef<FDraggableBoxUIDragOperation> New(const TSharedRef<class SWidget>& InUIToBeDragged, const FVector2D InDecoratorOffset, 
+	MODELINGUI_API static TSharedRef<FDraggableBoxUIDragOperation> New(const TSharedRef<class SWidget>& InUIToBeDragged, const FVector2D InDecoratorOffset, 
 		const FVector2D& OwnerAreaSize, SDraggableBox::FOnDragComplete& OnDragComplete);
 
-	virtual ~FDraggableBoxUIDragOperation() {}
+	MODELINGUI_API virtual ~FDraggableBoxUIDragOperation() {}
 
 protected:
 	/** The constructor is protected, so that this class can only be instantiated as a shared pointer. */
-	FDraggableBoxUIDragOperation(const TSharedRef<class SWidget>& InUIToBeDragged, const FVector2D InDecoratorOffsetFromCursor, 
+	MODELINGUI_API FDraggableBoxUIDragOperation(const TSharedRef<class SWidget>& InUIToBeDragged, const FVector2D InDecoratorOffsetFromCursor, 
 		const FVector2D& OwnerAreaSize, SDraggableBox::FOnDragComplete& OnDragCompletee);
 
 	/**
diff --git a/Engine/Plugins/Online/OnlineSubsystem/Source/Public/Interfaces/OnlineUserCloudInterface.h b/Engine/Plugins/Online/OnlineSubsystem/Source/Public/Interfaces/OnlineUserCloudInterface.h
--- a/Engine/Plugins/Online/OnlineSubsystem/Source/Public/Interfaces/OnlineUserCloudInterface.h
+++ b/Engine/Plugins/Online/OnlineSubsystem/Source/Public/Interfaces/OnlineUserCloudInterface.h
@@ -129,14 +129,10 @@ class IOnlineUserCloud
 {
 
 protected:
-	PRAGMA_DISABLE_DEPRECATION_WARNINGS
-	IOnlineUserCloud() {};
-	PRAGMA_ENABLE_DEPRECATION_WARNINGS
+	ONLINESUBSYSTEM_API IOnlineUserCloud();
 
 public:
-	PRAGMA_DISABLE_DEPRECATION_WARNINGS
-	virtual ~IOnlineUserCloud() {};
-	PRAGMA_ENABLE_DEPRECATION_WARNINGS
+	ONLINESUBSYSTEM_API virtual ~IOnlineUserCloud();
 
 	/**
 	 * Copies the file data into the specified buffer for the specified file
diff --git a/Engine/Plugins/Online/SocketSubsystemEOS/Source/SocketSubsystemEOS/Public/SocketEOS.h b/Engine/Plugins/Online/SocketSubsystemEOS/Source/SocketSubsystemEOS/Public/SocketEOS.h
--- a/Engine/Plugins/Online/SocketSubsystemEOS/Source/SocketSubsystemEOS/Public/SocketEOS.h
+++ b/Engine/Plugins/Online/SocketSubsystemEOS/Source/SocketSubsystemEOS/Public/SocketEOS.h
@@ -14,12 +14,12 @@ class FSocketSubsystemEOS;
 	#include "eos_p2p_types.h"
 #endif
 
-class SOCKETSUBSYSTEMEOS_API FSocketEOS
+class FSocketEOS
 	: public FSocket
 {
 public:
-	FSocketEOS(FSocketSubsystemEOS& SocketSubsystem, const FString& InSocketDescription);
-	virtual ~FSocketEOS();
+	SOCKETSUBSYSTEMEOS_API FSocketEOS(FSocketSubsystemEOS& SocketSubsystem, const FString& InSocketDescription);
+	SOCKETSUBSYSTEMEOS_API virtual ~FSocketEOS();
 
 	//~ Begin FSocket Interface
 	virtual bool Shutdown(ESocketShutdownMode Mode) override;
diff --git a/Engine/Plugins/Online/SocketSubsystemEOS/Source/SocketSubsystemEOS/Public/SocketSubsystemEOS.h b/Engine/Plugins/Online/SocketSubsystemEOS/Source/SocketSubsystemEOS/Public/SocketSubsystemEOS.h
--- a/Engine/Plugins/Online/SocketSubsystemEOS/Source/SocketSubsystemEOS/Public/SocketSubsystemEOS.h
+++ b/Engine/Plugins/Online/SocketSubsystemEOS/Source/SocketSubsystemEOS/Public/SocketSubsystemEOS.h
@@ -29,12 +29,12 @@ class FSocketEOS;
 
 typedef TSet<uint8> FChannelSet;
 
-class SOCKETSUBSYSTEMEOS_API FSocketSubsystemEOS
+class FSocketSubsystemEOS
 	: public ISocketSubsystem
 {
 public:
-	FSocketSubsystemEOS(IEOSPlatformHandlePtr InPlatformHandle, ISocketSubsystemEOSUtilsPtr InUtils);
-	virtual ~FSocketSubsystemEOS();
+	SOCKETSUBSYSTEMEOS_API FSocketSubsystemEOS(IEOSPlatformHandlePtr InPlatformHandle, ISocketSubsystemEOSUtilsPtr InUtils);
+	SOCKETSUBSYSTEMEOS_API virtual ~FSocketSubsystemEOS();
 
 //~ Begin ISocketSubsystem Interface
 	virtual bool Init(FString& Error) override;
diff --git a/Engine/Plugins/Runtime/GameplayAbilities/Source/GameplayAbilities/Public/GameplayDebuggerCategory_Abilities.h b/Engine/Plugins/Runtime/GameplayAbilities/Source/GameplayAbilities/Public/GameplayDebuggerCategory_Abilities.h
--- a/Engine/Plugins/Runtime/GameplayAbilities/Source/GameplayAbilities/Public/GameplayDebuggerCategory_Abilities.h
+++ b/Engine/Plugins/Runtime/GameplayAbilities/Source/GameplayAbilities/Public/GameplayDebuggerCategory_Abilities.h
@@ -14,20 +14,20 @@ class APlayerController;
 class UAbilitySystemComponent;
 class UPackageMap;
 
-class GAMEPLAYABILITIES_API FGameplayDebuggerCategory_Abilities : public FGameplayDebuggerCategory
+class FGameplayDebuggerCategory_Abilities : public FGameplayDebuggerCategory
 {
 public:
-	FGameplayDebuggerCategory_Abilities();
+	GAMEPLAYABILITIES_API FGameplayDebuggerCategory_Abilities();
 
-	virtual void CollectData(APlayerController* OwnerPC, AActor* DebugActor) override;
-	virtual void DrawData(APlayerController* OwnerPC, FGameplayDebuggerCanvasContext& CanvasContext) override;
+	GAMEPLAYABILITIES_API virtual void CollectData(APlayerController* OwnerPC, AActor* DebugActor) override;
+	GAMEPLAYABILITIES_API virtual void DrawData(APlayerController* OwnerPC, FGameplayDebuggerCanvasContext& CanvasContext) override;
 
-	static TSharedRef<FGameplayDebuggerCategory> MakeInstance();
+	GAMEPLAYABILITIES_API static TSharedRef<FGameplayDebuggerCategory> MakeInstance();
 
-	void OnShowGameplayTagsToggle();
-	void OnShowGameplayAbilitiesToggle();
-	void OnShowGameplayEffectsToggle();
-	void OnShowGameplayAttributesToggle();
+	GAMEPLAYABILITIES_API void OnShowGameplayTagsToggle();
+	GAMEPLAYABILITIES_API void OnShowGameplayAbilitiesToggle();
+	GAMEPLAYABILITIES_API void OnShowGameplayEffectsToggle();
+	GAMEPLAYABILITIES_API void OnShowGameplayAttributesToggle();
 		
 
 	// Some GAS features such as Attributes can exist on the server, client, or both.  We can also get 'detached' if both sides have the same values (such as Attributes) that aren't networked.
@@ -42,10 +42,10 @@ public:
 
 protected:
 
-	void DrawGameplayTags(FGameplayDebuggerCanvasContext& CanvasContext, const APlayerController* OwnerPC) const;
-	void DrawGameplayAbilities(FGameplayDebuggerCanvasContext& CanvasContext, const APlayerController* OwnerPC) const;
-	void DrawGameplayEffects(FGameplayDebuggerCanvasContext& CanvasContext, const APlayerController* OwnerPC) const;
-	void DrawGameplayAttributes(FGameplayDebuggerCanvasContext& CanvasContext, const APlayerController* OwnerPC) const;
+	GAMEPLAYABILITIES_API void DrawGameplayTags(FGameplayDebuggerCanvasContext& CanvasContext, const APlayerController* OwnerPC) const;
+	GAMEPLAYABILITIES_API void DrawGameplayAbilities(FGameplayDebuggerCanvasContext& CanvasContext, const APlayerController* OwnerPC) const;
+	GAMEPLAYABILITIES_API void DrawGameplayEffects(FGameplayDebuggerCanvasContext& CanvasContext, const APlayerController* OwnerPC) const;
+	GAMEPLAYABILITIES_API void DrawGameplayAttributes(FGameplayDebuggerCanvasContext& CanvasContext, const APlayerController* OwnerPC) const;
 
 	struct FRepData
 	{
@@ -89,7 +89,7 @@ protected:
 	};
 	FRepData DataPack;
 
-	bool WrapStringAccordingToViewport(const FString& iStr, FString& oStr, FGameplayDebuggerCanvasContext& CanvasContext, float ViewportWitdh) const;
+	GAMEPLAYABILITIES_API bool WrapStringAccordingToViewport(const FString& iStr, FString& oStr, FGameplayDebuggerCanvasContext& CanvasContext, float ViewportWitdh) const;
 
 private:
 	TArray<FRepData::FGameplayAttributeDebug> CollectAttributeData(const APlayerController* OwnerPC, const UAbilitySystemComponent* AbilityComp) const;
diff --git a/Engine/Plugins/Runtime/Harmonix/Source/HarmonixDsp/Private/HarmonixDsp/Effects/Delay.cpp b/Engine/Plugins/Runtime/Harmonix/Source/HarmonixDsp/Private/HarmonixDsp/Effects/Delay.cpp
--- a/Engine/Plugins/Runtime/Harmonix/Source/HarmonixDsp/Private/HarmonixDsp/Effects/Delay.cpp
+++ b/Engine/Plugins/Runtime/Harmonix/Source/HarmonixDsp/Private/HarmonixDsp/Effects/Delay.cpp
@@ -60,6 +60,11 @@ namespace Harmonix::Dsp::Effects
 		SetParamsToTargets();
 	}
 
+	FDelay::~FDelay()
+	{
+		FreeUpMemory();
+	}
+
 	void FDelay::Prepare(float InSampleRate, uint32 InMaxChannels, float InMaxDelayTimeMs)
 	{
 		FreeUpMemory();
diff --git a/Engine/Plugins/Runtime/Harmonix/Source/HarmonixDsp/Public/HarmonixDsp/Effects/Delay.h b/Engine/Plugins/Runtime/Harmonix/Source/HarmonixDsp/Public/HarmonixDsp/Effects/Delay.h
--- a/Engine/Plugins/Runtime/Harmonix/Source/HarmonixDsp/Public/HarmonixDsp/Effects/Delay.h
+++ b/Engine/Plugins/Runtime/Harmonix/Source/HarmonixDsp/Public/HarmonixDsp/Effects/Delay.h
@@ -18,7 +18,7 @@ namespace Harmonix::Dsp::Effects
 		class FTestDelay;
 	}
 
-	class HARMONIXDSP_API FDelay
+	class FDelay
 	{
 	public:
 		
@@ -28,16 +28,13 @@ namespace Harmonix::Dsp::Effects
 		// The max number of channels we expect to support with the delay
 		inline static constexpr int32 AbsoluteMaxChannels = 12;
 		
-		FDelay();
+		HARMONIXDSP_API FDelay();
 
-		void Prepare(float InSampleRate, uint32 InMaxChannels, float InMaxDelayTimeMs);
+		HARMONIXDSP_API void Prepare(float InSampleRate, uint32 InMaxChannels, float InMaxDelayTimeMs);
 
-		void Unprepare();
+		HARMONIXDSP_API void Unprepare();
 
-		virtual ~FDelay()
-		{
-			FreeUpMemory();
-		}
+		HARMONIXDSP_API virtual ~FDelay();
 
 		void Clear()
 		{
@@ -83,74 +80,74 @@ namespace Harmonix::Dsp::Effects
 
 		*/
 
-		void Process(Audio::FMultichannelBufferView& InOutBuffer);
+		HARMONIXDSP_API void Process(Audio::FMultichannelBufferView& InOutBuffer);
 
-		float CalculateSecsToIdle();
+		HARMONIXDSP_API float CalculateSecsToIdle();
 
 		float GetSampleRate() const { return SampleRate; }
 		
 		// units are either seconds or beats, depending on beat-sync setting
-		void  SetDelaySeconds(float Seconds);
+		HARMONIXDSP_API void  SetDelaySeconds(float Seconds);
 		float GetDelaySeconds() const { return DelayTimeSeconds; }
 
 		// in beat-sync mode, delays are specified in beats, not seconds, and the
 		//  delay must be notified when the tempo changes.
-		void  SetTimeSyncOption(ETimeSyncOption Option);
+		HARMONIXDSP_API void  SetTimeSyncOption(ETimeSyncOption Option);
 		ETimeSyncOption  GetTimeSyncOption() const { return TimeSyncOption; }
 
-		void  SetTempo(float Bpm);
+		HARMONIXDSP_API void  SetTempo(float Bpm);
 		float GetTempo() const { return TempoBpm; }
 
-		void  SetSpeed(float Speed);
+		HARMONIXDSP_API void  SetSpeed(float Speed);
 		float GetSpeed() const { return Speed; }
 
 		float GetOutputGain() const { return OutputGain; }
 		void SetOutputGain(const float InGain) { OutputGain = InGain; }
 
-		void SetFeedbackGain(float Gain);
-		float GetFeedbackGain() const;
+		HARMONIXDSP_API void SetFeedbackGain(float Gain);
+		HARMONIXDSP_API float GetFeedbackGain() const;
 
-		void SetWetGain(float Gain);
-		float GetWetGain() const;
+		HARMONIXDSP_API void SetWetGain(float Gain);
+		HARMONIXDSP_API float GetWetGain() const;
 
-		void SetDryGain(float Gain);
-		float GetDryGain() const;
+		HARMONIXDSP_API void SetDryGain(float Gain);
+		HARMONIXDSP_API float GetDryGain() const;
 
-		void SetWetFilterEnabled(bool Enabled);
+		HARMONIXDSP_API void SetWetFilterEnabled(bool Enabled);
 		bool GetWetFilterEnabled() const { return WetFilters.GetSettings().IsEnabled; }
 
-		void SetFeedbackFilterEnabled(bool Enabled);
+		HARMONIXDSP_API void SetFeedbackFilterEnabled(bool Enabled);
 		bool GetFeedbackFilterEnabled() const { return FeedbackFilters.GetSettings().IsEnabled; }
 		
-		void SetFilterFreq(float Freq);
+		HARMONIXDSP_API void SetFilterFreq(float Freq);
 		float GetFilterFreq() const { return WetFilters.GetSettings().Freq; }
 
-		void SetFilterQ(float Q);
+		HARMONIXDSP_API void SetFilterQ(float Q);
 		float GetFilterQ() const { return WetFilters.GetSettings().Q; }
 
-		void SetFilterType(EDelayFilterType Type);
-		void SetFilterType(EBiquadFilterType Type);
+		HARMONIXDSP_API void SetFilterType(EDelayFilterType Type);
+		HARMONIXDSP_API void SetFilterType(EBiquadFilterType Type);
 		EBiquadFilterType GetFilterType() const { return WetFilters.GetSettings().Type; }
 
-		void SetLfoEnabled(bool bEnabled);
+		HARMONIXDSP_API void SetLfoEnabled(bool bEnabled);
 		bool GetLfoEnabled() const { return LfoSettings.IsEnabled; }
 
-		void SetLfoTimeSyncOption(ETimeSyncOption Option);
+		HARMONIXDSP_API void SetLfoTimeSyncOption(ETimeSyncOption Option);
 		ETimeSyncOption GetLfoTimeSyncOption() const { return LfoSyncOption; }
 
-		void SetLfoFreq(float Freq);
+		HARMONIXDSP_API void SetLfoFreq(float Freq);
 		float GetLfoFreq() const { return LfoSettings.Freq; }
 
-		void SetLfoDepth(float InDepth);
-		float GetLfoDepth() const;
+		HARMONIXDSP_API void SetLfoDepth(float InDepth);
+		HARMONIXDSP_API float GetLfoDepth() const;
 
-		void SetStereoSpreadLeft(float Spread);
+		HARMONIXDSP_API void SetStereoSpreadLeft(float Spread);
 		float GetStereoSpreadLeft() const { return DelaySpreadLeft.GetCurrent(); }
 
-		void SetStereoSpreadRight(float Spread);
+		HARMONIXDSP_API void SetStereoSpreadRight(float Spread);
 		float GetStereoSpreadRight() const { return DelaySpreadRight.GetCurrent(); }
 
-		void SetStereoType(EDelayStereoType Type);
+		HARMONIXDSP_API void SetStereoType(EDelayStereoType Type);
 		EDelayStereoType GetStereoType() const { return DelayType; }
 
 	private:
diff --git a/Engine/Plugins/Runtime/Harmonix/Source/HarmonixDsp/Public/HarmonixDsp/Modulators/Lfo.h b/Engine/Plugins/Runtime/Harmonix/Source/HarmonixDsp/Public/HarmonixDsp/Modulators/Lfo.h
--- a/Engine/Plugins/Runtime/Harmonix/Source/HarmonixDsp/Public/HarmonixDsp/Modulators/Lfo.h
+++ b/Engine/Plugins/Runtime/Harmonix/Source/HarmonixDsp/Public/HarmonixDsp/Modulators/Lfo.h
@@ -50,39 +50,39 @@ private:
 // because it doesn't do anything about aliasing, which is audible if the
 // frequency is in the audible range. infrasonic frequencies are the target
 // for this implementation.
-struct HARMONIXDSP_API FLfo
+struct FLfo
 {
 public:
 
-	FLfo();
+	HARMONIXDSP_API FLfo();
 
-	void Prepare(float InSampleRate);
+	HARMONIXDSP_API void Prepare(float InSampleRate);
 
 	//  input is [0,1], returns [0,1]
 	typedef float (*FWaveShapeFunction)(float);
 
-	void UseSettings(const FLfoSettings* InSettings);
+	HARMONIXDSP_API void UseSettings(const FLfoSettings* InSettings);
 
 	const FLfoSettings* GetSettings() const 
 	{ 
 		return Settings; 
 	}
 
-	void SetPhase(double InPhase);
+	HARMONIXDSP_API void SetPhase(double InPhase);
 
-	double GetPhase() const;
+	HARMONIXDSP_API double GetPhase() const;
 
 	// sets phase to the initial phase indicated
 	// by the user setting (and possibly by the shape)
-	void Retrigger();
+	HARMONIXDSP_API void Retrigger();
 
-	TInterval<float> GetRange() const;
-	void SetRangeAndMode(const TInterval<float>& InRange, ELfoMode InMode);
+	HARMONIXDSP_API TInterval<float> GetRange() const;
+	HARMONIXDSP_API void SetRangeAndMode(const TInterval<float>& InRange, ELfoMode InMode);
 
 	// advances the phase of the oscillator
-	void Advance(uint32 InNumSamples);
+	HARMONIXDSP_API void Advance(uint32 InNumSamples);
 
-	float GetValue() const;
+	HARMONIXDSP_API float GetValue() const;
 private:
 	void ComputeCyclesPerSample();
 
diff --git a/Engine/Source/Runtime/AdvancedWidgets/Private/Slate/SRadialSlider.cpp b/Engine/Source/Runtime/AdvancedWidgets/Private/Slate/SRadialSlider.cpp
--- a/Engine/Source/Runtime/AdvancedWidgets/Private/Slate/SRadialSlider.cpp
+++ b/Engine/Source/Runtime/AdvancedWidgets/Private/Slate/SRadialSlider.cpp
@@ -20,6 +20,8 @@ SRadialSlider::SRadialSlider()
 #endif
 }
 
+SRadialSlider::~SRadialSlider() = default;
+
 void SRadialSlider::Construct( const SRadialSlider::FArguments& InDeclaration )
 {
 	check(InDeclaration._Style);
diff --git a/Engine/Source/Runtime/AdvancedWidgets/Public/Widgets/SRadialSlider.h b/Engine/Source/Runtime/AdvancedWidgets/Public/Widgets/SRadialSlider.h
--- a/Engine/Source/Runtime/AdvancedWidgets/Public/Widgets/SRadialSlider.h
+++ b/Engine/Source/Runtime/AdvancedWidgets/Public/Widgets/SRadialSlider.h
@@ -145,6 +145,7 @@ public:
 	SLATE_END_ARGS()
 
 	ADVANCEDWIDGETS_API SRadialSlider();
+	ADVANCEDWIDGETS_API virtual ~SRadialSlider();
 
 	/**
 	 * Construct the widget.
diff --git a/Engine/Source/Runtime/AudioExtensions/Private/IAudioParameterTransmitter.cpp b/Engine/Source/Runtime/AudioExtensions/Private/IAudioParameterTransmitter.cpp
--- a/Engine/Source/Runtime/AudioExtensions/Private/IAudioParameterTransmitter.cpp
+++ b/Engine/Source/Runtime/AudioExtensions/Private/IAudioParameterTransmitter.cpp
@@ -18,6 +18,8 @@ namespace Audio
 	{
 	}
 
+	FParameterTransmitterBase::~FParameterTransmitterBase() = default;
+
 	bool FParameterTransmitterBase::GetParameter(FName InName, FAudioParameter& OutValue) const
 	{
 		if (const FAudioParameter* Param = FAudioParameter::FindParam(AudioParameters, InName))
diff --git a/Engine/Source/Runtime/AudioExtensions/Public/IAudioParameterTransmitter.h b/Engine/Source/Runtime/AudioExtensions/Public/IAudioParameterTransmitter.h
--- a/Engine/Source/Runtime/AudioExtensions/Public/IAudioParameterTransmitter.h
+++ b/Engine/Source/Runtime/AudioExtensions/Public/IAudioParameterTransmitter.h
@@ -92,7 +92,7 @@ namespace Audio
 	{
 	public:
 		AUDIOEXTENSIONS_API FParameterTransmitterBase(TArray<FAudioParameter>&& InDefaultParams);
-		virtual ~FParameterTransmitterBase() = default;
+		AUDIOEXTENSIONS_API virtual ~FParameterTransmitterBase();
 
 		AUDIOEXTENSIONS_API virtual bool GetParameter(FName InName, FAudioParameter& OutParam) const override;
 		AUDIOEXTENSIONS_API virtual void ResetParameters() override;
diff --git a/Engine/Source/Runtime/Core/Private/GenericPlatform/GenericPlatformFile.cpp b/Engine/Source/Runtime/Core/Private/GenericPlatform/GenericPlatformFile.cpp
--- a/Engine/Source/Runtime/Core/Private/GenericPlatform/GenericPlatformFile.cpp
+++ b/Engine/Source/Runtime/Core/Private/GenericPlatform/GenericPlatformFile.cpp
@@ -489,6 +489,9 @@ bool FGenericReadRequest::CheckForPrecache()
 	return false;
 }
 
+IPlatformFile::IPlatformFile() = default;
+IPlatformFile::~IPlatformFile() = default;
+
 FFileOpenResult IPlatformFile::OpenRead(const TCHAR* Filename, IPlatformFile::EOpenReadFlags Flags)
 {
 	if (EnumHasAnyFlags(Flags, EOpenReadFlags::AllowDelete))
diff --git a/Engine/Source/Runtime/Core/Private/Misc/AutomationTest.cpp b/Engine/Source/Runtime/Core/Private/Misc/AutomationTest.cpp
--- a/Engine/Source/Runtime/Core/Private/Misc/AutomationTest.cpp
+++ b/Engine/Source/Runtime/Core/Private/Misc/AutomationTest.cpp
@@ -147,6 +147,22 @@ namespace AutomationTest
 	}
 };
 
+FAutomationTestBase::FAutomationTestBase( const FString& InName, const bool bInComplexTask )
+	: bComplexTask( bInComplexTask )
+{
+	LLM_SCOPE_BYNAME(TEXT("AutomationTest/Framework"));
+	TestName = InName;
+	// Register the newly created automation test into the automation testing framework
+	FAutomationTestFramework::Get().RegisterAutomationTest( InName, this );
+}
+
+/** Destructor */
+FAutomationTestBase::~FAutomationTestBase() 
+{ 
+	// Unregister the automation test from the automation testing framework
+	FAutomationTestFramework::Get().UnregisterAutomationTest( TestName );
+}
+
 bool FAutomationTestBase::bSuppressLogWarnings = false;
 bool FAutomationTestBase::bSuppressLogErrors = false;
 bool FAutomationTestBase::bElevateLogWarningsToErrors = false;
diff --git a/Engine/Source/Runtime/Core/Private/Misc/BufferedOutputDevice.cpp b/Engine/Source/Runtime/Core/Private/Misc/BufferedOutputDevice.cpp
--- a/Engine/Source/Runtime/Core/Private/Misc/BufferedOutputDevice.cpp
+++ b/Engine/Source/Runtime/Core/Private/Misc/BufferedOutputDevice.cpp
@@ -5,6 +5,9 @@
 #include "Misc/ScopeLock.h"
 #include "Templates/UnrealTemplate.h"
 
+FBufferedOutputDevice::FBufferedOutputDevice() = default; // Here to satisfy modular builds on some old platforms
+FBufferedOutputDevice::~FBufferedOutputDevice() = default; // Here to satisfy modular builds on some old platforms
+
 void FBufferedOutputDevice::Serialize(const TCHAR* InData, ELogVerbosity::Type Verbosity, const FName& Category)
 {
 	if (Verbosity > FilterLevel)
diff --git a/Engine/Source/Runtime/Core/Private/Misc/LocalTimestampDirectoryVisitor.cpp b/Engine/Source/Runtime/Core/Private/Misc/LocalTimestampDirectoryVisitor.cpp
--- a/Engine/Source/Runtime/Core/Private/Misc/LocalTimestampDirectoryVisitor.cpp
+++ b/Engine/Source/Runtime/Core/Private/Misc/LocalTimestampDirectoryVisitor.cpp
@@ -30,6 +30,7 @@ FLocalTimestampDirectoryVisitor::FLocalTimestampDirectoryVisitor( IPlatformFile&
 	}
 }
 
+FLocalTimestampDirectoryVisitor::~FLocalTimestampDirectoryVisitor() = default;
 
 /* FLocalTimestampVisitor interface
  *****************************************************************************/
diff --git a/Engine/Source/Runtime/Core/Private/Misc/OutputDevice.cpp b/Engine/Source/Runtime/Core/Private/Misc/OutputDevice.cpp
--- a/Engine/Source/Runtime/Core/Private/Misc/OutputDevice.cpp
+++ b/Engine/Source/Runtime/Core/Private/Misc/OutputDevice.cpp
@@ -19,6 +19,17 @@ class FOutputDeviceError;
 
 DEFINE_LOG_CATEGORY(LogOutputDevice);
 
+FOutputDevice::FOutputDevice()
+		: bSuppressEventTag      (false)
+		, bAutoEmitLineTerminator(true)
+{}
+
+FOutputDevice::FOutputDevice(FOutputDevice&&) = default;
+FOutputDevice::FOutputDevice(const FOutputDevice&) = default;
+FOutputDevice& FOutputDevice::operator=(FOutputDevice&&) = default;
+FOutputDevice& FOutputDevice::operator=(const FOutputDevice&) = default;
+FOutputDevice::~FOutputDevice() = default;
+
 void FOutputDevice::Log( ELogVerbosity::Type Verbosity, const TCHAR* Str )
 {
 	Serialize( Str, Verbosity, NAME_None );
diff --git a/Engine/Source/Runtime/Core/Private/Misc/OutputDeviceMemory.cpp b/Engine/Source/Runtime/Core/Private/Misc/OutputDeviceMemory.cpp
--- a/Engine/Source/Runtime/Core/Private/Misc/OutputDeviceMemory.cpp
+++ b/Engine/Source/Runtime/Core/Private/Misc/OutputDeviceMemory.cpp
@@ -40,6 +40,8 @@ FOutputDeviceMemory::FOutputDeviceMemory(int32 InPreserveSize /*= 256 * 1024*/,
 	}
 }
 
+FOutputDeviceMemory::~FOutputDeviceMemory() = default;
+
 void FOutputDeviceMemory::TearDown() 
 {
 	if (!bSuppressEventTag)
diff --git a/Engine/Source/Runtime/Core/Private/Misc/SecureHash.cpp b/Engine/Source/Runtime/Core/Private/Misc/SecureHash.cpp
--- a/Engine/Source/Runtime/Core/Private/Misc/SecureHash.cpp
+++ b/Engine/Source/Runtime/Core/Private/Misc/SecureHash.cpp
@@ -1533,6 +1533,11 @@ void appOnFailSHAVerification(const TCHAR* FailedPathname, bool bFailedDueToMiss
 #endif
 }
 
+FBufferReaderWithSHA::~FBufferReaderWithSHA()
+{
+	Close();
+}
+
 bool FBufferReaderWithSHA::Close()
 {
 	// don't redo if we were already closed
diff --git a/Engine/Source/Runtime/Core/Private/Misc/TextFilterExpressionEvaluator.cpp b/Engine/Source/Runtime/Core/Private/Misc/TextFilterExpressionEvaluator.cpp
--- a/Engine/Source/Runtime/Core/Private/Misc/TextFilterExpressionEvaluator.cpp
+++ b/Engine/Source/Runtime/Core/Private/Misc/TextFilterExpressionEvaluator.cpp
@@ -257,6 +257,8 @@ FTextFilterExpressionEvaluator::FTextFilterExpressionEvaluator(const FTextFilter
 	SetFilterText(Other.FilterText);
 }
 
+FTextFilterExpressionEvaluator::~FTextFilterExpressionEvaluator() = default;
+
 FTextFilterExpressionEvaluator& FTextFilterExpressionEvaluator::operator=(const FTextFilterExpressionEvaluator& Other)
 {
 	FilterType = ETextFilterExpressionType::Empty;
diff --git a/Engine/Source/Runtime/Core/Private/Serialization/Archive.cpp b/Engine/Source/Runtime/Core/Private/Serialization/Archive.cpp
--- a/Engine/Source/Runtime/Core/Private/Serialization/Archive.cpp
+++ b/Engine/Source/Runtime/Core/Private/Serialization/Archive.cpp
@@ -1247,6 +1247,8 @@ void FArchive::SerializeCompressedNew(void* V, int64 Length, FName CompressionFo
 	}
 }
 
+FArchive::~FArchive() = default;
+
 void FArchive::ByteSwap(void* V, int32 Length)
 {
 	uint8* Ptr = (uint8*)V;
diff --git a/Engine/Source/Runtime/Core/Private/Serialization/BitReader.cpp b/Engine/Source/Runtime/Core/Private/Serialization/BitReader.cpp
--- a/Engine/Source/Runtime/Core/Private/Serialization/BitReader.cpp
+++ b/Engine/Source/Runtime/Core/Private/Serialization/BitReader.cpp
@@ -163,6 +163,12 @@ FBitReader::FBitReader(const uint8* Src, int64 CountBits)
 	}
 }
 
+FBitReader::~FBitReader() = default;
+FBitReader::FBitReader(const FBitReader&) = default;
+FBitReader& FBitReader::operator=(const FBitReader&) = default;
+FBitReader::FBitReader(FBitReader&&) = default;
+FBitReader& FBitReader::operator=(FBitReader&&) = default;
+
 void FBitReader::SetData( uint8* Src, int64 CountBits )
 {
 	Num			= CountBits;
diff --git a/Engine/Source/Runtime/Core/Private/Serialization/BitWriter.cpp b/Engine/Source/Runtime/Core/Private/Serialization/BitWriter.cpp
--- a/Engine/Source/Runtime/Core/Private/Serialization/BitWriter.cpp
+++ b/Engine/Source/Runtime/Core/Private/Serialization/BitWriter.cpp
@@ -59,6 +59,12 @@ FBitWriter::FBitWriter(void)
 	ArIsNetArchive = true;
 }
 
+FBitWriter::FBitWriter(const FBitWriter&) = default;
+FBitWriter& FBitWriter::operator=(const FBitWriter&) = default;
+FBitWriter::FBitWriter(FBitWriter&&) = default;
+FBitWriter& FBitWriter::operator=(FBitWriter&&) = default;
+FBitWriter::~FBitWriter() = default;
+
 /**
  * Resets the bit writer back to its initial state
  */
diff --git a/Engine/Source/Runtime/Core/Private/Serialization/LargeMemoryWriter.cpp b/Engine/Source/Runtime/Core/Private/Serialization/LargeMemoryWriter.cpp
--- a/Engine/Source/Runtime/Core/Private/Serialization/LargeMemoryWriter.cpp
+++ b/Engine/Source/Runtime/Core/Private/Serialization/LargeMemoryWriter.cpp
@@ -20,6 +20,8 @@ FLargeMemoryWriter::FLargeMemoryWriter(const int64 PreAllocateBytes, bool bIsPer
 	this->SetIsPersistent(bIsPersistent);
 }
 
+FLargeMemoryWriter::~FLargeMemoryWriter() = default;
+
 void FLargeMemoryWriter::Serialize(void* InData, int64 Num)
 {
 	UE_CLOG(!Data.HasData(), LogSerialization, Fatal, TEXT("Tried to serialize data to an FLargeMemoryWriter that was already released. Archive name: %s."), *ArchiveName);
diff --git a/Engine/Source/Runtime/Core/Public/GenericPlatform/GenericPlatformFile.h b/Engine/Source/Runtime/Core/Public/GenericPlatform/GenericPlatformFile.h
--- a/Engine/Source/Runtime/Core/Public/GenericPlatform/GenericPlatformFile.h
+++ b/Engine/Source/Runtime/Core/Public/GenericPlatform/GenericPlatformFile.h
@@ -336,8 +336,9 @@ public:
 	static CORE_API IPlatformFile& GetPlatformPhysical();
 	/** Returns the name of the physical platform file type. */
 	static CORE_API const TCHAR* GetPhysicalTypeName();
-	/** Destructor. */
-	virtual ~IPlatformFile() {}
+	/** Constructor/Destructor. */
+	CORE_API IPlatformFile();
+	CORE_API virtual ~IPlatformFile();
 
 	/**
 	 *	Set whether the sandbox is enabled or not
diff --git a/Engine/Source/Runtime/Core/Public/Misc/AutomationTest.h b/Engine/Source/Runtime/Core/Public/Misc/AutomationTest.h
--- a/Engine/Source/Runtime/Core/Public/Misc/AutomationTest.h
+++ b/Engine/Source/Runtime/Core/Public/Misc/AutomationTest.h
@@ -1532,21 +1532,10 @@ public:
 	 *
 	 * @param	InName	Name of the test
 	 */
-	FAutomationTestBase( const FString& InName, const bool bInComplexTask )
-		: bComplexTask( bInComplexTask )
-	{
-		LLM_SCOPE_BYNAME(TEXT("AutomationTest/Framework"));
-		TestName = InName;
-		// Register the newly created automation test into the automation testing framework
-		FAutomationTestFramework::Get().RegisterAutomationTest( InName, this );
-	}
+	CORE_API FAutomationTestBase( const FString& InName, const bool bInComplexTask );
 
 	/** Destructor */
-	virtual ~FAutomationTestBase() 
-	{ 
-		// Unregister the automation test from the automation testing framework
-		FAutomationTestFramework::Get().UnregisterAutomationTest( TestName );
-	}
+	CORE_API virtual ~FAutomationTestBase() ;
 
 	/** Log flags */
 	static CORE_API bool bSuppressLogWarnings;
diff --git a/Engine/Source/Runtime/Core/Public/Misc/BufferedOutputDevice.h b/Engine/Source/Runtime/Core/Public/Misc/BufferedOutputDevice.h
--- a/Engine/Source/Runtime/Core/Public/Misc/BufferedOutputDevice.h
+++ b/Engine/Source/Runtime/Core/Public/Misc/BufferedOutputDevice.h
@@ -24,4 +24,7 @@ public:
 
 	/** Pushes buffered lines into the specified output device. */
 	CORE_API void	RedirectTo(FOutputDevice& Ar);
+
+	CORE_API FBufferedOutputDevice();
+	CORE_API virtual ~FBufferedOutputDevice();
 };
diff --git a/Engine/Source/Runtime/Core/Public/Misc/LocalTimestampDirectoryVisitor.h b/Engine/Source/Runtime/Core/Public/Misc/LocalTimestampDirectoryVisitor.h
--- a/Engine/Source/Runtime/Core/Public/Misc/LocalTimestampDirectoryVisitor.h
+++ b/Engine/Source/Runtime/Core/Public/Misc/LocalTimestampDirectoryVisitor.h
@@ -33,6 +33,7 @@ public:
 	 * @param bInMakeLowerCase - Whether to lower case filenames and directories.
 	 */
 	CORE_API FLocalTimestampDirectoryVisitor( IPlatformFile& InFileInterface, const TArray<FString>& InDirectoriesToIgnore, const TArray<FString>& InDirectoriesToNotRecurse, bool bInCacheDirectories = false, bool bInMakeLowerCase = false );
+	CORE_API virtual ~FLocalTimestampDirectoryVisitor();
 
 public:
 
diff --git a/Engine/Source/Runtime/Core/Public/Misc/OutputDevice.h b/Engine/Source/Runtime/Core/Public/Misc/OutputDevice.h
--- a/Engine/Source/Runtime/Core/Public/Misc/OutputDevice.h
+++ b/Engine/Source/Runtime/Core/Public/Misc/OutputDevice.h
@@ -132,17 +132,14 @@ class FName;
 class FOutputDevice
 {
 public:
-	FOutputDevice()
-		 : bSuppressEventTag      (false)
-		 , bAutoEmitLineTerminator(true)
-	{}
+	CORE_API FOutputDevice();
 
-	FOutputDevice(FOutputDevice&&) = default;
-	FOutputDevice(const FOutputDevice&) = default;
-	FOutputDevice& operator=(FOutputDevice&&) = default;
-	FOutputDevice& operator=(const FOutputDevice&) = default;
+	CORE_API FOutputDevice(FOutputDevice&&);
+	CORE_API FOutputDevice(const FOutputDevice&);
+	CORE_API FOutputDevice& operator=(FOutputDevice&&);
+	CORE_API FOutputDevice& operator=(const FOutputDevice&);
 
-	virtual ~FOutputDevice() = default;
+	CORE_API virtual ~FOutputDevice();
 
 	// FOutputDevice interface.
 	virtual void Serialize( const TCHAR* V, ELogVerbosity::Type Verbosity, const FName& Category ) = 0;
diff --git a/Engine/Source/Runtime/Core/Public/Misc/OutputDeviceMemory.h b/Engine/Source/Runtime/Core/Public/Misc/OutputDeviceMemory.h
--- a/Engine/Source/Runtime/Core/Public/Misc/OutputDeviceMemory.h
+++ b/Engine/Source/Runtime/Core/Public/Misc/OutputDeviceMemory.h
@@ -23,6 +23,8 @@ public:
 	 * @param InBufferSize		Maximum size of the memory ring buffer
 	 */
 	CORE_API FOutputDeviceMemory(int32 InPreserveSize = 256 * 1024, int32 InBufferSize = 2048 * 1024, bool bInSuppressEventTag = false);
+	
+	CORE_API ~FOutputDeviceMemory();
 
 	/** Dumps the contents of the buffer to an archive */
 	CORE_API virtual void Dump(FArchive& Ar) override;
diff --git a/Engine/Source/Runtime/Core/Public/Misc/SecureHash.h b/Engine/Source/Runtime/Core/Public/Misc/SecureHash.h
--- a/Engine/Source/Runtime/Core/Public/Misc/SecureHash.h
+++ b/Engine/Source/Runtime/Core/Public/Misc/SecureHash.h
@@ -537,10 +537,7 @@ public:
 	{
 	}
 
-	~FBufferReaderWithSHA() override
-	{
-		Close();
-	}
+	CORE_API virtual ~FBufferReaderWithSHA();
 
 	CORE_API bool Close() override;
 
diff --git a/Engine/Source/Runtime/Core/Public/Misc/TextFilterExpressionEvaluator.h b/Engine/Source/Runtime/Core/Public/Misc/TextFilterExpressionEvaluator.h
--- a/Engine/Source/Runtime/Core/Public/Misc/TextFilterExpressionEvaluator.h
+++ b/Engine/Source/Runtime/Core/Public/Misc/TextFilterExpressionEvaluator.h
@@ -201,7 +201,7 @@ public:
 	CORE_API FTextFilterExpressionEvaluator(const ETextFilterExpressionEvaluatorMode InMode);
 	CORE_API FTextFilterExpressionEvaluator(const FTextFilterExpressionEvaluator& Other);
 	CORE_API FTextFilterExpressionEvaluator& operator=(const FTextFilterExpressionEvaluator& Other);
-	virtual ~FTextFilterExpressionEvaluator() {}
+	CORE_API virtual ~FTextFilterExpressionEvaluator();
 
 	/** Get the complexity of the current filter terms */
 	CORE_API ETextFilterExpressionType GetFilterType() const;
diff --git a/Engine/Source/Runtime/Core/Public/Serialization/Archive.h b/Engine/Source/Runtime/Core/Public/Serialization/Archive.h
--- a/Engine/Source/Runtime/Core/Public/Serialization/Archive.h
+++ b/Engine/Source/Runtime/Core/Public/Serialization/Archive.h
@@ -1141,7 +1141,7 @@ public:
 	FArchive() = default;
 	FArchive(const FArchive&) = default;
 	FArchive& operator=(const FArchive& ArchiveToCopy) = default;
-	~FArchive() = default;
+	CORE_API virtual ~FArchive();
 
 protected:
 	using FArchiveState::LinkProxy;
diff --git a/Engine/Source/Runtime/Core/Public/Serialization/BitReader.h b/Engine/Source/Runtime/Core/Public/Serialization/BitReader.h
--- a/Engine/Source/Runtime/Core/Public/Serialization/BitReader.h
+++ b/Engine/Source/Runtime/Core/Public/Serialization/BitReader.h
@@ -27,11 +27,12 @@ struct FBitReader : public FBitArchive
 
 public:
 	CORE_API FBitReader( const uint8* Src = nullptr, int64 CountBits = 0 );
+	CORE_API ~FBitReader();
 
-	FBitReader(const FBitReader&) = default;
-    FBitReader& operator=(const FBitReader&) = default;
-    FBitReader(FBitReader&&) = default;
-    FBitReader& operator=(FBitReader&&) = default;
+	CORE_API FBitReader(const FBitReader&);
+    CORE_API FBitReader& operator=(const FBitReader&);
+    CORE_API FBitReader(FBitReader&&);
+    CORE_API FBitReader& operator=(FBitReader&&);
 
 	CORE_API void SetData( FBitReader& Src, int64 CountBits );
 	CORE_API void SetData( uint8* Src, int64 CountBits );
diff --git a/Engine/Source/Runtime/Core/Public/Serialization/BitWriter.h b/Engine/Source/Runtime/Core/Public/Serialization/BitWriter.h
--- a/Engine/Source/Runtime/Core/Public/Serialization/BitWriter.h
+++ b/Engine/Source/Runtime/Core/Public/Serialization/BitWriter.h
@@ -24,17 +24,18 @@ struct FBitWriter : public FBitArchive
 
 public:
 	/** Default constructor. Zeros everything. */
-	CORE_API FBitWriter(void);
+	CORE_API FBitWriter();
+	CORE_API virtual ~FBitWriter();
 
 	/**
 	 * Constructor using known size the buffer needs to be
 	 */
 	CORE_API FBitWriter( int64 InMaxBits, bool AllowResize = false );
 
-	FBitWriter(const FBitWriter&) = default;
-    FBitWriter& operator=(const FBitWriter&) = default;
-    FBitWriter(FBitWriter&&) = default;
-    FBitWriter& operator=(FBitWriter&&) = default;
+	CORE_API FBitWriter(const FBitWriter&);
+    CORE_API FBitWriter& operator=(const FBitWriter&);
+    CORE_API FBitWriter(FBitWriter&&);
+    CORE_API FBitWriter& operator=(FBitWriter&&);
 
 	CORE_API virtual void SerializeBits( void* Src, int64 LengthBits ) override;
 
diff --git a/Engine/Source/Runtime/Core/Public/Serialization/LargeMemoryWriter.h b/Engine/Source/Runtime/Core/Public/Serialization/LargeMemoryWriter.h
--- a/Engine/Source/Runtime/Core/Public/Serialization/LargeMemoryWriter.h
+++ b/Engine/Source/Runtime/Core/Public/Serialization/LargeMemoryWriter.h
@@ -18,6 +18,7 @@ class FLargeMemoryWriter : public FMemoryArchive
 public:
 	
 	CORE_API FLargeMemoryWriter(const int64 PreAllocateBytes = 0, bool bIsPersistent = false, const TCHAR* InFilename = nullptr);
+	CORE_API virtual ~FLargeMemoryWriter();
 
 	CORE_API virtual void Serialize(void* InData, int64 Num) override;
 
diff --git a/Engine/Source/Runtime/CoreUObject/Private/UObject/Class.cpp b/Engine/Source/Runtime/CoreUObject/Private/UObject/Class.cpp
--- a/Engine/Source/Runtime/CoreUObject/Private/UObject/Class.cpp
+++ b/Engine/Source/Runtime/CoreUObject/Private/UObject/Class.cpp
@@ -6201,6 +6201,8 @@ UClass::UClass
 	*(const TCHAR**)&ClassConfigName = InConfigName; //-V580
 }
 
+UClass::~UClass() = default;
+
 void* UClass::CreateSparseClassData()
 {
 	check(SparseClassData == nullptr);
diff --git a/Engine/Source/Runtime/CoreUObject/Private/UObject/CoreNet.cpp b/Engine/Source/Runtime/CoreUObject/Private/UObject/CoreNet.cpp
--- a/Engine/Source/Runtime/CoreUObject/Private/UObject/CoreNet.cpp
+++ b/Engine/Source/Runtime/CoreUObject/Private/UObject/CoreNet.cpp
@@ -286,6 +286,10 @@ void FClassNetCache::CountBytes(FArchive& Ar) const
 	UPackageMap implementation.
 -----------------------------------------------------------------------------*/
 
+UPackageMap::UPackageMap() = default;
+UPackageMap::~UPackageMap() = default;
+
+
 bool UPackageMap::SerializeName(FArchive& Ar, FName& InName)
 {
 	return StaticSerializeName(Ar, InName);
@@ -511,6 +515,12 @@ FNetBitReader::FNetBitReader(UPackageMap* InPackageMap, const uint8* Src, int64
 {
 }
 
+FNetBitReader::FNetBitReader(const FNetBitReader&) = default;
+FNetBitReader& FNetBitReader::operator=(const FNetBitReader&) = default;
+FNetBitReader::FNetBitReader(FNetBitReader&&) = default;
+FNetBitReader& FNetBitReader::operator=(FNetBitReader&&) = default;
+FNetBitReader::~FNetBitReader() = default;
+
 FArchive& FNetBitReader::operator<<(UObject*& Object)
 {
 	PackageMap->SerializeObject(*this, UObject::StaticClass(), Object);
diff --git a/Engine/Source/Runtime/CoreUObject/Private/UObject/Enum.cpp b/Engine/Source/Runtime/CoreUObject/Private/UObject/Enum.cpp
--- a/Engine/Source/Runtime/CoreUObject/Private/UObject/Enum.cpp
+++ b/Engine/Source/Runtime/CoreUObject/Private/UObject/Enum.cpp
@@ -32,6 +32,8 @@ UEnum::UEnum(const FObjectInitializer& ObjectInitializer)
 {
 }
 
+UEnum::~UEnum() = default;
+
 void UEnum::Serialize( FArchive& Ar )
 {
 	Ar.UsingCustomVersion(FCoreObjectVersion::GUID);
diff --git a/Engine/Source/Runtime/CoreUObject/Private/UObject/GarbageCollection.cpp b/Engine/Source/Runtime/CoreUObject/Private/UObject/GarbageCollection.cpp
--- a/Engine/Source/Runtime/CoreUObject/Private/UObject/GarbageCollection.cpp
+++ b/Engine/Source/Runtime/CoreUObject/Private/UObject/GarbageCollection.cpp
@@ -3751,6 +3751,8 @@ FReferenceFinder::FReferenceFinder(TArray<UObject*>& InObjectArray, UObject* InO
 	}
 }
 
+FReferenceFinder::~FReferenceFinder() = default;
+
 void FReferenceFinder::FindReferences(UObject* Object, UObject* InReferencingObject, FProperty* InReferencingProperty)
 {
 	check(Object != NULL);
diff --git a/Engine/Source/Runtime/CoreUObject/Private/UObject/UObjectGlobals.cpp b/Engine/Source/Runtime/CoreUObject/Private/UObject/UObjectGlobals.cpp
--- a/Engine/Source/Runtime/CoreUObject/Private/UObject/UObjectGlobals.cpp
+++ b/Engine/Source/Runtime/CoreUObject/Private/UObject/UObjectGlobals.cpp
@@ -4772,6 +4772,8 @@ public:
 
 };
 
+FReferenceCollector::~FReferenceCollector() = default;
+
 void FReferenceCollector::AddStableReference(UObject** Object)
 {
 	AROPrivate::AddReferencedObject(*this, *Object);
diff --git a/Engine/Source/Runtime/CoreUObject/Public/UObject/Class.h b/Engine/Source/Runtime/CoreUObject/Public/UObject/Class.h
--- a/Engine/Source/Runtime/CoreUObject/Public/UObject/Class.h
+++ b/Engine/Source/Runtime/CoreUObject/Public/UObject/Class.h
@@ -2681,6 +2681,8 @@ public:
 	COREUOBJECT_API virtual void BeginDestroy() override;
 	// End of UObject interface.
 
+	COREUOBJECT_API ~UEnum();
+
 protected:
 	/** List of pairs of all enum names and values. */
 	TArray<TPair<FName, int64>> Names;
@@ -3102,6 +3104,9 @@ public:
 		ClassVTableHelperCtorCallerType InClassVTableHelperCtorCaller,
 		FUObjectCppClassStaticFunctions&& InCppClassStaticFunctions);
 
+	// Destructor
+	COREUOBJECT_API ~UClass();
+
 #if WITH_RELOAD
 	/**
 	 * Called when a class is reloading from a DLL...updates various information in-place.
diff --git a/Engine/Source/Runtime/CoreUObject/Public/UObject/CoreNet.h b/Engine/Source/Runtime/CoreUObject/Public/UObject/CoreNet.h
--- a/Engine/Source/Runtime/CoreUObject/Public/UObject/CoreNet.h
+++ b/Engine/Source/Runtime/CoreUObject/Public/UObject/CoreNet.h
@@ -232,6 +232,9 @@ class UPackageMap : public UObject
 
 	COREUOBJECT_API virtual void Serialize(FArchive& Ar) override;
 
+	COREUOBJECT_API UPackageMap();
+	COREUOBJECT_API virtual ~UPackageMap();
+
 protected:
 	bool					bShouldTrackUnmappedGuids;
 	TSet< FNetworkGUID >	TrackedUnmappedNetGuids;
@@ -410,6 +413,11 @@ public:
 
 public:
 	COREUOBJECT_API FNetBitReader(UPackageMap* InPackageMap=nullptr, const uint8* Src=nullptr, int64 CountBits=0);
+	COREUOBJECT_API FNetBitReader(const FNetBitReader&);
+    COREUOBJECT_API FNetBitReader& operator=(const FNetBitReader&);
+    COREUOBJECT_API FNetBitReader(FNetBitReader&&);
+    COREUOBJECT_API FNetBitReader& operator=(FNetBitReader&&);
+	COREUOBJECT_API virtual ~FNetBitReader();
 
 	COREUOBJECT_API virtual FArchive& operator<<(FName& Name) override;
 	COREUOBJECT_API virtual FArchive& operator<<(UObject*& Object) override;
diff --git a/Engine/Source/Runtime/CoreUObject/Public/UObject/UObjectGlobals.h b/Engine/Source/Runtime/CoreUObject/Public/UObject/UObjectGlobals.h
--- a/Engine/Source/Runtime/CoreUObject/Public/UObject/UObjectGlobals.h
+++ b/Engine/Source/Runtime/CoreUObject/Public/UObject/UObjectGlobals.h
@@ -2335,7 +2335,7 @@ public:
 class FReferenceCollector
 {
 public:
-	virtual ~FReferenceCollector() {}
+	COREUOBJECT_API virtual ~FReferenceCollector();
 	
 	/** Preferred way to add a reference that allows batching. Object must outlive GC tracing, can't be used for temporary/stack references. */
 	COREUOBJECT_API virtual void AddStableReference(TObjectPtr<UObject>* Object);
@@ -2991,6 +2991,7 @@ public:
 	 */
 	COREUOBJECT_API FReferenceFinder(TArray<UObject*>& InObjectArray, UObject* InOuter = nullptr, bool bInRequireDirectOuter = true, bool bInShouldIgnoreArchetype = false, bool bInSerializeRecursively = false, bool bInShouldIgnoreTransient = false);
 
+	COREUOBJECT_API ~FReferenceFinder();
 	/**
 	 * Finds all objects referenced by Object.
 	 *
diff --git a/Engine/Source/Runtime/Engine/Classes/Animation/AnimInstance.h b/Engine/Source/Runtime/Engine/Classes/Animation/AnimInstance.h
--- a/Engine/Source/Runtime/Engine/Classes/Animation/AnimInstance.h
+++ b/Engine/Source/Runtime/Engine/Classes/Animation/AnimInstance.h
@@ -325,10 +325,7 @@ class UAnimInstance : public UObject
 
 	typedef FAnimInstanceProxy ProxyType;
 
-	// Disable compiler-generated deprecation warnings by implementing our own destructor
-	PRAGMA_DISABLE_DEPRECATION_WARNINGS
-	~UAnimInstance() {}
-	PRAGMA_ENABLE_DEPRECATION_WARNINGS
+	ENGINE_API ~UAnimInstance();
 
 	/** This is used to extract animation. If Mesh exists, this will be overwritten by Mesh->Skeleton */
 	UPROPERTY(transient)
diff --git a/Engine/Source/Runtime/Engine/Classes/Animation/AnimNode_CustomProperty.h b/Engine/Source/Runtime/Engine/Classes/Animation/AnimNode_CustomProperty.h
--- a/Engine/Source/Runtime/Engine/Classes/Animation/AnimNode_CustomProperty.h
+++ b/Engine/Source/Runtime/Engine/Classes/Animation/AnimNode_CustomProperty.h
@@ -30,6 +30,7 @@ struct FAnimNode_CustomProperty : public FAnimNode_Base
 public:
 
 	ENGINE_API FAnimNode_CustomProperty();
+	ENGINE_API FAnimNode_CustomProperty(const FAnimNode_CustomProperty&);
 	ENGINE_API ~FAnimNode_CustomProperty();
 
 	/* Set Target Instance */
diff --git a/Engine/Source/Runtime/Engine/Classes/Engine/Attenuation.h b/Engine/Source/Runtime/Engine/Classes/Engine/Attenuation.h
--- a/Engine/Source/Runtime/Engine/Classes/Engine/Attenuation.h
+++ b/Engine/Source/Runtime/Engine/Classes/Engine/Attenuation.h
@@ -57,7 +57,7 @@ struct FBaseAttenuationSettings
 {
 	GENERATED_USTRUCT_BODY()
 
-	virtual ~FBaseAttenuationSettings() { }
+	ENGINE_API virtual ~FBaseAttenuationSettings();
 
 	/* The type of attenuation as a function of distance to use. */
 	UPROPERTY(EditAnywhere, BlueprintReadWrite, Category= AttenuationDistance, meta = (DisplayName = "Attenuation Function"))
diff --git a/Engine/Source/Runtime/Engine/Classes/PhysicsEngine/BoxElem.h b/Engine/Source/Runtime/Engine/Classes/PhysicsEngine/BoxElem.h
--- a/Engine/Source/Runtime/Engine/Classes/PhysicsEngine/BoxElem.h
+++ b/Engine/Source/Runtime/Engine/Classes/PhysicsEngine/BoxElem.h
@@ -83,6 +83,8 @@ struct FKBoxElem : public FKShapeElem
 
 	}
 
+	ENGINE_API virtual ~FKBoxElem();
+
 #if WITH_EDITORONLY_DATA
 	void FixupDeprecated( FArchive& Ar );
 #endif
diff --git a/Engine/Source/Runtime/Engine/Classes/PhysicsEngine/ShapeElem.h b/Engine/Source/Runtime/Engine/Classes/PhysicsEngine/ShapeElem.h
--- a/Engine/Source/Runtime/Engine/Classes/PhysicsEngine/ShapeElem.h
+++ b/Engine/Source/Runtime/Engine/Classes/PhysicsEngine/ShapeElem.h
@@ -67,7 +67,7 @@ struct FKShapeElem
 	{
 	}
 
-	virtual ~FKShapeElem(){}
+	ENGINE_API virtual ~FKShapeElem();
 
 	const FKShapeElem& operator=(const FKShapeElem& Other)
 	{
diff --git a/Engine/Source/Runtime/Engine/Classes/PhysicsEngine/SphereElem.h b/Engine/Source/Runtime/Engine/Classes/PhysicsEngine/SphereElem.h
--- a/Engine/Source/Runtime/Engine/Classes/PhysicsEngine/SphereElem.h
+++ b/Engine/Source/Runtime/Engine/Classes/PhysicsEngine/SphereElem.h
@@ -50,6 +50,8 @@ struct FKSphereElem : public FKShapeElem
 
 	}
 
+	ENGINE_API virtual ~FKSphereElem();
+
 #if WITH_EDITORONLY_DATA
 	void FixupDeprecated( FArchive& Ar );
 #endif
diff --git a/Engine/Source/Runtime/Engine/Classes/PhysicsEngine/SphylElem.h b/Engine/Source/Runtime/Engine/Classes/PhysicsEngine/SphylElem.h
--- a/Engine/Source/Runtime/Engine/Classes/PhysicsEngine/SphylElem.h
+++ b/Engine/Source/Runtime/Engine/Classes/PhysicsEngine/SphylElem.h
@@ -65,6 +65,8 @@ struct FKSphylElem : public FKShapeElem
 
 	}
 
+	ENGINE_API virtual ~FKSphylElem();
+
 #if WITH_EDITORONLY_DATA
 	void FixupDeprecated( FArchive& Ar );
 #endif
diff --git a/Engine/Source/Runtime/Engine/Classes/PhysicsEngine/TaperedCapsuleElem.h b/Engine/Source/Runtime/Engine/Classes/PhysicsEngine/TaperedCapsuleElem.h
--- a/Engine/Source/Runtime/Engine/Classes/PhysicsEngine/TaperedCapsuleElem.h
+++ b/Engine/Source/Runtime/Engine/Classes/PhysicsEngine/TaperedCapsuleElem.h
@@ -36,25 +36,10 @@ struct FKTaperedCapsuleElem : public FKShapeElem
 	UPROPERTY(Category= Capsule, EditAnywhere)
 	float Length;
 
-	FKTaperedCapsuleElem()
-	: FKShapeElem(EAggCollisionShape::TaperedCapsule)
-	, Center(FVector::ZeroVector)
-	, Rotation(FRotator::ZeroRotator)
-	, Radius0(1.0f)
-	, Radius1(1.0f)
-	, Length(1.0f)
-	{
-	}
-
-	FKTaperedCapsuleElem( float InRadius0, float InRadius1, float InLength )
-	: FKShapeElem(EAggCollisionShape::TaperedCapsule)
-	, Center(FVector::ZeroVector)
-	, Rotation(FRotator::ZeroRotator)
-	, Radius0(InRadius0)
-	, Radius1(InRadius1)
-	, Length(InLength)
-	{
-	}
+	ENGINE_API FKTaperedCapsuleElem();
+	ENGINE_API FKTaperedCapsuleElem( float InRadius0, float InRadius1, float InLength );
+	ENGINE_API FKTaperedCapsuleElem(const FKTaperedCapsuleElem&);
+	ENGINE_API ~FKTaperedCapsuleElem();
 
 	friend bool operator==( const FKTaperedCapsuleElem& LHS, const FKTaperedCapsuleElem& RHS )
 	{
diff --git a/Engine/Source/Runtime/Engine/Classes/Sound/QuartzQuantizationUtilities.cpp b/Engine/Source/Runtime/Engine/Classes/Sound/QuartzQuantizationUtilities.cpp
--- a/Engine/Source/Runtime/Engine/Classes/Sound/QuartzQuantizationUtilities.cpp
+++ b/Engine/Source/Runtime/Engine/Classes/Sound/QuartzQuantizationUtilities.cpp
@@ -471,6 +471,9 @@ namespace Audio
 	{
 	}
 
+	IQuartzQuantizedCommand::IQuartzQuantizedCommand() = default;
+	IQuartzQuantizedCommand::~IQuartzQuantizedCommand() = default;
+
 	TSharedPtr<IQuartzQuantizedCommand> IQuartzQuantizedCommand::GetDeepCopyOfDerivedObject() const
 	{
 		// implement this method to allow copies to be made from pointers to base class
diff --git a/Engine/Source/Runtime/Engine/Classes/Sound/QuartzQuantizationUtilities.h b/Engine/Source/Runtime/Engine/Classes/Sound/QuartzQuantizationUtilities.h
--- a/Engine/Source/Runtime/Engine/Classes/Sound/QuartzQuantizationUtilities.h
+++ b/Engine/Source/Runtime/Engine/Classes/Sound/QuartzQuantizationUtilities.h
@@ -589,10 +589,10 @@ namespace Audio
 	public:
 
 		// ctor
-		IQuartzQuantizedCommand() {};
+		ENGINE_API IQuartzQuantizedCommand();
 
 		// dtor
-		virtual ~IQuartzQuantizedCommand() {};
+		ENGINE_API virtual ~IQuartzQuantizedCommand();
 
 		// allocate a copy of the derived class
 		ENGINE_API virtual TSharedPtr<IQuartzQuantizedCommand> GetDeepCopyOfDerivedObject() const;
diff --git a/Engine/Source/Runtime/Engine/Classes/Sound/SoundAttenuation.h b/Engine/Source/Runtime/Engine/Classes/Sound/SoundAttenuation.h
--- a/Engine/Source/Runtime/Engine/Classes/Sound/SoundAttenuation.h
+++ b/Engine/Source/Runtime/Engine/Classes/Sound/SoundAttenuation.h
@@ -394,83 +394,8 @@ struct FSoundAttenuationSettings : public FBaseAttenuationSettings
 	UPROPERTY(EditAnywhere, BlueprintReadWrite, Category = AttenuationPluginSettings, meta = (ShowOnlyInnerProperties))
 	FSoundAttenuationPluginSettings PluginSettings;
 
-	FSoundAttenuationSettings()
-		: bAttenuate(true)
-		, bSpatialize(true)
-		, bAttenuateWithLPF(false)
-		, bEnableListenerFocus(false)
-		, bEnableFocusInterpolation(false)
-		, bEnableOcclusion(false)
-		, bUseComplexCollisionForOcclusion(false)
-		, bEnableReverbSend(true)
-		, bEnablePriorityAttenuation(false)
-		, bApplyNormalizationToStereoSounds(false)
-		, bEnableLogFrequencyScaling(false)
-		, bEnableSubmixSends(false)
-		, bEnableSourceDataOverride(false)
-		, bEnableSendToAudioLink(true)
-		, SpatializationAlgorithm(ESoundSpatializationAlgorithm::SPATIALIZATION_Default)
-		, AudioLinkSettingsOverride(nullptr)
-		, BinauralRadius(0.0f)
-		, AbsorptionMethod(EAirAbsorptionMethod::Linear)
-		, OcclusionTraceChannel(ECC_Visibility)
-		, ReverbSendMethod(EReverbSendMethod::Linear)
-		, PriorityAttenuationMethod(EPriorityAttenuationMethod::Linear)
-#if WITH_EDITORONLY_DATA
-		, DistanceType_DEPRECATED(SOUNDDISTANCE_Normal)
-		, OmniRadius_DEPRECATED(0.0f)
-#endif
-		, NonSpatializedRadiusStart(0.0f)
-		, NonSpatializedRadiusEnd(0.0f)
-		, NonSpatializedRadiusMode(ENonSpatializedRadiusSpeakerMapMode::OmniDirectional)
-		, StereoSpread(200.0f)
-#if WITH_EDITORONLY_DATA
-		, SpatializationPluginSettings_DEPRECATED(nullptr)
-		, RadiusMin_DEPRECATED(400.f)
-		, RadiusMax_DEPRECATED(4000.f)
-#endif
-		, LPFRadiusMin(3000.f)
-		, LPFRadiusMax(6000.f)
-		, LPFFrequencyAtMin(20000.f)
-		, LPFFrequencyAtMax(20000.f)
-		, HPFFrequencyAtMin(0.0f)
-		, HPFFrequencyAtMax(0.0f)
-		, FocusAzimuth(30.0f)
-		, NonFocusAzimuth(60.0f)
-		, FocusDistanceScale(1.0f)
-		, NonFocusDistanceScale(1.0f)
-		, FocusPriorityScale(1.0f)
-		, NonFocusPriorityScale(1.0f)
-		, FocusVolumeAttenuation(1.0f)
-		, NonFocusVolumeAttenuation(1.0f)
-		, FocusAttackInterpSpeed(1.0f)
-		, FocusReleaseInterpSpeed(1.0f)
-		, OcclusionLowPassFilterFrequency(20000.f)
-		, OcclusionVolumeAttenuation(1.0f)
-		, OcclusionInterpolationTime(0.1f)
-#if WITH_EDITORONLY_DATA
-		, OcclusionPluginSettings_DEPRECATED(nullptr)
-		, ReverbPluginSettings_DEPRECATED(nullptr)
-#endif
-		, ReverbWetLevelMin(0.3f)
-		, ReverbWetLevelMax(0.95f)
-		, ReverbDistanceMin(UE_REAL_TO_FLOAT(AttenuationShapeExtents.X))
-		, ReverbDistanceMax(UE_REAL_TO_FLOAT(AttenuationShapeExtents.X) + FalloffDistance)
-		, ManualReverbSendLevel(0.0f)
-		, PriorityAttenuationMin(1.0f)
-		, PriorityAttenuationMax(1.0f)
-		, PriorityAttenuationDistanceMin(UE_REAL_TO_FLOAT(AttenuationShapeExtents.X))
-		, PriorityAttenuationDistanceMax(UE_REAL_TO_FLOAT(AttenuationShapeExtents.X) + FalloffDistance)
-		, ManualPriorityAttenuation(1.0f)
-	{
-#if WITH_EDITOR
-		if (const USoundAttenuationEditorSettings* SoundAttenuationEditorSettings = GetDefault<USoundAttenuationEditorSettings>())
-		{
-			bEnableReverbSend = SoundAttenuationEditorSettings->bEnableReverbSend;
-			bEnableSendToAudioLink = SoundAttenuationEditorSettings->bEnableSendToAudioLink;
-		}
-#endif // WITH_EDITOR
-	}
+	ENGINE_API FSoundAttenuationSettings();
+	ENGINE_API ~FSoundAttenuationSettings();
 
 	ENGINE_API bool operator==(const FSoundAttenuationSettings& Other) const;
 #if WITH_EDITORONLY_DATA
diff --git a/Engine/Source/Runtime/Engine/Classes/Sound/SoundWave.h b/Engine/Source/Runtime/Engine/Classes/Sound/SoundWave.h
--- a/Engine/Source/Runtime/Engine/Classes/Sound/SoundWave.h
+++ b/Engine/Source/Runtime/Engine/Classes/Sound/SoundWave.h
@@ -1701,7 +1701,7 @@ public:
 
 	ENGINE_API explicit FSoundWaveProxy(USoundWave* InWave);
 
-	FSoundWaveProxy(const FSoundWaveProxy& Other) = default;
+	ENGINE_API FSoundWaveProxy(const FSoundWaveProxy& Other);
 
 	ENGINE_API ~FSoundWaveProxy();
 
diff --git a/Engine/Source/Runtime/Engine/Private/Animation/AnimInstance.cpp b/Engine/Source/Runtime/Engine/Private/Animation/AnimInstance.cpp
--- a/Engine/Source/Runtime/Engine/Private/Animation/AnimInstance.cpp
+++ b/Engine/Source/Runtime/Engine/Private/Animation/AnimInstance.cpp
@@ -142,6 +142,11 @@ UAnimInstance::UAnimInstance(const FObjectInitializer& ObjectInitializer)
 #endif // WITH_EDITOR	
 }
 
+// Disable compiler-generated deprecation warnings by implementing our own destructor
+PRAGMA_DISABLE_DEPRECATION_WARNINGS
+UAnimInstance::~UAnimInstance() = default;
+PRAGMA_ENABLE_DEPRECATION_WARNINGS
+
 // this is only used by montage marker based sync
 void UAnimInstance::MakeMontageTickRecord(FAnimTickRecord& TickRecord, class UAnimMontage* Montage, float CurrentPosition, float Weight, TArray<FPassedMarker>& MarkersPassedThisTick, FMarkerTickRecord& MarkerTickRecord)
 {
diff --git a/Engine/Source/Runtime/Engine/Private/Animation/AnimNode_CustomProperty.cpp b/Engine/Source/Runtime/Engine/Private/Animation/AnimNode_CustomProperty.cpp
--- a/Engine/Source/Runtime/Engine/Private/Animation/AnimNode_CustomProperty.cpp
+++ b/Engine/Source/Runtime/Engine/Private/Animation/AnimNode_CustomProperty.cpp
@@ -12,9 +12,8 @@ FAnimNode_CustomProperty::FAnimNode_CustomProperty()
 {
 }
 
-FAnimNode_CustomProperty::~FAnimNode_CustomProperty()
-{
-}
+FAnimNode_CustomProperty::FAnimNode_CustomProperty(const FAnimNode_CustomProperty&) = default;
+FAnimNode_CustomProperty::~FAnimNode_CustomProperty() = default;
 
 void FAnimNode_CustomProperty::SetTargetInstance(UObject* InInstance)
 {
diff --git a/Engine/Source/Runtime/Engine/Private/Attenuation.cpp b/Engine/Source/Runtime/Engine/Private/Attenuation.cpp
--- a/Engine/Source/Runtime/Engine/Private/Attenuation.cpp
+++ b/Engine/Source/Runtime/Engine/Private/Attenuation.cpp
@@ -27,6 +27,8 @@ FBaseAttenuationSettings::FBaseAttenuationSettings()
 {
 }
 
+FBaseAttenuationSettings::~FBaseAttenuationSettings() = default;
+
 float FBaseAttenuationSettings::GetMaxDimension() const
 {
 	float MaxDimension = GetMaxFalloffDistance();
diff --git a/Engine/Source/Runtime/Engine/Private/Audio.cpp b/Engine/Source/Runtime/Engine/Private/Audio.cpp
--- a/Engine/Source/Runtime/Engine/Private/Audio.cpp
+++ b/Engine/Source/Runtime/Engine/Private/Audio.cpp
@@ -276,6 +276,8 @@ FString FSoundBuffer::Describe(bool bUseLongName)
 	FSoundSource implementation.
 -----------------------------------------------------------------------------*/
 
+FSoundSource::~FSoundSource() = default;
+
 FString FSoundSource::Describe(bool bUseLongName)
 {
 	return FString::Printf(TEXT("Wave: %s, Volume: %6.2f, Owner: %s"),
diff --git a/Engine/Source/Runtime/Engine/Private/AudioDecompress.cpp b/Engine/Source/Runtime/Engine/Private/AudioDecompress.cpp
--- a/Engine/Source/Runtime/Engine/Private/AudioDecompress.cpp
+++ b/Engine/Source/Runtime/Engine/Private/AudioDecompress.cpp
@@ -67,6 +67,8 @@ IStreamedCompressedInfo::IStreamedCompressedInfo()
 	StartTimeInCycles = FPlatformTime::Cycles64();
 }
 
+IStreamedCompressedInfo::~IStreamedCompressedInfo() = default;
+
 uint32 IStreamedCompressedInfo::Read(void *OutBuffer, uint32 DataSize)
 {
 	uint32 BytesToRead = FMath::Min(DataSize, SrcBufferDataSize - SrcBufferOffset);
@@ -277,6 +279,12 @@ bool IStreamedCompressedInfo::StreamCompressedInfoInternal(const FSoundWaveProxy
 	return false;
 }
 
+ICompressedAudioInfo::ICompressedAudioInfo()
+	: StreamingSoundWave(nullptr)
+{}
+
+ICompressedAudioInfo::~ICompressedAudioInfo() = default;
+
 bool ICompressedAudioInfo::HasError() const
 {
 	return bHasError || DECODER_MATCHES_WAVE(AudioDecompressPrivate::ForceDecoderErrorOnWaveCVar);
diff --git a/Engine/Source/Runtime/Engine/Private/Collision/KAggregateGeom.cpp b/Engine/Source/Runtime/Engine/Private/Collision/KAggregateGeom.cpp
--- a/Engine/Source/Runtime/Engine/Private/Collision/KAggregateGeom.cpp
+++ b/Engine/Source/Runtime/Engine/Private/Collision/KAggregateGeom.cpp
@@ -310,6 +310,8 @@ static bool EnsureHullIsValid(TArray<FVector>& InVerts)
 ///////////// FKShapeElem ////////////
 ///////////////////////////////////////
 
+FKShapeElem::~FKShapeElem() = default;
+
 EAggCollisionShape::Type FKShapeElem::StaticShapeType = EAggCollisionShape::Unknown;
 
 
@@ -317,6 +319,8 @@ EAggCollisionShape::Type FKShapeElem::StaticShapeType = EAggCollisionShape::Unkn
 ///////////// FKSphereElem ////////////
 ///////////////////////////////////////
 
+FKSphereElem::~FKSphereElem() = default;
+
 EAggCollisionShape::Type FKSphereElem::StaticShapeType = EAggCollisionShape::Sphere;
 
 FBox FKSphereElem::CalcAABB(const FTransform& BoneTM, float Scale) const
@@ -336,6 +340,8 @@ FBox FKSphereElem::CalcAABB(const FTransform& BoneTM, float Scale) const
 ////////////// FKBoxElem //////////////
 ///////////////////////////////////////
 
+FKBoxElem::~FKBoxElem() = default;
+
 EAggCollisionShape::Type FKBoxElem::StaticShapeType = EAggCollisionShape::Box;
 
 FBox FKBoxElem::CalcAABB(const FTransform& BoneTM, float Scale) const
@@ -375,6 +381,8 @@ FBox FKBoxElem::CalcAABB(const FTransform& BoneTM, float Scale) const
 ////////////// FKSphylElem ////////////
 ///////////////////////////////////////
 
+FKSphylElem::~FKSphylElem() = default;
+
 EAggCollisionShape::Type FKSphylElem::StaticShapeType = EAggCollisionShape::Sphyl;
 
 FBox FKSphylElem::CalcAABB(const FTransform& BoneTM, float Scale) const
@@ -461,6 +469,29 @@ void FKConvexElem::GetPlanes(TArray<FPlane>& Planes) const
 //////// FKTaperedCapsuleElem /////////
 ///////////////////////////////////////
 
+FKTaperedCapsuleElem::FKTaperedCapsuleElem()
+: FKShapeElem(EAggCollisionShape::TaperedCapsule)
+, Center(FVector::ZeroVector)
+, Rotation(FRotator::ZeroRotator)
+, Radius0(1.0f)
+, Radius1(1.0f)
+, Length(1.0f)
+{
+}
+
+FKTaperedCapsuleElem::FKTaperedCapsuleElem( float InRadius0, float InRadius1, float InLength )
+: FKShapeElem(EAggCollisionShape::TaperedCapsule)
+, Center(FVector::ZeroVector)
+, Rotation(FRotator::ZeroRotator)
+, Radius0(InRadius0)
+, Radius1(InRadius1)
+, Length(InLength)
+{
+}
+
+FKTaperedCapsuleElem::FKTaperedCapsuleElem(const FKTaperedCapsuleElem&) = default;
+FKTaperedCapsuleElem::~FKTaperedCapsuleElem() = default;
+
 EAggCollisionShape::Type FKTaperedCapsuleElem::StaticShapeType = EAggCollisionShape::TaperedCapsule;
 
 FBox FKTaperedCapsuleElem::CalcAABB(const FTransform& BoneTM, float Scale) const
diff --git a/Engine/Source/Runtime/Engine/Private/DebugRenderSceneProxy.cpp b/Engine/Source/Runtime/Engine/Private/DebugRenderSceneProxy.cpp
--- a/Engine/Source/Runtime/Engine/Private/DebugRenderSceneProxy.cpp
+++ b/Engine/Source/Runtime/Engine/Private/DebugRenderSceneProxy.cpp
@@ -27,6 +27,10 @@ FDebugRenderSceneProxy::FDebugRenderSceneProxy(const UPrimitiveComponent* InComp
 {
 }
 
+FDebugRenderSceneProxy::FDebugRenderSceneProxy(const FDebugRenderSceneProxy&) = default;
+
+FDebugRenderSceneProxy::~FDebugRenderSceneProxy() = default;
+
 void FDebugDrawDelegateHelper::RegisterDebugDrawDelegateInternal()
 {
 	// note that it's possible at this point for State == RegisteredState since RegisterDebugDrawDelegateInternal can get 
diff --git a/Engine/Source/Runtime/Engine/Private/PrimitiveSceneProxy.cpp b/Engine/Source/Runtime/Engine/Private/PrimitiveSceneProxy.cpp
--- a/Engine/Source/Runtime/Engine/Private/PrimitiveSceneProxy.cpp
+++ b/Engine/Source/Runtime/Engine/Private/PrimitiveSceneProxy.cpp
@@ -264,6 +264,8 @@ FPrimitiveSceneProxyDesc::FPrimitiveSceneProxyDesc(const UPrimitiveComponent* In
 	InitializeFrom(InComponent);
 }
 
+FPrimitiveSceneProxy::FPrimitiveSceneProxy(FPrimitiveSceneProxy const&) = default;
+
 void FPrimitiveSceneProxyDesc::InitializeFrom(const UPrimitiveComponent* InComponent)
 {
 	CastShadow = InComponent->CastShadow;
diff --git a/Engine/Source/Runtime/Engine/Private/RawIndexBuffer.cpp b/Engine/Source/Runtime/Engine/Private/RawIndexBuffer.cpp
--- a/Engine/Source/Runtime/Engine/Private/RawIndexBuffer.cpp
+++ b/Engine/Source/Runtime/Engine/Private/RawIndexBuffer.cpp
@@ -153,6 +153,8 @@ FRawStaticIndexBuffer::FRawStaticIndexBuffer(bool InNeedsCPUAccess)
 {
 }
 
+FRawStaticIndexBuffer::~FRawStaticIndexBuffer() = default;
+
 void FRawStaticIndexBuffer::SetIndices(const TArray<uint32>& InIndices, EIndexBufferStride::Type DesiredStride)
 {
 	int32 NumIndices = InIndices.Num();
diff --git a/Engine/Source/Runtime/Engine/Private/SceneView.cpp b/Engine/Source/Runtime/Engine/Private/SceneView.cpp
--- a/Engine/Source/Runtime/Engine/Private/SceneView.cpp
+++ b/Engine/Source/Runtime/Engine/Private/SceneView.cpp
@@ -3134,6 +3134,19 @@ FSceneViewFamily::~FSceneViewFamily()
 }
 PRAGMA_ENABLE_DEPRECATION_WARNINGS
 
+FSceneViewFamily::FSceneViewFamily(FSceneViewFamily&& InViewFamily)
+	: FSceneViewFamily(static_cast<const FSceneViewFamily&>(InViewFamily))
+{
+	check(ScreenPercentageInterface == nullptr);
+	check(TemporalUpscalerInterface == nullptr);
+	check(PrimarySpatialUpscalerInterface == nullptr);
+	check(SecondarySpatialUpscalerInterface == nullptr);
+}
+
+PRAGMA_DISABLE_DEPRECATION_WARNINGS
+FSceneViewFamily::FSceneViewFamily(const FSceneViewFamily&) = default;
+PRAGMA_ENABLE_DEPRECATION_WARNINGS
+
 ERHIFeatureLevel::Type FSceneViewFamily::GetFeatureLevel() const
 {
 	if (Scene)
diff --git a/Engine/Source/Runtime/Engine/Private/SoundAttenuation.cpp b/Engine/Source/Runtime/Engine/Private/SoundAttenuation.cpp
--- a/Engine/Source/Runtime/Engine/Private/SoundAttenuation.cpp
+++ b/Engine/Source/Runtime/Engine/Private/SoundAttenuation.cpp
@@ -10,6 +10,86 @@
 	USoundAttenuation implementation.
 -----------------------------------------------------------------------------*/
 
+FSoundAttenuationSettings::FSoundAttenuationSettings()
+	: bAttenuate(true)
+	, bSpatialize(true)
+	, bAttenuateWithLPF(false)
+	, bEnableListenerFocus(false)
+	, bEnableFocusInterpolation(false)
+	, bEnableOcclusion(false)
+	, bUseComplexCollisionForOcclusion(false)
+	, bEnableReverbSend(true)
+	, bEnablePriorityAttenuation(false)
+	, bApplyNormalizationToStereoSounds(false)
+	, bEnableLogFrequencyScaling(false)
+	, bEnableSubmixSends(false)
+	, bEnableSourceDataOverride(false)
+	, bEnableSendToAudioLink(true)
+	, SpatializationAlgorithm(ESoundSpatializationAlgorithm::SPATIALIZATION_Default)
+	, AudioLinkSettingsOverride(nullptr)
+	, BinauralRadius(0.0f)
+	, AbsorptionMethod(EAirAbsorptionMethod::Linear)
+	, OcclusionTraceChannel(ECC_Visibility)
+	, ReverbSendMethod(EReverbSendMethod::Linear)
+	, PriorityAttenuationMethod(EPriorityAttenuationMethod::Linear)
+#if WITH_EDITORONLY_DATA
+	, DistanceType_DEPRECATED(SOUNDDISTANCE_Normal)
+	, OmniRadius_DEPRECATED(0.0f)
+#endif
+	, NonSpatializedRadiusStart(0.0f)
+	, NonSpatializedRadiusEnd(0.0f)
+	, NonSpatializedRadiusMode(ENonSpatializedRadiusSpeakerMapMode::OmniDirectional)
+	, StereoSpread(200.0f)
+#if WITH_EDITORONLY_DATA
+	, SpatializationPluginSettings_DEPRECATED(nullptr)
+	, RadiusMin_DEPRECATED(400.f)
+	, RadiusMax_DEPRECATED(4000.f)
+#endif
+	, LPFRadiusMin(3000.f)
+	, LPFRadiusMax(6000.f)
+	, LPFFrequencyAtMin(20000.f)
+	, LPFFrequencyAtMax(20000.f)
+	, HPFFrequencyAtMin(0.0f)
+	, HPFFrequencyAtMax(0.0f)
+	, FocusAzimuth(30.0f)
+	, NonFocusAzimuth(60.0f)
+	, FocusDistanceScale(1.0f)
+	, NonFocusDistanceScale(1.0f)
+	, FocusPriorityScale(1.0f)
+	, NonFocusPriorityScale(1.0f)
+	, FocusVolumeAttenuation(1.0f)
+	, NonFocusVolumeAttenuation(1.0f)
+	, FocusAttackInterpSpeed(1.0f)
+	, FocusReleaseInterpSpeed(1.0f)
+	, OcclusionLowPassFilterFrequency(20000.f)
+	, OcclusionVolumeAttenuation(1.0f)
+	, OcclusionInterpolationTime(0.1f)
+#if WITH_EDITORONLY_DATA
+	, OcclusionPluginSettings_DEPRECATED(nullptr)
+	, ReverbPluginSettings_DEPRECATED(nullptr)
+#endif
+	, ReverbWetLevelMin(0.3f)
+	, ReverbWetLevelMax(0.95f)
+	, ReverbDistanceMin(UE_REAL_TO_FLOAT(AttenuationShapeExtents.X))
+	, ReverbDistanceMax(UE_REAL_TO_FLOAT(AttenuationShapeExtents.X) + FalloffDistance)
+	, ManualReverbSendLevel(0.0f)
+	, PriorityAttenuationMin(1.0f)
+	, PriorityAttenuationMax(1.0f)
+	, PriorityAttenuationDistanceMin(UE_REAL_TO_FLOAT(AttenuationShapeExtents.X))
+	, PriorityAttenuationDistanceMax(UE_REAL_TO_FLOAT(AttenuationShapeExtents.X) + FalloffDistance)
+	, ManualPriorityAttenuation(1.0f)
+{
+#if WITH_EDITOR
+	if (const USoundAttenuationEditorSettings* SoundAttenuationEditorSettings = GetDefault<USoundAttenuationEditorSettings>())
+	{
+		bEnableReverbSend = SoundAttenuationEditorSettings->bEnableReverbSend;
+		bEnableSendToAudioLink = SoundAttenuationEditorSettings->bEnableSendToAudioLink;
+	}
+#endif // WITH_EDITOR
+}
+
+FSoundAttenuationSettings::~FSoundAttenuationSettings() = default;
+
 #if WITH_EDITORONLY_DATA
 void FSoundAttenuationSettings::PostSerialize(const FArchive& Ar)
 {
diff --git a/Engine/Source/Runtime/Engine/Private/SoundWave.cpp b/Engine/Source/Runtime/Engine/Private/SoundWave.cpp
--- a/Engine/Source/Runtime/Engine/Private/SoundWave.cpp
+++ b/Engine/Source/Runtime/Engine/Private/SoundWave.cpp
@@ -4626,6 +4626,8 @@ FSoundWaveProxy::FSoundWaveProxy(USoundWave* InWave)
 	TypeHash = ::GetTypeHash(InWave);
 }
 
+FSoundWaveProxy::FSoundWaveProxy(const FSoundWaveProxy& Other) = default;
+
 FSoundWaveProxy::~FSoundWaveProxy()
 {
 }
diff --git a/Engine/Source/Runtime/Engine/Private/UnrealClient.cpp b/Engine/Source/Runtime/Engine/Private/UnrealClient.cpp
--- a/Engine/Source/Runtime/Engine/Private/UnrealClient.cpp
+++ b/Engine/Source/Runtime/Engine/Private/UnrealClient.cpp
@@ -172,6 +172,9 @@ float FRenderTarget::GetEngineDisplayGamma() // static
 	}
 }
 
+FRenderTarget::FRenderTarget() = default;
+FRenderTarget::~FRenderTarget() = default;
+
 const FTextureRHIRef& FRenderTarget::GetRenderTargetTexture() const
 {
 	return RenderTargetTextureRHI;
diff --git a/Engine/Source/Runtime/Engine/Private/UserInterface/CanvasItem.cpp b/Engine/Source/Runtime/Engine/Private/UserInterface/CanvasItem.cpp
--- a/Engine/Source/Runtime/Engine/Private/UserInterface/CanvasItem.cpp
+++ b/Engine/Source/Runtime/Engine/Private/UserInterface/CanvasItem.cpp
@@ -831,6 +831,26 @@ void FCanvasBorderItem::Draw( class FCanvas* InCanvas )
 
 }
 
+FCanvasTextItemBase::FCanvasTextItemBase( const FVector2D& InPosition, const FLinearColor& InColor )
+	: FCanvasItem( InPosition )
+	, HorizSpacingAdjust( 0.0f )
+	, Depth( 1.0f )
+	, ShadowColor( FLinearColor::Black )
+	, ShadowOffset( FVector2D::ZeroVector )
+	, DrawnSize( FVector2D::ZeroVector )
+	, bCentreX( false )
+	, bCentreY( false )
+	, bOutlined( false )
+	, OutlineColor( FLinearColor::Black )
+	, bDontCorrectStereoscopic( true )
+	, TileItem( InPosition, FVector2D::ZeroVector, InColor )
+{
+	SetColor( InColor );
+	Scale.Set( 1.0f, 1.0f );
+	BlendMode = SE_BLEND_Translucent;
+}
+
+FCanvasTextItemBase::~FCanvasTextItemBase() = default;
 
 void FCanvasTextItemBase::Draw( class FCanvas* InCanvas )
 {	
@@ -1338,6 +1358,20 @@ void FCanvasSimpleTextItem::DrawStringInternal_RuntimeCache(FCanvas* InCanvas, F
 	}
 }
 
+FCanvasTextItem::FCanvasTextItem(const FVector2D& InPosition, const FText& InText, const UFont* InFont, const FLinearColor& InColor)
+	: FCanvasSimpleTextItem(InPosition, InFont, InColor)
+	, Text(InText)
+{
+}
+
+FCanvasTextItem::FCanvasTextItem(const FVector2D& InPosition, const FText& InText, const FSlateFontInfo& InFontInfo, const FLinearColor& InColor)
+	: FCanvasSimpleTextItem(InPosition, InFontInfo, InColor)
+	, Text(InText)
+{
+}
+
+FCanvasTextItem::~FCanvasTextItem() = default;
+
 bool FCanvasTextItem::HasValidText() const
 {
 	return FCanvasSimpleTextItem::HasValidText() && !Text.IsEmpty();
diff --git a/Engine/Source/Runtime/Engine/Private/WeightedRandomSampler.cpp b/Engine/Source/Runtime/Engine/Private/WeightedRandomSampler.cpp
--- a/Engine/Source/Runtime/Engine/Private/WeightedRandomSampler.cpp
+++ b/Engine/Source/Runtime/Engine/Private/WeightedRandomSampler.cpp
@@ -11,6 +11,8 @@ FWeightedRandomSampler::FWeightedRandomSampler()
 
 }
 
+FWeightedRandomSampler::~FWeightedRandomSampler() = default;
+
 void FWeightedRandomSampler::Initialize()
 {
 	//Gather all weights.
diff --git a/Engine/Source/Runtime/Engine/Public/Audio.h b/Engine/Source/Runtime/Engine/Public/Audio.h
--- a/Engine/Source/Runtime/Engine/Public/Audio.h
+++ b/Engine/Source/Runtime/Engine/Public/Audio.h
@@ -638,7 +638,7 @@ public:
 	}
 
 	/** Destructor */
-	virtual ~FSoundSource() {}
+	ENGINE_API virtual ~FSoundSource();
 
 	/* Prepares the source voice for initialization. This may parse a compressed asset header on some platforms */
 	virtual bool PrepareForInitialization(FWaveInstance* InWaveInstance) { return true; }
diff --git a/Engine/Source/Runtime/Engine/Public/AudioDecompress.h b/Engine/Source/Runtime/Engine/Public/AudioDecompress.h
--- a/Engine/Source/Runtime/Engine/Public/AudioDecompress.h
+++ b/Engine/Source/Runtime/Engine/Public/AudioDecompress.h
@@ -30,14 +30,12 @@ class FStreamedAudioChunkSeekTable;
 class ICompressedAudioInfo
 {
 public:
-	ICompressedAudioInfo()
-		: StreamingSoundWave(nullptr)
-	{}
+	ENGINE_API ICompressedAudioInfo();
 
 	/**
 	* Virtual destructor.
 	*/
-	virtual ~ICompressedAudioInfo() { }
+	ENGINE_API virtual ~ICompressedAudioInfo();
 
 	/**
 	* Reads the header information of a compressed format
@@ -200,7 +198,7 @@ class IStreamedCompressedInfo : public ICompressedAudioInfo
 {
 public:
 	ENGINE_API IStreamedCompressedInfo();
-	virtual ~IStreamedCompressedInfo() {}
+	ENGINE_API virtual ~IStreamedCompressedInfo();
 
 	//~ Begin ICompressedInfo Interface
 	ENGINE_API virtual bool ReadCompressedInfo(const uint8* InSrcBufferData, uint32 InSrcBufferDataSize, FSoundQualityInfo* QualityInfo) override;
diff --git a/Engine/Source/Runtime/Engine/Public/CanvasItem.h b/Engine/Source/Runtime/Engine/Public/CanvasItem.h
--- a/Engine/Source/Runtime/Engine/Public/CanvasItem.h
+++ b/Engine/Source/Runtime/Engine/Public/CanvasItem.h
@@ -305,24 +305,8 @@ public:
 class FCanvasTextItemBase : public FCanvasItem
 {
 public:
-	FCanvasTextItemBase( const FVector2D& InPosition, const FLinearColor& InColor )
-		: FCanvasItem( InPosition )
-		, HorizSpacingAdjust( 0.0f )
-		, Depth( 1.0f )
-		, ShadowColor( FLinearColor::Black )
-		, ShadowOffset( FVector2D::ZeroVector )
-		, DrawnSize( FVector2D::ZeroVector )
-		, bCentreX( false )
-		, bCentreY( false )
-		, bOutlined( false )
-		, OutlineColor( FLinearColor::Black )
-		, bDontCorrectStereoscopic( true )
-		, TileItem( InPosition, FVector2D::ZeroVector, InColor )
-	{
-		SetColor( InColor );
-		Scale.Set( 1.0f, 1.0f );
-		BlendMode = SE_BLEND_Translucent;
-	}
+	ENGINE_API FCanvasTextItemBase( const FVector2D& InPosition, const FLinearColor& InColor );
+	ENGINE_API virtual ~FCanvasTextItemBase();
 
 	/** 
 	 * Set the shadow offset and color. 
@@ -519,11 +503,7 @@ public:
 	 * @param	InFont			Font to draw with.
 	 * @param	InColor			Color to draw the text with.
 	 */
-	FCanvasTextItem(const FVector2D& InPosition, const FText& InText, const UFont* InFont, const FLinearColor& InColor)
-		: FCanvasSimpleTextItem(InPosition, InFont, InColor)
-		, Text(InText)
-	{
-	}
+	ENGINE_API FCanvasTextItem(const FVector2D& InPosition, const FText& InText, const UFont* InFont, const FLinearColor& InColor);
 
 	/** 	 
 	 * Constructor
@@ -533,13 +513,9 @@ public:
 	 * @param	InFontInfo		Font info to draw with.
 	 * @param	InColor			Color to draw the text with.
 	 */
-	FCanvasTextItem(const FVector2D& InPosition, const FText& InText, const FSlateFontInfo& InFontInfo, const FLinearColor& InColor)
-		: FCanvasSimpleTextItem(InPosition, InFontInfo, InColor)
-		, Text(InText)
-	{
-	}
+	ENGINE_API FCanvasTextItem(const FVector2D& InPosition, const FText& InText, const FSlateFontInfo& InFontInfo, const FLinearColor& InColor);
 	
-	ENGINE_API virtual ~FCanvasTextItem() override = default;
+	ENGINE_API virtual ~FCanvasTextItem();
 
 	/* The text to draw. */
 	FText Text;
diff --git a/Engine/Source/Runtime/Engine/Public/DebugRenderSceneProxy.h b/Engine/Source/Runtime/Engine/Public/DebugRenderSceneProxy.h
--- a/Engine/Source/Runtime/Engine/Public/DebugRenderSceneProxy.h
+++ b/Engine/Source/Runtime/Engine/Public/DebugRenderSceneProxy.h
@@ -25,7 +25,7 @@ DECLARE_DELEGATE_TwoParams(FDebugDrawDelegate, UCanvas*, APlayerController*);
 class FDebugRenderSceneProxy : public FPrimitiveSceneProxy
 {
 public:
-	virtual ~FDebugRenderSceneProxy() {};
+	ENGINE_API virtual ~FDebugRenderSceneProxy();
 	
 	ENGINE_API SIZE_T GetTypeHash() const override;
 
@@ -37,7 +37,7 @@ public:
 		Invalid = 3,
 	};
 	ENGINE_API FDebugRenderSceneProxy(const UPrimitiveComponent* InComponent);
-	FDebugRenderSceneProxy(FDebugRenderSceneProxy const&) = default;
+	ENGINE_API FDebugRenderSceneProxy(FDebugRenderSceneProxy const&);
 
 	// FPrimitiveSceneProxy interface.
 
diff --git a/Engine/Source/Runtime/Engine/Public/PrimitiveSceneProxy.h b/Engine/Source/Runtime/Engine/Public/PrimitiveSceneProxy.h
--- a/Engine/Source/Runtime/Engine/Public/PrimitiveSceneProxy.h
+++ b/Engine/Source/Runtime/Engine/Public/PrimitiveSceneProxy.h
@@ -244,7 +244,7 @@ public:
 	ENGINE_API FPrimitiveSceneProxy(const FPrimitiveSceneProxyDesc& InDesc, FName ResourceName = NAME_None);
 
 	/** Copy constructor. */
-	FPrimitiveSceneProxy(FPrimitiveSceneProxy const&) = default;
+	ENGINE_API FPrimitiveSceneProxy(FPrimitiveSceneProxy const&);
 
 	/** Virtual destructor. */
 	ENGINE_API virtual ~FPrimitiveSceneProxy();
diff --git a/Engine/Source/Runtime/Engine/Public/RawIndexBuffer.h b/Engine/Source/Runtime/Engine/Public/RawIndexBuffer.h
--- a/Engine/Source/Runtime/Engine/Public/RawIndexBuffer.h
+++ b/Engine/Source/Runtime/Engine/Public/RawIndexBuffer.h
@@ -134,6 +134,7 @@ public:
 	 * @param InNeedsCPUAccess	True if resource array data should be accessible by the CPU.
 	 */
 	ENGINE_API FRawStaticIndexBuffer(bool InNeedsCPUAccess=false);
+	ENGINE_API virtual ~FRawStaticIndexBuffer();
 
 	/**
 	 * Copy everything, keeping reference to the same RHI resources.
diff --git a/Engine/Source/Runtime/Engine/Public/SceneView.h b/Engine/Source/Runtime/Engine/Public/SceneView.h
--- a/Engine/Source/Runtime/Engine/Public/SceneView.h
+++ b/Engine/Source/Runtime/Engine/Public/SceneView.h
@@ -2384,14 +2384,7 @@ public:
 	void operator = (const FSceneViewFamily&) = delete;
 
 	// Allow moving view family as long as no screen percentage interface are set.
-	FSceneViewFamily(const FSceneViewFamily&& InViewFamily)
-		: FSceneViewFamily(static_cast<const FSceneViewFamily&>(InViewFamily))
-	{
-		check(ScreenPercentageInterface == nullptr);
-		check(TemporalUpscalerInterface == nullptr);
-		check(PrimarySpatialUpscalerInterface == nullptr);
-		check(SecondarySpatialUpscalerInterface == nullptr);
-	}
+	ENGINE_API FSceneViewFamily(FSceneViewFamily&& InViewFamily);
 
 	template<typename TExtensionData> const TExtensionData* GetExtentionData() const
 	{
@@ -2490,10 +2483,8 @@ private:
 	/** True if this view is the current editing view or the active game view */
 	bool bIsInFocus = true;
 
-	PRAGMA_DISABLE_DEPRECATION_WARNINGS
 	// Only FSceneRenderer can copy a view family.
-	FSceneViewFamily(const FSceneViewFamily&) = default;
-	PRAGMA_ENABLE_DEPRECATION_WARNINGS
+	ENGINE_API FSceneViewFamily(const FSceneViewFamily&);
 
 	friend class FSceneRenderer;
 	friend class FViewFamilyInfo;
diff --git a/Engine/Source/Runtime/Engine/Public/UnrealClient.h b/Engine/Source/Runtime/Engine/Public/UnrealClient.h
--- a/Engine/Source/Runtime/Engine/Public/UnrealClient.h
+++ b/Engine/Source/Runtime/Engine/Public/UnrealClient.h
@@ -32,12 +32,12 @@ public:
 	/**
 	 * Default constructor
 	 */
-	FRenderTarget() {};
+	ENGINE_API FRenderTarget();
 
 	/**
 	 * Destructor
 	 */
-	virtual ~FRenderTarget() {};
+	ENGINE_API virtual ~FRenderTarget();
 
 	/**
 	* Accessor for the surface RHI when setting this render target
diff --git a/Engine/Source/Runtime/Engine/Public/WeightedRandomSampler.h b/Engine/Source/Runtime/Engine/Public/WeightedRandomSampler.h
--- a/Engine/Source/Runtime/Engine/Public/WeightedRandomSampler.h
+++ b/Engine/Source/Runtime/Engine/Public/WeightedRandomSampler.h
@@ -31,7 +31,7 @@ struct FWeightedRandomSampler
 	DECLARE_TYPE_LAYOUT(FWeightedRandomSampler, Abstract);
 public:
 	ENGINE_API FWeightedRandomSampler();
-	virtual ~FWeightedRandomSampler() { }
+	ENGINE_API virtual ~FWeightedRandomSampler();
 
 	/**
 	Takes two random values (0...1) and returns the corresponding element index.
diff --git a/Engine/Source/Runtime/Experimental/Chaos/Private/GeometryCollection/GeometryCollection.cpp b/Engine/Source/Runtime/Experimental/Chaos/Private/GeometryCollection/GeometryCollection.cpp
--- a/Engine/Source/Runtime/Experimental/Chaos/Private/GeometryCollection/GeometryCollection.cpp
+++ b/Engine/Source/Runtime/Experimental/Chaos/Private/GeometryCollection/GeometryCollection.cpp
@@ -58,6 +58,8 @@ FGeometryCollection::FGeometryCollection(FGeometryCollectionDefaults InDefaults)
 	Construct();
 }
 
+FGeometryCollection::~FGeometryCollection() = default;
+
 void FGeometryCollection::DefineGeometrySchema(FManagedArrayCollection& InCollection)
 {
 	FTransformCollection::DefineTransformSchema(InCollection);
diff --git a/Engine/Source/Runtime/Experimental/Chaos/Private/GeometryCollection/ManagedArrayCollection.cpp b/Engine/Source/Runtime/Experimental/Chaos/Private/GeometryCollection/ManagedArrayCollection.cpp
--- a/Engine/Source/Runtime/Experimental/Chaos/Private/GeometryCollection/ManagedArrayCollection.cpp
+++ b/Engine/Source/Runtime/Experimental/Chaos/Private/GeometryCollection/ManagedArrayCollection.cpp
@@ -18,6 +18,14 @@ FManagedArrayCollection::FManagedArrayCollection()
 	Version = 9;
 }
 
+
+FManagedArrayCollection::~FManagedArrayCollection() = default;
+
+FManagedArrayCollection::FManagedArrayCollection(const FManagedArrayCollection& In) { In.CopyTo(this); }
+FManagedArrayCollection& FManagedArrayCollection::operator=(const FManagedArrayCollection& In) { Reset(); In.CopyTo(this); return *this; }
+FManagedArrayCollection::FManagedArrayCollection(FManagedArrayCollection&&) = default;
+FManagedArrayCollection& FManagedArrayCollection::operator=(FManagedArrayCollection&&)= default;
+
 void FManagedArrayCollection::AddGroup(FName Group)
 {
 	ensure(!GroupInfo.Contains(Group));
diff --git a/Engine/Source/Runtime/Experimental/Chaos/Private/GeometryCollection/TransformCollection.cpp b/Engine/Source/Runtime/Experimental/Chaos/Private/GeometryCollection/TransformCollection.cpp
--- a/Engine/Source/Runtime/Experimental/Chaos/Private/GeometryCollection/TransformCollection.cpp
+++ b/Engine/Source/Runtime/Experimental/Chaos/Private/GeometryCollection/TransformCollection.cpp
@@ -22,6 +22,8 @@ FTransformCollection::FTransformCollection()
 	Construct();
 }
 
+FTransformCollection::~FTransformCollection() = default;
+
 void FTransformCollection::DefineTransformSchema(FManagedArrayCollection& InCollection)
 {
 	// Hierarchy Group
diff --git a/Engine/Source/Runtime/Experimental/Chaos/Public/GeometryCollection/GeometryCollection.h b/Engine/Source/Runtime/Experimental/Chaos/Public/GeometryCollection/GeometryCollection.h
--- a/Engine/Source/Runtime/Experimental/Chaos/Public/GeometryCollection/GeometryCollection.h
+++ b/Engine/Source/Runtime/Experimental/Chaos/Public/GeometryCollection/GeometryCollection.h
@@ -35,6 +35,7 @@ public:
 	typedef FTransformCollection Super;
 
 	CHAOS_API FGeometryCollection(FGeometryCollectionDefaults InDefaults = FGeometryCollectionDefaults() );
+	CHAOS_API ~FGeometryCollection();
 	FGeometryCollection(FGeometryCollection &) = delete;
 	FGeometryCollection& operator=(const FGeometryCollection &) = delete;
 	FGeometryCollection(FGeometryCollection &&) = default;
diff --git a/Engine/Source/Runtime/Experimental/Chaos/Public/GeometryCollection/ManagedArrayCollection.h b/Engine/Source/Runtime/Experimental/Chaos/Public/GeometryCollection/ManagedArrayCollection.h
--- a/Engine/Source/Runtime/Experimental/Chaos/Public/GeometryCollection/ManagedArrayCollection.h
+++ b/Engine/Source/Runtime/Experimental/Chaos/Public/GeometryCollection/ManagedArrayCollection.h
@@ -61,12 +61,12 @@ public:
 
 
 	CHAOS_API FManagedArrayCollection();
-	virtual ~FManagedArrayCollection() {}
+	CHAOS_API virtual ~FManagedArrayCollection();
 
-	FManagedArrayCollection(const FManagedArrayCollection& In) { In.CopyTo(this); }
-	FManagedArrayCollection& operator=(const FManagedArrayCollection& In) { Reset(); In.CopyTo(this); return *this; }
-	FManagedArrayCollection(FManagedArrayCollection&&) = default;
-	FManagedArrayCollection& operator=(FManagedArrayCollection&&)= default;
+	CHAOS_API FManagedArrayCollection(const FManagedArrayCollection& In);
+	CHAOS_API FManagedArrayCollection& operator=(const FManagedArrayCollection& In);
+	CHAOS_API FManagedArrayCollection(FManagedArrayCollection&&);
+	CHAOS_API FManagedArrayCollection& operator=(FManagedArrayCollection&&);
 
 	static CHAOS_API int8 Invalid;
 	typedef EManagedArrayType EArrayType;
diff --git a/Engine/Source/Runtime/Experimental/Chaos/Public/GeometryCollection/TransformCollection.h b/Engine/Source/Runtime/Experimental/Chaos/Public/GeometryCollection/TransformCollection.h
--- a/Engine/Source/Runtime/Experimental/Chaos/Public/GeometryCollection/TransformCollection.h
+++ b/Engine/Source/Runtime/Experimental/Chaos/Public/GeometryCollection/TransformCollection.h
@@ -19,6 +19,7 @@ public:
 	typedef FManagedArrayCollection Super;
 
 	CHAOS_API FTransformCollection();
+	CHAOS_API ~FTransformCollection();
 	FTransformCollection(FTransformCollection &) = delete;
 	FTransformCollection& operator=(const FTransformCollection&) = delete;
 	FTransformCollection(FTransformCollection&&) = default;
diff --git a/Engine/Source/Runtime/GeometryCore/Private/Generators/FlatTriangulationMeshGenerator.cpp b/Engine/Source/Runtime/GeometryCore/Private/Generators/FlatTriangulationMeshGenerator.cpp
--- a/Engine/Source/Runtime/GeometryCore/Private/Generators/FlatTriangulationMeshGenerator.cpp
+++ b/Engine/Source/Runtime/GeometryCore/Private/Generators/FlatTriangulationMeshGenerator.cpp
@@ -11,7 +11,7 @@ FFlatTriangulationMeshGenerator::FFlatTriangulationMeshGenerator()
 	IndicesMap = FIndex2i(0, 1);
 }
 
-
+FFlatTriangulationMeshGenerator::~FFlatTriangulationMeshGenerator() = default;
 
 FMeshShapeGenerator& FFlatTriangulationMeshGenerator::Generate()
 {
diff --git a/Engine/Source/Runtime/GeometryCore/Private/Generators/PlanarPolygonMeshGenerator.cpp b/Engine/Source/Runtime/GeometryCore/Private/Generators/PlanarPolygonMeshGenerator.cpp
--- a/Engine/Source/Runtime/GeometryCore/Private/Generators/PlanarPolygonMeshGenerator.cpp
+++ b/Engine/Source/Runtime/GeometryCore/Private/Generators/PlanarPolygonMeshGenerator.cpp
@@ -11,6 +11,7 @@ FPlanarPolygonMeshGenerator::FPlanarPolygonMeshGenerator()
 	IndicesMap = FIndex2i(0, 1);
 }
 
+FPlanarPolygonMeshGenerator::~FPlanarPolygonMeshGenerator() = default;
 
 
 void FPlanarPolygonMeshGenerator::SetPolygon(const TArray<FVector2D>& PolygonVerts)
diff --git a/Engine/Source/Runtime/GeometryCore/Private/Generators/PolygonEdgeMeshGenerator.cpp b/Engine/Source/Runtime/GeometryCore/Private/Generators/PolygonEdgeMeshGenerator.cpp
--- a/Engine/Source/Runtime/GeometryCore/Private/Generators/PolygonEdgeMeshGenerator.cpp
+++ b/Engine/Source/Runtime/GeometryCore/Private/Generators/PolygonEdgeMeshGenerator.cpp
@@ -198,7 +198,7 @@ FPolygonEdgeMeshGenerator::FPolygonEdgeMeshGenerator(const TArray<FFrame3d>& InP
 	check(Polygon.Num() == OffsetScaleFactors.Num());
 }
 
-
+FPolygonEdgeMeshGenerator::~FPolygonEdgeMeshGenerator() = default;
 
 // Given a piecewise linear path, replace each corner with a discretized arc
 void FPolygonEdgeMeshGenerator::CurvePath(const TArray<FVector3d>& InPath,
diff --git a/Engine/Source/Runtime/GeometryCore/Private/Generators/StairGenerator.cpp b/Engine/Source/Runtime/GeometryCore/Private/Generators/StairGenerator.cpp
--- a/Engine/Source/Runtime/GeometryCore/Private/Generators/StairGenerator.cpp
+++ b/Engine/Source/Runtime/GeometryCore/Private/Generators/StairGenerator.cpp
@@ -9,6 +9,8 @@ using namespace UE::Geometry;
  * FStairGenerator
  */
 
+FStairGenerator::~FStairGenerator() = default;
+
 FMeshShapeGenerator& FStairGenerator::Generate()
 {
 	TRACE_CPUPROFILER_EVENT_SCOPE(StairGenerator_Generate);
diff --git a/Engine/Source/Runtime/GeometryCore/Public/Generators/FlatTriangulationMeshGenerator.h b/Engine/Source/Runtime/GeometryCore/Public/Generators/FlatTriangulationMeshGenerator.h
--- a/Engine/Source/Runtime/GeometryCore/Public/Generators/FlatTriangulationMeshGenerator.h
+++ b/Engine/Source/Runtime/GeometryCore/Public/Generators/FlatTriangulationMeshGenerator.h
@@ -37,6 +37,7 @@ public:
 
 public:
 	GEOMETRYCORE_API FFlatTriangulationMeshGenerator();
+	GEOMETRYCORE_API ~FFlatTriangulationMeshGenerator();
 
 	/** Generate the triangulation */
 	GEOMETRYCORE_API virtual FMeshShapeGenerator& Generate() override;
diff --git a/Engine/Source/Runtime/GeometryCore/Public/Generators/PlanarPolygonMeshGenerator.h b/Engine/Source/Runtime/GeometryCore/Public/Generators/PlanarPolygonMeshGenerator.h
--- a/Engine/Source/Runtime/GeometryCore/Public/Generators/PlanarPolygonMeshGenerator.h
+++ b/Engine/Source/Runtime/GeometryCore/Public/Generators/PlanarPolygonMeshGenerator.h
@@ -33,6 +33,7 @@ public:
 
 public:
 	GEOMETRYCORE_API FPlanarPolygonMeshGenerator();
+	GEOMETRYCORE_API ~FPlanarPolygonMeshGenerator();
 
 	/** Initialize the polygon from an array of 2D vertices */
 	GEOMETRYCORE_API void SetPolygon(const TArray<FVector2D>& PolygonVerts);
diff --git a/Engine/Source/Runtime/GeometryCore/Public/Generators/PolygonEdgeMeshGenerator.h b/Engine/Source/Runtime/GeometryCore/Public/Generators/PolygonEdgeMeshGenerator.h
--- a/Engine/Source/Runtime/GeometryCore/Public/Generators/PolygonEdgeMeshGenerator.h
+++ b/Engine/Source/Runtime/GeometryCore/Public/Generators/PolygonEdgeMeshGenerator.h
@@ -67,6 +67,8 @@ public:
 		bool bInLimitCornerRadius = false,
 		int InNumArcVertices = 3);
 
+	GEOMETRYCORE_API ~FPolygonEdgeMeshGenerator();
+
 	// Generate triangulation
 	// TODO: Enable more subdivisions along the width and length dimensions if requested
 	GEOMETRYCORE_API virtual FMeshShapeGenerator& Generate() final;
diff --git a/Engine/Source/Runtime/GeometryCore/Public/Generators/StairGenerator.h b/Engine/Source/Runtime/GeometryCore/Public/Generators/StairGenerator.h
--- a/Engine/Source/Runtime/GeometryCore/Public/Generators/StairGenerator.h
+++ b/Engine/Source/Runtime/GeometryCore/Public/Generators/StairGenerator.h
@@ -47,9 +47,7 @@ public:
 	};
 
 public:
-	virtual ~FStairGenerator() override
-	{
-	}
+	GEOMETRYCORE_API virtual ~FStairGenerator() override;
 
 	/** Generate the mesh */
 	GEOMETRYCORE_API virtual FMeshShapeGenerator& Generate() override;
diff --git a/Engine/Source/Runtime/GeometryFramework/Private/Changes/MeshChange.cpp b/Engine/Source/Runtime/GeometryFramework/Private/Changes/MeshChange.cpp
--- a/Engine/Source/Runtime/GeometryFramework/Private/Changes/MeshChange.cpp
+++ b/Engine/Source/Runtime/GeometryFramework/Private/Changes/MeshChange.cpp
@@ -7,15 +7,15 @@
 
 using namespace UE::Geometry;
 
-FMeshChange::FMeshChange()
-{
-}
+FMeshChange::FMeshChange() = default;
 
 FMeshChange::FMeshChange(TUniquePtr<FDynamicMeshChange> DynamicMeshChangeIn)
 {
 	DynamicMeshChange = MoveTemp(DynamicMeshChangeIn);
 }
 
+FMeshChange::~FMeshChange() = default;
+
 void FMeshChange::Apply(UObject* Object)
 {
 	IMeshCommandChangeTarget* ChangeTarget = CastChecked<IMeshCommandChangeTarget>(Object);
diff --git a/Engine/Source/Runtime/GeometryFramework/Public/Changes/MeshChange.h b/Engine/Source/Runtime/GeometryFramework/Public/Changes/MeshChange.h
--- a/Engine/Source/Runtime/GeometryFramework/Public/Changes/MeshChange.h
+++ b/Engine/Source/Runtime/GeometryFramework/Public/Changes/MeshChange.h
@@ -32,6 +32,7 @@ class FMeshChange : public FToolCommandChange
 public:
 	GEOMETRYFRAMEWORK_API FMeshChange();
 	GEOMETRYFRAMEWORK_API FMeshChange(TUniquePtr<UE::Geometry::FDynamicMeshChange> DynamicMeshChangeIn);
+	GEOMETRYFRAMEWORK_API ~FMeshChange();
 
 	TUniquePtr<UE::Geometry::FDynamicMeshChange> DynamicMeshChange;
 
diff --git a/Engine/Source/Runtime/InteractiveToolsFramework/Private/ToolDataVisualizer.cpp b/Engine/Source/Runtime/InteractiveToolsFramework/Private/ToolDataVisualizer.cpp
--- a/Engine/Source/Runtime/InteractiveToolsFramework/Private/ToolDataVisualizer.cpp
+++ b/Engine/Source/Runtime/InteractiveToolsFramework/Private/ToolDataVisualizer.cpp
@@ -24,6 +24,8 @@ FToolDataVisualizer::FToolDataVisualizer()
 	PopAllTransforms();
 }
 
+FToolDataVisualizer::FToolDataVisualizer(const FToolDataVisualizer&) = default;
+FToolDataVisualizer::~FToolDataVisualizer() = default;
 
 void FToolDataVisualizer::BeginFrame(IToolsContextRenderAPI* RenderAPI, const FViewCameraState& CameraStateIn)
 {
diff --git a/Engine/Source/Runtime/InteractiveToolsFramework/Public/ToolDataVisualizer.h b/Engine/Source/Runtime/InteractiveToolsFramework/Public/ToolDataVisualizer.h
--- a/Engine/Source/Runtime/InteractiveToolsFramework/Public/ToolDataVisualizer.h
+++ b/Engine/Source/Runtime/InteractiveToolsFramework/Public/ToolDataVisualizer.h
@@ -35,7 +35,8 @@ public:
 
 public:
 	INTERACTIVETOOLSFRAMEWORK_API FToolDataVisualizer();
-	virtual ~FToolDataVisualizer() {}
+	INTERACTIVETOOLSFRAMEWORK_API FToolDataVisualizer(const FToolDataVisualizer&);
+	INTERACTIVETOOLSFRAMEWORK_API virtual ~FToolDataVisualizer();
 
 
 	//
diff --git a/Engine/Source/Runtime/LevelSequence/Private/LevelSequenceSpawnRegister.cpp b/Engine/Source/Runtime/LevelSequence/Private/LevelSequenceSpawnRegister.cpp
--- a/Engine/Source/Runtime/LevelSequence/Private/LevelSequenceSpawnRegister.cpp
+++ b/Engine/Source/Runtime/LevelSequence/Private/LevelSequenceSpawnRegister.cpp
@@ -15,6 +15,9 @@ FLevelSequenceSpawnRegister::FLevelSequenceSpawnRegister()
 	LevelSequenceModule.GenerateObjectSpawners(MovieSceneObjectSpawners);
 }
 
+FLevelSequenceSpawnRegister::FLevelSequenceSpawnRegister(const FLevelSequenceSpawnRegister&) = default;
+FLevelSequenceSpawnRegister::~FLevelSequenceSpawnRegister() = default;
+
 UObject* FLevelSequenceSpawnRegister::SpawnObject(FMovieSceneSpawnable& Spawnable, FMovieSceneSequenceIDRef TemplateID, TSharedRef<const FSharedPlaybackState> SharedPlaybackState)
 {
 	for (TSharedRef<IMovieSceneObjectSpawner> MovieSceneObjectSpawner : MovieSceneObjectSpawners)
diff --git a/Engine/Source/Runtime/LevelSequence/Public/LevelSequenceSpawnRegister.h b/Engine/Source/Runtime/LevelSequence/Public/LevelSequenceSpawnRegister.h
--- a/Engine/Source/Runtime/LevelSequence/Public/LevelSequenceSpawnRegister.h
+++ b/Engine/Source/Runtime/LevelSequence/Public/LevelSequenceSpawnRegister.h
@@ -21,6 +21,8 @@ class FLevelSequenceSpawnRegister : public FMovieSceneSpawnRegister
 {
 public:
 	LEVELSEQUENCE_API FLevelSequenceSpawnRegister();
+	LEVELSEQUENCE_API FLevelSequenceSpawnRegister(const FLevelSequenceSpawnRegister&);
+	LEVELSEQUENCE_API ~FLevelSequenceSpawnRegister();
 
 protected:
 	/** ~ FMovieSceneSpawnRegister interface */
diff --git a/Engine/Source/Runtime/MediaUtils/Private/MediaSamples.cpp b/Engine/Source/Runtime/MediaUtils/Private/MediaSamples.cpp
--- a/Engine/Source/Runtime/MediaUtils/Private/MediaSamples.cpp
+++ b/Engine/Source/Runtime/MediaUtils/Private/MediaSamples.cpp
@@ -69,6 +69,7 @@ FMediaSamples::FMediaSamples(uint32 InMaxNumberOfQueuedAudioSamples, uint32 InMa
 {
 }
 
+FMediaSamples::~FMediaSamples() = default;
 
 bool FMediaSamples::FetchAudio(TRange<FTimespan> TimeRange, TSharedPtr<IMediaAudioSample, ESPMode::ThreadSafe>& OutSample)
 {
diff --git a/Engine/Source/Runtime/MediaUtils/Public/MediaSamples.h b/Engine/Source/Runtime/MediaUtils/Public/MediaSamples.h
--- a/Engine/Source/Runtime/MediaUtils/Public/MediaSamples.h
+++ b/Engine/Source/Runtime/MediaUtils/Public/MediaSamples.h
@@ -36,6 +36,8 @@ public:
 				  uint32 InMaxNumberOfQueuedSubtitlesSamples = MaxNumberOfQueuedSubtitlesSamples,
 				  uint32 InMaxNumberOfQueuedMetaDataSamples = MaxNumberOfQueuedMetaDataSamples);
 
+	MEDIAUTILS_API ~FMediaSamples();
+
 	/**
 	 * Add the given audio sample to the cache.
 	 *
diff --git a/Engine/Source/Runtime/MovieScene/Private/Evaluation/Blending/MovieSceneInitialValueStore.cpp b/Engine/Source/Runtime/MovieScene/Private/Evaluation/Blending/MovieSceneInitialValueStore.cpp
--- a/Engine/Source/Runtime/MovieScene/Private/Evaluation/Blending/MovieSceneInitialValueStore.cpp
+++ b/Engine/Source/Runtime/MovieScene/Private/Evaluation/Blending/MovieSceneInitialValueStore.cpp
@@ -48,6 +48,8 @@ FMovieSceneRemoveInitialValueTokenProducer::FMovieSceneRemoveInitialValueTokenPr
 {
 }
 
+FMovieSceneRemoveInitialValueTokenProducer::~FMovieSceneRemoveInitialValueTokenProducer() = default;
+
 IMovieScenePreAnimatedTokenPtr FMovieSceneRemoveInitialValueTokenProducer::CacheExistingState(UObject& Object) const
 {
 	return FMovieSceneRemoveInitialValueToken(WeakActuator);
@@ -58,6 +60,8 @@ FMovieSceneRemoveInitialGlobalValueTokenProducer::FMovieSceneRemoveInitialGlobal
 {
 }
 
+FMovieSceneRemoveInitialGlobalValueTokenProducer::~FMovieSceneRemoveInitialGlobalValueTokenProducer() = default;
+
 IMovieScenePreAnimatedGlobalTokenPtr FMovieSceneRemoveInitialGlobalValueTokenProducer::CacheExistingState() const
 {
 	return FMovieSceneRemoveInitialGlobalValueToken(WeakActuator);
diff --git a/Engine/Source/Runtime/MovieScene/Private/Evaluation/MovieScenePropertyTemplate.cpp b/Engine/Source/Runtime/MovieScene/Private/Evaluation/MovieScenePropertyTemplate.cpp
--- a/Engine/Source/Runtime/MovieScene/Private/Evaluation/MovieScenePropertyTemplate.cpp
+++ b/Engine/Source/Runtime/MovieScene/Private/Evaluation/MovieScenePropertyTemplate.cpp
@@ -18,6 +18,10 @@ void PropertyTemplate::FSectionData::Initialize(FName InPropertyName, FString In
 	PropertyBindings = MakeShareable(new FTrackInstancePropertyBindings(InPropertyName, MoveTemp(InPropertyPath)));
 }
 
+FMovieScenePropertySectionTemplate::FMovieScenePropertySectionTemplate() = default;
+FMovieScenePropertySectionTemplate::FMovieScenePropertySectionTemplate(const FMovieScenePropertySectionTemplate&) = default;
+FMovieScenePropertySectionTemplate::~FMovieScenePropertySectionTemplate() = default;
+
 FMovieScenePropertySectionTemplate::FMovieScenePropertySectionTemplate(FName PropertyName, const FString& InPropertyPath)
 	: PropertyData(PropertyName, InPropertyPath)
 {}
diff --git a/Engine/Source/Runtime/MovieScene/Private/MovieSceneSpawnRegister.cpp b/Engine/Source/Runtime/MovieScene/Private/MovieSceneSpawnRegister.cpp
--- a/Engine/Source/Runtime/MovieScene/Private/MovieSceneSpawnRegister.cpp
+++ b/Engine/Source/Runtime/MovieScene/Private/MovieSceneSpawnRegister.cpp
@@ -10,6 +10,10 @@
 
 UE::MovieScene::TPlaybackCapabilityID<FMovieSceneSpawnRegister> FMovieSceneSpawnRegister::ID = UE::MovieScene::TPlaybackCapabilityID<FMovieSceneSpawnRegister>::Register();
 
+FMovieSceneSpawnRegister::FMovieSceneSpawnRegister() = default;
+FMovieSceneSpawnRegister::FMovieSceneSpawnRegister(const FMovieSceneSpawnRegister&) = default;
+FMovieSceneSpawnRegister::~FMovieSceneSpawnRegister() = default;
+
 TWeakObjectPtr<> FMovieSceneSpawnRegister::FindSpawnedObject(const FGuid& BindingId, FMovieSceneSequenceIDRef TemplateID, int BindingIndex/* = 0*/) const
 {
 	FMovieSceneSpawnRegisterKey Key(TemplateID, BindingId, BindingIndex);
diff --git a/Engine/Source/Runtime/MovieScene/Public/Evaluation/Blending/MovieSceneInitialValueStore.h b/Engine/Source/Runtime/MovieScene/Public/Evaluation/Blending/MovieSceneInitialValueStore.h
--- a/Engine/Source/Runtime/MovieScene/Public/Evaluation/Blending/MovieSceneInitialValueStore.h
+++ b/Engine/Source/Runtime/MovieScene/Public/Evaluation/Blending/MovieSceneInitialValueStore.h
@@ -20,6 +20,7 @@ struct FMovieSceneRemoveInitialValueTokenProducer : IMovieScenePreAnimatedTokenP
 {
 	/** Construction from the object whose initial value to remove, and the actuator to remove it from */
 	MOVIESCENE_API FMovieSceneRemoveInitialValueTokenProducer(TWeakPtr<IMovieSceneBlendingActuator> InWeakActuator);
+	MOVIESCENE_API ~FMovieSceneRemoveInitialValueTokenProducer();
 
 	MOVIESCENE_API virtual IMovieScenePreAnimatedTokenPtr CacheExistingState(UObject& InObject) const override;
 
@@ -33,6 +34,7 @@ struct FMovieSceneRemoveInitialGlobalValueTokenProducer : IMovieScenePreAnimated
 {
 	/** Construction from the object whose initial value to remove, and the actuator to remove it from */
 	MOVIESCENE_API FMovieSceneRemoveInitialGlobalValueTokenProducer(TWeakPtr<IMovieSceneBlendingActuator> InWeakActuator);
+	MOVIESCENE_API ~FMovieSceneRemoveInitialGlobalValueTokenProducer();
 
 	MOVIESCENE_API virtual IMovieScenePreAnimatedGlobalTokenPtr CacheExistingState() const override;
 
