Moved the MusicBattleInstance tick to c++

#tests Played battlestage, printed out debug values from UMaterialParameterCollection and made sure they were updating as expected.
#rnx

diff --git a/Engine/Plugins/Runtime/Harmonix/Source/HarmonixMetasound/Public/HarmonixMetasound/Analysis/PeakTamer.h b/Engine/Plugins/Runtime/Harmonix/Source/HarmonixMetasound/Public/HarmonixMetasound/Analysis/PeakTamer.h
--- a/Engine/Plugins/Runtime/Harmonix/Source/HarmonixMetasound/Public/HarmonixMetasound/Analysis/PeakTamer.h
+++ b/Engine/Plugins/Runtime/Harmonix/Source/HarmonixMetasound/Public/HarmonixMetasound/Analysis/PeakTamer.h
@@ -9,19 +9,19 @@ struct FHarmonixPeakTamerSettings
 {
 	GENERATED_BODY()
 
-	UPROPERTY(BlueprintReadWrite, Category="AudioAnalysis")
+	UPROPERTY(EditAnywhere, BlueprintReadWrite, Category="AudioAnalysis")
 	float PeakAttackTimeSeconds = 0.01f;
 
-	UPROPERTY(BlueprintReadWrite, Category="AudioAnalysis")
+	UPROPERTY(EditAnywhere, BlueprintReadWrite, Category="AudioAnalysis")
 	float PeakReleaseTimeSeconds = 2.0f;
 
-	UPROPERTY(BlueprintReadWrite, Category="AudioAnalysis")
+	UPROPERTY(EditAnywhere, BlueprintReadWrite, Category="AudioAnalysis")
 	bool bEnableValueSmoothing = true;
 
-	UPROPERTY(BlueprintReadWrite, Category="AudioAnalysis")
+	UPROPERTY(EditAnywhere, BlueprintReadWrite, Category="AudioAnalysis")
 	float ValueAttackTimeSeconds = 0.01f;
 	
-	UPROPERTY(BlueprintReadWrite, Category="AudioAnalysis")
+	UPROPERTY(EditAnywhere, BlueprintReadWrite, Category="AudioAnalysis")
 	float ValueReleaseTimeSeconds = 0.01f;
 };
 
