Fix issue in musicclockcomponent not updating songmaps when midi files changes. When using an externally driven clock.

diff --git a/Engine/Plugins/Runtime/Harmonix/Source/HarmonixMetasound/Private/HarmonixMetasound/DataTypes/MidiClock.cpp b/Engine/Plugins/Runtime/Harmonix/Source/HarmonixMetasound/Private/HarmonixMetasound/DataTypes/MidiClock.cpp
--- a/Engine/Plugins/Runtime/Harmonix/Source/HarmonixMetasound/Private/HarmonixMetasound/DataTypes/MidiClock.cpp
+++ b/Engine/Plugins/Runtime/Harmonix/Source/HarmonixMetasound/Private/HarmonixMetasound/DataTypes/MidiClock.cpp
@@ -237,7 +237,6 @@ namespace HarmonixMetasound
 		if (ExternalClockDriver && ExternalClockDriver->MidiDataChangedInBlock)
 		{
 			RebuildSongMapEvaluator(ExternalClockDriver->SongMapEvaluator->GetSongMapsWithTempoMap(), SongMapEvaluator->GetSongMapsWithOtherMaps());
-			PostTempoOrTimeSignatureEventsIfNeeded();
 		}
 	}
 
@@ -1078,6 +1077,7 @@ namespace HarmonixMetasound
 	void FMidiClock::RebuildSongMapEvaluator(const TSharedPtr<const ISongMapEvaluator>& MapWithTempo, const TSharedPtr<const ISongMapEvaluator>& MapWithOtherMaps)
 	{
 		SongMapEvaluator = MakeShared<FSongMapsWithAlternateTempoSource>(MapWithTempo, MapWithOtherMaps);
+		SongMapsChanged();
 	}
 
 }
