MusicMapBase.h More robust checking when converting float Farctional "Point" to an Index in the Points map. Handle all possible negative numbers.

#tests Found bug while testing MusicClockComponent in editor
#rnx

diff --git a/Engine/Plugins/Runtime/Harmonix/Source/HarmonixMidi/Public/HarmonixMidi/MusicMapBase.h b/Engine/Plugins/Runtime/Harmonix/Source/HarmonixMidi/Public/HarmonixMidi/MusicMapBase.h
--- a/Engine/Plugins/Runtime/Harmonix/Source/HarmonixMidi/Public/HarmonixMidi/MusicMapBase.h
+++ b/Engine/Plugins/Runtime/Harmonix/Source/HarmonixMidi/Public/HarmonixMidi/MusicMapBase.h
@@ -114,7 +114,7 @@ public:
 	static float GetFractionalTickForFractionalPoint(const TArray<T>& Points, float Point)
 	{
 		int32 Index = FMath::FloorToInt(Point);
-		if (Index == -1 || Index >= Points.Num())
+		if (!Points.IsValidIndex(Index))
 		{
 			return 0.0f;
 		}
