Removed NAN checks from Harmonix Biquad filter as we rely on Flush to Zero floating unit flags for audio processing

diff --git a/Engine/Plugins/Runtime/Harmonix/Source/HarmonixDsp/Public/HarmonixDsp/Effects/BiquadFilter.h b/Engine/Plugins/Runtime/Harmonix/Source/HarmonixDsp/Public/HarmonixDsp/Effects/BiquadFilter.h
--- a/Engine/Plugins/Runtime/Harmonix/Source/HarmonixDsp/Public/HarmonixDsp/Effects/BiquadFilter.h
+++ b/Engine/Plugins/Runtime/Harmonix/Source/HarmonixDsp/Public/HarmonixDsp/Effects/BiquadFilter.h
@@ -174,10 +174,6 @@ namespace Harmonix::Dsp::Effects
 			FF2 = FF1;
 			FF1 = FF0;
 
-			// protected against NAN
-			FF1 = FGenericPlatformMath::IsNaN(FF1) ? 0.0 : FF1;
-			FF2 = FGenericPlatformMath::IsNaN(FF2) ? 0.0 : FF2;
-
 			return (float)Out;
 		}
 
@@ -215,8 +211,8 @@ namespace Harmonix::Dsp::Effects
 				IncFF1 = IncFF0;
 			}
 
-			FF1 = FGenericPlatformMath::IsNaN(IncFF1) ? 0.0 : IncFF1;
-			FF2 = FGenericPlatformMath::IsNaN(IncFF2) ? 0.0 : IncFF2;
+			FF1 = IncFF1;
+			FF2 = IncFF2;
 		}
 
 		void ProcessTransposed(float const* input,
@@ -250,8 +246,8 @@ namespace Harmonix::Dsp::Effects
 				output[n] = (float)Out;
 			}
 
-			FF1 = FGenericPlatformMath::IsNaN(d1) ? 0.0 : d1;
-			FF2 = FGenericPlatformMath::IsNaN(d2) ? 0.0 : d2;
+			FF1 = d1;
+			FF2 = d2;
 		}
 
 		void ProcessInterleaved(TDynamicStridePtr<float>& input,
@@ -352,8 +348,8 @@ namespace Harmonix::Dsp::Effects
 				ff1 = ff0;
 			}
 
-			History[0] = !FGenericPlatformMath::IsNaN(ff1) ? ff1 : (T)0.0;
-			History[1] = !FGenericPlatformMath::IsNaN(ff2) ? ff2 : (T)0.0;
+			History[0] = ff1;
+			History[1] = ff2;
 
 			for (int32 i = 1; i < InNumPasses; ++i)
 			{
@@ -376,8 +372,8 @@ namespace Harmonix::Dsp::Effects
 					ff1 = ff0;
 				}
 
-				History[2 * i] = !FGenericPlatformMath::IsNaN(ff1) ? ff1 : (T)0.0;
-				History[2 * i + 1] = !FGenericPlatformMath::IsNaN(ff2) ? ff2 : (T)0.0;
+				History[2 * i] = ff1;
+				History[2 * i + 1] = ff2;
 			}
 		}
 
