We no longer need to cache off a WeakPtr to the MidiClock in the Pulse Generator.

Refactor PulseGenerator to be more "Functional" and simpler

#rnx
#tests Harmonix Tests

diff --git a/Engine/Plugins/Runtime/Harmonix/Source/HarmonixMetasound/Private/HarmonixMetasound/MidiOps/PulseGenerator.cpp b/Engine/Plugins/Runtime/Harmonix/Source/HarmonixMetasound/Private/HarmonixMetasound/MidiOps/PulseGenerator.cpp
--- a/Engine/Plugins/Runtime/Harmonix/Source/HarmonixMetasound/Private/HarmonixMetasound/MidiOps/PulseGenerator.cpp
+++ b/Engine/Plugins/Runtime/Harmonix/Source/HarmonixMetasound/Private/HarmonixMetasound/MidiOps/PulseGenerator.cpp
@@ -9,22 +9,6 @@ namespace Harmonix::Midi::Ops
 		Enabled = bEnable;
 	}
 
-	void FPulseGenerator::SetClock(const TSharedPtr<const HarmonixMetasound::FMidiClock, ESPMode::NotThreadSafe>& NewClock)
-	{
-		Clock = NewClock;
-		CurrentTimeSignature = *NewClock->GetSongMapEvaluator().GetTimeSignatureAtTick(NewClock->GetLastProcessedMidiTick());
-
-		// Find the next pulse and line up phase with the current bar
-		const FMusicTimestamp ClockCurrentTimestamp = NewClock->GetMusicTimestampAtBlockOffset(0);
-		NextPulseTimestamp.Bar = ClockCurrentTimestamp.Bar;
-		NextPulseTimestamp.Beat = 1;
-		IncrementTimestampByOffset(NextPulseTimestamp, Interval, CurrentTimeSignature);
-		while (NextPulseTimestamp < ClockCurrentTimestamp)
-		{
-			IncrementTimestampByInterval(NextPulseTimestamp, Interval, CurrentTimeSignature);
-		}
-	}
-
 	void FPulseGenerator::SetInterval(const FMusicTimeInterval& NewInterval)
 	{
 		Interval = NewInterval;
@@ -38,20 +22,28 @@ namespace Harmonix::Midi::Ops
 		NextPulseTimestamp = { -1, -1 };
 	}
 
-	void FPulseGenerator::Process(const TFunctionRef<void(const FPulseInfo&)>& OnPulse)
+	void FPulseGenerator::Process(const HarmonixMetasound::FMidiClock& MidiClock, const TFunctionRef<void(const FPulseInfo&)>& OnPulse)
 	{
-		check(Clock.IsValid()); // if we're here, and we don't have a clock, we have problems
-		const auto PinnedClock = Clock.Pin();
-
-		if (!PinnedClock)
+		// ensure the pulse generator is lined up with the current clock phase
+		if (NextPulseTimestamp == FMusicTimestamp(-1, -1))
 		{
-			return;
+			CurrentTimeSignature = *MidiClock.GetSongMapEvaluator().GetTimeSignatureAtTick(MidiClock.GetLastProcessedMidiTick());
+
+			// Find the next pulse and line up phase with the current bar
+			const FMusicTimestamp ClockCurrentTimestamp = MidiClock.GetMusicTimestampAtBlockOffset(0);
+			NextPulseTimestamp.Bar = ClockCurrentTimestamp.Bar;
+			NextPulseTimestamp.Beat = 1;
+			IncrementTimestampByOffset(NextPulseTimestamp, Interval, CurrentTimeSignature);
+			while (NextPulseTimestamp < ClockCurrentTimestamp)
+			{
+				IncrementTimestampByInterval(NextPulseTimestamp, Interval, CurrentTimeSignature);
+			}
 		}
-
+		
 		using namespace HarmonixMetasound;
 		using namespace HarmonixMetasound::MidiClockMessageTypes;
 
-		for (const FMidiClockEvent& ClockEvent : PinnedClock->GetMidiClockEventsInBlock())
+		for (const FMidiClockEvent& ClockEvent : MidiClock.GetMidiClockEventsInBlock())
 		{
 			if (const FAdvance* AsAdvance = ClockEvent.TryGet<FAdvance>())
 			{
@@ -60,7 +52,7 @@ namespace Harmonix::Midi::Ops
 					return;
 				}
 				
-				int32 NextPulseTick = PinnedClock->GetSongMapEvaluator().MusicTimestampToTick(NextPulseTimestamp);
+				int32 NextPulseTick = MidiClock.GetSongMapEvaluator().MusicTimestampToTick(NextPulseTimestamp);
 
 				while (AsAdvance->LastTickToProcess() >= NextPulseTick)
 				{
@@ -68,7 +60,7 @@ namespace Harmonix::Midi::Ops
 
 					IncrementTimestampByInterval(NextPulseTimestamp, Interval, CurrentTimeSignature);
 
-					NextPulseTick = PinnedClock->GetSongMapEvaluator().MusicTimestampToTick(NextPulseTimestamp);
+					NextPulseTick = MidiClock.GetSongMapEvaluator().MusicTimestampToTick(NextPulseTimestamp);
 				}
 			}
 			else if (const FTimeSignatureChange* AsTimeSigChange = ClockEvent.TryGet<FTimeSignatureChange>())
@@ -77,19 +69,19 @@ namespace Harmonix::Midi::Ops
 				
 				// Time sig changes will come on the downbeat, and if we change time signature,
 				// we want to reset the pulse, so the next pulse is now plus the offset
-				NextPulseTimestamp = PinnedClock->GetSongMapEvaluator().TickToMusicTimestamp(AsTimeSigChange->Tick);
+				NextPulseTimestamp = MidiClock.GetSongMapEvaluator().TickToMusicTimestamp(AsTimeSigChange->Tick);
 				IncrementTimestampByOffset(NextPulseTimestamp, Interval, CurrentTimeSignature);
 			}
 			else if (const FLoop* AsLoop = ClockEvent.TryGet<FLoop>())
 			{
 				// We assume the pulse should be reset on loop, and the loop start should imply the phase of the pulse
-				NextPulseTimestamp = PinnedClock->GetSongMapEvaluator().TickToMusicTimestamp(AsLoop->FirstTickInLoop);
+				NextPulseTimestamp = MidiClock.GetSongMapEvaluator().TickToMusicTimestamp(AsLoop->FirstTickInLoop);
 				IncrementTimestampByOffset(NextPulseTimestamp, Interval, CurrentTimeSignature);
 			}
 			else if (const FSeek* AsSeek = ClockEvent.TryGet<FSeek>())
 			{
 				// When we seek, reset the pulse phase to the current bar
-				const FMusicTimestamp ClockCurrentTimestamp = PinnedClock->GetSongMapEvaluator().TickToMusicTimestamp(AsSeek->NewNextTick);
+				const FMusicTimestamp ClockCurrentTimestamp = MidiClock.GetSongMapEvaluator().TickToMusicTimestamp(AsSeek->NewNextTick);
 				NextPulseTimestamp.Bar = ClockCurrentTimestamp.Bar;
 				NextPulseTimestamp.Beat = 1;
 				IncrementTimestampByOffset(NextPulseTimestamp, Interval, CurrentTimeSignature);
@@ -108,11 +100,11 @@ namespace Harmonix::Midi::Ops
 		LastNoteOn.Reset();
 	}
 
-	void FMidiPulseGenerator::Process(HarmonixMetasound::FMidiStream& OutStream)
+	void FMidiPulseGenerator::Process(const HarmonixMetasound::FMidiClock& MidiClock, HarmonixMetasound::FMidiStream& OutStream)
 	{
 		OutStream.PrepareBlock();
 
-		FPulseGenerator::Process([this, &OutStream](const FPulseInfo& Pulse)
+		FPulseGenerator::Process(MidiClock, [this, &OutStream](const FPulseInfo& Pulse)
 		{
 			AddPulseNote(Pulse.BlockFrameIndex, Pulse.Tick, OutStream);
 		});
diff --git a/Engine/Plugins/Runtime/Harmonix/Source/HarmonixMetasound/Private/HarmonixMetasound/Nodes/MidiClockSubdivisionTriggerNode.cpp b/Engine/Plugins/Runtime/Harmonix/Source/HarmonixMetasound/Private/HarmonixMetasound/Nodes/MidiClockSubdivisionTriggerNode.cpp
--- a/Engine/Plugins/Runtime/Harmonix/Source/HarmonixMetasound/Private/HarmonixMetasound/Nodes/MidiClockSubdivisionTriggerNode.cpp
+++ b/Engine/Plugins/Runtime/Harmonix/Source/HarmonixMetasound/Private/HarmonixMetasound/Nodes/MidiClockSubdivisionTriggerNode.cpp
@@ -196,7 +196,6 @@ namespace HarmonixMetasound::Nodes::MidiClockSubdivisionTriggerNode
 		InVertexData.BindReadVertex(METASOUND_GET_PARAM_NAME(Inputs::OffsetUnits), GridOffsetUnitsInPin);
 
 		PulseGenerator.Reset();
-		PulseGenerator.SetClock(MidiClockInPin->AsShared());
 	}
 
 	void FMidiClockSubdivisionTriggerOperator::BindOutputs(FOutputVertexInterfaceData& InVertexData)
@@ -209,7 +208,6 @@ namespace HarmonixMetasound::Nodes::MidiClockSubdivisionTriggerNode
 		TriggerOutPin->Reset();
 		
 		PulseGenerator.Reset();
-		PulseGenerator.SetClock(MidiClockInPin->AsShared());
 	}
 
 	void FMidiClockSubdivisionTriggerOperator::Execute()
@@ -224,7 +222,7 @@ namespace HarmonixMetasound::Nodes::MidiClockSubdivisionTriggerNode
 
 		// If there were pulses, trigger the output
 		TriggerOutPin->AdvanceBlock();
-		PulseGenerator.Process([this](const Harmonix::Midi::Ops::FPulseGenerator::FPulseInfo& Pulse)
+		PulseGenerator.Process(*MidiClockInPin, [this](const Harmonix::Midi::Ops::FPulseGenerator::FPulseInfo& Pulse)
 		{
 			TriggerOutPin->TriggerFrame(Pulse.BlockFrameIndex);
 		});
diff --git a/Engine/Plugins/Runtime/Harmonix/Source/HarmonixMetasound/Private/HarmonixMetasound/Nodes/MidiPulseGeneratorNode.cpp b/Engine/Plugins/Runtime/Harmonix/Source/HarmonixMetasound/Private/HarmonixMetasound/Nodes/MidiPulseGeneratorNode.cpp
--- a/Engine/Plugins/Runtime/Harmonix/Source/HarmonixMetasound/Private/HarmonixMetasound/Nodes/MidiPulseGeneratorNode.cpp
+++ b/Engine/Plugins/Runtime/Harmonix/Source/HarmonixMetasound/Private/HarmonixMetasound/Nodes/MidiPulseGeneratorNode.cpp
@@ -160,7 +160,6 @@ namespace HarmonixMetasound::Nodes::MidiPulseGeneratorNode
 		void Reset(const FResetParams&)
 		{
 			PulseGenerator.Reset();
-			PulseGenerator.SetClock(Inputs.Clock->AsShared());
 			ApplyParameters();
 		}
 
@@ -177,7 +176,6 @@ namespace HarmonixMetasound::Nodes::MidiPulseGeneratorNode
 			InVertexData.BindReadVertex(Inputs::MidiVelocityName, Inputs.Velocity);
 
 			PulseGenerator.Reset();
-			PulseGenerator.SetClock(Inputs.Clock->AsShared());
 			ApplyParameters();
 		}
 
@@ -190,7 +188,7 @@ namespace HarmonixMetasound::Nodes::MidiPulseGeneratorNode
 		{
 			ApplyParameters();
 			
-			PulseGenerator.Process(*Outputs.MidiStream);
+			PulseGenerator.Process(*Inputs.Clock, *Outputs.MidiStream);
 		}
 	private:
 		void ApplyParameters()
diff --git a/Engine/Plugins/Runtime/Harmonix/Source/HarmonixMetasound/Public/HarmonixMetasound/MidiOps/PulseGenerator.h b/Engine/Plugins/Runtime/Harmonix/Source/HarmonixMetasound/Public/HarmonixMetasound/MidiOps/PulseGenerator.h
--- a/Engine/Plugins/Runtime/Harmonix/Source/HarmonixMetasound/Public/HarmonixMetasound/MidiOps/PulseGenerator.h
+++ b/Engine/Plugins/Runtime/Harmonix/Source/HarmonixMetasound/Public/HarmonixMetasound/MidiOps/PulseGenerator.h
@@ -19,8 +19,6 @@ namespace Harmonix::Midi::Ops
 		
 		void Enable(bool bEnable);
 
-		void SetClock(const TSharedPtr<const HarmonixMetasound::FMidiClock, ESPMode::NotThreadSafe>& NewClock);
-		
 		void SetInterval(const FMusicTimeInterval& NewInterval);
 		FMusicTimeInterval GetInterval() const { return Interval; }
 
@@ -32,11 +30,11 @@ namespace Harmonix::Midi::Ops
 			int32 Tick;
 		};
 		
-		void Process(const TFunctionRef<void(const FPulseInfo&)>& OnPulse);
+		void Process(const HarmonixMetasound::FMidiClock& MidiClock, const TFunctionRef<void(const FPulseInfo&)>& OnPulse);
 
 	protected:
 		bool Enabled{ true };
-		TWeakPtr<const HarmonixMetasound::FMidiClock, ESPMode::NotThreadSafe> Clock;
+
 		FMusicTimeInterval Interval{};
 		FTimeSignature CurrentTimeSignature{};
 		FMusicTimestamp NextPulseTimestamp{ -1, -1 };
@@ -57,7 +55,7 @@ namespace Harmonix::Midi::Ops
 
 		virtual void Reset() override;
 
-		void Process(HarmonixMetasound::FMidiStream& OutStream);
+		void Process(const HarmonixMetasound::FMidiClock& MidiClock, HarmonixMetasound::FMidiStream& OutStream);
 
 	private:
 		void AddPulseNote(const int32 BlockFrameIndex, const int32 PulseTick, HarmonixMetasound::FMidiStream& OutStream);
diff --git a/Engine/Plugins/Runtime/Harmonix/Source/HarmonixMetasoundTests/Private/MidiOps/PulseGeneratorTests.cpp b/Engine/Plugins/Runtime/Harmonix/Source/HarmonixMetasoundTests/Private/MidiOps/PulseGeneratorTests.cpp
--- a/Engine/Plugins/Runtime/Harmonix/Source/HarmonixMetasoundTests/Private/MidiOps/PulseGeneratorTests.cpp
+++ b/Engine/Plugins/Runtime/Harmonix/Source/HarmonixMetasoundTests/Private/MidiOps/PulseGeneratorTests.cpp
@@ -25,8 +25,6 @@ namespace Harmonix::Midi::Ops::Tests
 		
 		HarmonixMetasound::FMidiStream OutputStream;
 
-		PulseGenerator.SetClock(Clock);
-
 		// Default: a pulse every beat
 		{
 			constexpr int32 NotesUntilWeAreSatisfiedThisWorks = 23;
@@ -62,7 +60,7 @@ namespace Harmonix::Midi::Ops::Tests
 
 				// Process, which will pop the next notes
 				OutputStream.PrepareBlock();
-				PulseGenerator.Process(OutputStream);
+				PulseGenerator.Process(*Clock, OutputStream);
 				
 				// If this is a block where we should get a pulse, check that we got it
 				if (Clock->GetLastProcessedMidiTick() >= Clock->GetSongMapEvaluator().MusicTimestampToTick(NextPulse))
@@ -147,8 +145,6 @@ namespace Harmonix::Midi::Ops::Tests
 		const auto Clock = MakeShared<HarmonixMetasound::FMidiClock, ESPMode::NotThreadSafe>(OperatorSettings);
 		Clock->AttachToSongMapEvaluator(MakeShared<FSongMaps>(Tempo, TimeSignature.Numerator, TimeSignature.Denominator));
 
-		PulseGenerator.SetClock(Clock);
-
 		// Advance forward a few pulses, then seek back and ensure we keep getting notes
 		{
 			constexpr int32 PulsesToDo = 10;
@@ -173,7 +169,7 @@ namespace Harmonix::Midi::Ops::Tests
 
 				// Process
 				TArray<FPulseGenerator::FPulseInfo> PulsesThisBlock;
-				PulseGenerator.Process([&PulsesThisBlock](const FPulseGenerator::FPulseInfo& Pulse)
+				PulseGenerator.Process(*Clock, [&PulsesThisBlock](const FPulseGenerator::FPulseInfo& Pulse)
 				{
 					PulsesThisBlock.Add(Pulse);
 				});
@@ -206,7 +202,7 @@ namespace Harmonix::Midi::Ops::Tests
 
 				// Process
 				TArray<FPulseGenerator::FPulseInfo> PulsesThisBlock;
-				PulseGenerator.Process([&PulsesThisBlock](const FPulseGenerator::FPulseInfo& Pulse)
+				PulseGenerator.Process(*Clock, [&PulsesThisBlock](const FPulseGenerator::FPulseInfo& Pulse)
 				{
 					PulsesThisBlock.Add(Pulse);
 				});
diff --git a/Engine/Plugins/Runtime/Harmonix/Source/HarmonixMetasoundTests/Private/Nodes/MidiClockSubdivisionTriggerNodeTests.cpp b/Engine/Plugins/Runtime/Harmonix/Source/HarmonixMetasoundTests/Private/Nodes/MidiClockSubdivisionTriggerNodeTests.cpp
--- a/Engine/Plugins/Runtime/Harmonix/Source/HarmonixMetasoundTests/Private/Nodes/MidiClockSubdivisionTriggerNodeTests.cpp
+++ b/Engine/Plugins/Runtime/Harmonix/Source/HarmonixMetasoundTests/Private/Nodes/MidiClockSubdivisionTriggerNodeTests.cpp
@@ -26,7 +26,6 @@ namespace HarmonixMetasound::Nodes::MidiClockSubdivisionTriggerNode::Tests
 		(*Clock)->SetTransportState(0, EMusicPlayerTransportState::Playing);
 
 		Harmonix::Midi::Ops::FPulseGenerator PulseGenerator;
-		PulseGenerator.SetClock((*Clock)->AsShared());
 
 		// Render for a bit and expect the same output from both the node and the raw processor
 		constexpr int32 NumBlocks = 1000;
@@ -42,7 +41,7 @@ namespace HarmonixMetasound::Nodes::MidiClockSubdivisionTriggerNode::Tests
 			Generator->OnGenerateAudio(Buffer.GetData(), Buffer.Num());
 
 			TArray<int32> PulseGeneratorTriggers;
-			PulseGenerator.Process([&PulseGeneratorTriggers](const Harmonix::Midi::Ops::FPulseGenerator::FPulseInfo& Pulse)
+			PulseGenerator.Process(*Clock.GetValue(), [&PulseGeneratorTriggers](const Harmonix::Midi::Ops::FPulseGenerator::FPulseInfo& Pulse)
 			{
 				PulseGeneratorTriggers.Add(Pulse.BlockFrameIndex);
 			});
diff --git a/Engine/Plugins/Runtime/Harmonix/Source/HarmonixMetasoundTests/Private/Nodes/MidiPulseGeneratorNodeTests.cpp b/Engine/Plugins/Runtime/Harmonix/Source/HarmonixMetasoundTests/Private/Nodes/MidiPulseGeneratorNodeTests.cpp
--- a/Engine/Plugins/Runtime/Harmonix/Source/HarmonixMetasoundTests/Private/Nodes/MidiPulseGeneratorNodeTests.cpp
+++ b/Engine/Plugins/Runtime/Harmonix/Source/HarmonixMetasoundTests/Private/Nodes/MidiPulseGeneratorNodeTests.cpp
@@ -26,7 +26,6 @@ namespace HarmonixMetasound::Nodes::MidiPulseGeneratorNode::Tests
 
 		Harmonix::Midi::Ops::FMidiPulseGenerator PulseGenerator;
 		FMidiStream PulseGeneratorMidiOutput;
-		PulseGenerator.SetClock((*Clock)->AsShared());
 
 		// Render for a bit and expect the same output from both the node and the raw processor
 		constexpr int32 NumBlocks = 1000;
@@ -41,7 +40,7 @@ namespace HarmonixMetasound::Nodes::MidiPulseGeneratorNode::Tests
 			// Process
 			Generator->OnGenerateAudio(Buffer.GetData(), Buffer.Num());
 			PulseGeneratorMidiOutput.PrepareBlock();
-			PulseGenerator.Process(PulseGeneratorMidiOutput);
+			PulseGenerator.Process(*Clock.GetValue(), PulseGeneratorMidiOutput);
 
 			// If there are notes in the pulse generator output, expect them in the node output
 			const TArray<FMidiStreamEvent> NodeEvents = (*NodeMidiOutput)->GetEventsInBlock();
