Fix bug/typo in BarMap

#tests Harmonix Tests
#rnx

diff --git a/Engine/Plugins/Runtime/Harmonix/Source/HarmonixMidi/Private/HarmonixMidi/BarMap.cpp b/Engine/Plugins/Runtime/Harmonix/Source/HarmonixMidi/Private/HarmonixMidi/BarMap.cpp
--- a/Engine/Plugins/Runtime/Harmonix/Source/HarmonixMidi/Private/HarmonixMidi/BarMap.cpp
+++ b/Engine/Plugins/Runtime/Harmonix/Source/HarmonixMidi/Private/HarmonixMidi/BarMap.cpp
@@ -205,18 +205,18 @@ int32 FBarMap::BarBeatTickIncludingCountInToTick(int32 BarIndex, int32 BeatInBar
 	}
 
 	int32 BarDelta = BarIndex - Points[TimeSigIndex].BarIndex;
-	int32 TicksInBeat = GetTicksInBeatAfterPoint(TimeSigIndex);
+	int32 TicksPerBeat = GetTicksInBeatAfterPoint(TimeSigIndex);
 	if (BeatInBar >= Points[TimeSigIndex].TimeSignature.Numerator)
 	{
 		UE_LOG(LogMIDI, Warning, TEXT("BarBeatTickToAbsoluteTick: Supplied 'BeatInBar' is greater than the number of beats in the specified bar!"));
 	}
-	if (TickInBeat >= TicksInBeat)
+	if (TickInBeat >= TicksPerBeat)
 	{
 		UE_LOG(LogMIDI, Warning, TEXT("BarBeatTickToAbsoluteTick: Supplied 'TickInBeat' is greater than the number of ticks in each beat at the specified bar!"));
 	}
 	return Points[TimeSigIndex].StartTick +
 			(BarDelta * GetTicksInBarAfterPoint(TimeSigIndex)) +
-			(BeatInBar * TickInBeat) +
+			(BeatInBar * TicksPerBeat) +
 			TickInBeat;
 }
 
