Add tests to MidiClock for SongMaps method

#rnx

diff --git a/Engine/Plugins/Runtime/Harmonix/Source/HarmonixMetasoundTests/Private/DataTypes/MidiClock.spec.cpp b/Engine/Plugins/Runtime/Harmonix/Source/HarmonixMetasoundTests/Private/DataTypes/MidiClock.spec.cpp
--- a/Engine/Plugins/Runtime/Harmonix/Source/HarmonixMetasoundTests/Private/DataTypes/MidiClock.spec.cpp
+++ b/Engine/Plugins/Runtime/Harmonix/Source/HarmonixMetasoundTests/Private/DataTypes/MidiClock.spec.cpp
@@ -541,6 +541,52 @@ namespace HarmonixMetasoundTests::MidiClock
 
 		Describe("Time signature Changes", [this]()
 		{
+			It("BarBeatTickIncludingCountInToTick", [this]()
+			{
+				SongMaps = MakeShared<FSongMaps>(123, 4, 4);
+
+				{
+					// calculate the test tick
+					int32 ExpectedTick = 10;
+
+					// make sure this method returns the correct value
+					const int32 Tick = SongMaps->BarBeatTickIncludingCountInToTick(0, 1, 10);
+
+								
+					if (!TestEqual("Tick at Bar 1 Beat 1 Tick 10", Tick , ExpectedTick))
+					{
+						return;
+					}
+				}
+
+				{
+					// calculate the test tick
+					int32 ExpectedTick = 10 + SongMaps->GetTicksPerQuarterNote();
+
+					// make sure this method returns the correct value
+					const int32 Tick = SongMaps->BarBeatTickIncludingCountInToTick(0, 2, 10);
+
+					if (!TestEqual("Tick at Bar 1 Beat 2 Tick 10", Tick , ExpectedTick))
+					{
+						return;
+					}
+				}
+
+				{
+					// calculate the test tick
+					int32 ExpectedTick = 10 + 2 * SongMaps->GetTicksPerQuarterNote();
+
+					// make sure this method returns the correct value
+					const int32 Tick = SongMaps->BarBeatTickIncludingCountInToTick(0, 3, 10);
+
+				
+					if (!TestEqual("Tick at Bar 1 Beat 3 Tick 10", Tick , ExpectedTick))
+					{
+						return;
+					}
+				}
+			});
+			
 			It("Without driving clock - One Change", [this]()
 			{
 				SongMaps = MakeShared<FSongMaps>(123, 5, 8);
