[Harmonix] Ensure step sequence note off numbers match their note on

diff --git a/Engine/Plugins/Runtime/Harmonix/Source/HarmonixMetasound/Private/HarmonixMetasound/Nodes/StepSequencePlayerNode.cpp b/Engine/Plugins/Runtime/Harmonix/Source/HarmonixMetasound/Private/HarmonixMetasound/Nodes/StepSequencePlayerNode.cpp
--- a/Engine/Plugins/Runtime/Harmonix/Source/HarmonixMetasound/Private/HarmonixMetasound/Nodes/StepSequencePlayerNode.cpp
+++ b/Engine/Plugins/Runtime/Harmonix/Source/HarmonixMetasound/Private/HarmonixMetasound/Nodes/StepSequencePlayerNode.cpp
@@ -125,6 +125,7 @@ namespace HarmonixMetasound::Nodes::StepSequencePlayer
 		bool  bNeedsRebase				    = false;
 		EStepSequencePlayerState PlayState  = EStepSequencePlayerState::NotPlaying;
 		TArray<FMidiVoiceId> CurrentCellNotes;
+		TArray<uint8> CurrentTransposedNotes;
 
 		void CheckForUpdatedSequenceTable();
 		void ResizeCellStatesForTable();
@@ -342,6 +343,7 @@ namespace HarmonixMetasound::Nodes::StepSequencePlayer
 		bNeedsRebase = false;
 		PlayState = EStepSequencePlayerState::NotPlaying;
 		CurrentCellNotes.Reset();
+		CurrentTransposedNotes.Reset();
 	}
 
 	void FStepSequencePlayerOperator::Execute()
@@ -358,6 +360,7 @@ namespace HarmonixMetasound::Nodes::StepSequencePlayer
 			{
 				AllNotesOff(0, MidiClockInPin->GetLastProcessedMidiTick(), true);
 				CurrentCellNotes.Empty();
+				CurrentTransposedNotes.Empty();
 			}
 			return;
 		}
@@ -832,6 +835,7 @@ namespace HarmonixMetasound::Nodes::StepSequencePlayer
 
 						// note off!
 						FMidiStreamEvent MidiEvent(CurrentCellNotes[NoteIdx].GetGeneratorId(), FMidiMsg::CreateNoteOff(MidiChannel, NoteIdx));
+						MidiEvent.MidiMessage.Data1 = CurrentTransposedNotes[NoteIdx];
 						MidiEvent.BlockSampleFrameIndex = BlockFrameIndex;
 						MidiEvent.AuthoredMidiTick = ProcessedThruTick;
 						MidiEvent.CurrentMidiTick = ProcessedThruTick;
@@ -839,6 +843,7 @@ namespace HarmonixMetasound::Nodes::StepSequencePlayer
 						MidiOutPin->AddNoteOffEventOrCancelPendingNoteOn(MidiEvent);
 						UE_LOG(LogStepSequencePlayer, Verbose, TEXT("0x%x Note-Off %d at %d"), (uint32)(size_t)this, SequenceTable->Notes[NoteIdx].NoteNumber + AdditionalOctaveNotes, BlockFrameIndex);
 						CurrentCellNotes[NoteIdx] = FMidiVoiceId::None();
+						CurrentTransposedNotes[NoteIdx] = 0u;
 					}
 				}
 			}
@@ -867,6 +872,7 @@ namespace HarmonixMetasound::Nodes::StepSequencePlayer
 							MidiOutPin->AddMidiEvent(MidiEvent);
 							UE_LOG(LogStepSequencePlayer, Verbose, TEXT("0x%x Note-On %d at %d"), (uint32)(size_t)this, SequenceTable->Notes[NoteIdx].NoteNumber + AdditionalOctaveNotes, BlockFrameIndex);
 							CurrentCellNotes[NoteIdx] = MidiEvent.GetVoiceId();
+							CurrentTransposedNotes[NoteIdx] = TransposedNote;
 						}
 					}
 				}
@@ -929,12 +935,14 @@ namespace HarmonixMetasound::Nodes::StepSequencePlayer
 				}
 			}
 			CurrentCellNotes.SetNum(SequenceTable->Notes.Num());
+			CurrentTransposedNotes.SetNum(SequenceTable->Notes.Num());
 		}
 		else
 		{
 			while (CurrentCellNotes.Num() < SequenceTable->Notes.Num())
 			{
 				CurrentCellNotes.Add(FMidiVoiceId());
+				CurrentTransposedNotes.Add(0u);
 			}
 		}
 	}
@@ -953,6 +961,7 @@ namespace HarmonixMetasound::Nodes::StepSequencePlayer
 				MidiOutPin->AddMidiEvent(MidiEvent);
 				UE_LOG(LogStepSequencePlayer, Verbose, TEXT("0x%x Note-Off %d (during all notes off)"), (uint32)(size_t)this, NoteIdx);
 				CurrentCellNotes[NoteIdx] = FMidiVoiceId::None();
+				CurrentTransposedNotes[NoteIdx] = 0u;
 			}
 		}
 		if (ResetCellIndex)
