Cache SharedPtr to MidiClock outside of the "DoRender" method to avoid race conditions creating it within multi-threaded rendering.

#tests PIE and Editor Tests
#rnx

diff --git a/Engine/Plugins/Runtime/Harmonix/Source/HarmonixMetasound/Private/HarmonixMetasound/Nodes/FusionSamplerNode.cpp b/Engine/Plugins/Runtime/Harmonix/Source/HarmonixMetasound/Private/HarmonixMetasound/Nodes/FusionSamplerNode.cpp
--- a/Engine/Plugins/Runtime/Harmonix/Source/HarmonixMetasound/Private/HarmonixMetasound/Nodes/FusionSamplerNode.cpp
+++ b/Engine/Plugins/Runtime/Harmonix/Source/HarmonixMetasound/Private/HarmonixMetasound/Nodes/FusionSamplerNode.cpp
@@ -229,6 +229,8 @@ namespace HarmonixMetasound
 		EAudioBufferChannelLayout OutputChannelLayout = EAudioBufferChannelLayout::Mono;
 		Harmonix::Midi::Ops::FStuckNoteGuard StuckNoteGuard;
 
+		TSharedPtr<const FMidiClock, ESPMode::NotThreadSafe> MidiClock;
+
 		void SetPatchAndOverrides(FFusionPatchDataProxy::NodePtr NewPatch)
 		{
 			KillAllVoices();
@@ -580,7 +582,9 @@ namespace HarmonixMetasound
 		UpdatePatchOverrides();
 
 		SetRawTransposition(*TranspositionInPin);
-
+		
+		MidiClock = MidiStreamInPin->GetClock();
+		
 		if (EnableMTRenderingInPin && !bDisableMultithreadedRender)
 		{
 			SyncLinkOutPin->KickAsyncRender([this]()
@@ -619,9 +623,6 @@ namespace HarmonixMetasound
 		// create an iterator for midi events in the block
 		const TArray<FMidiStreamEvent>& MidiEvents = MidiStreamInPin->GetEventsInBlock();
 		auto MidiEventIterator = MidiEvents.begin();
-
-		// create an iterator for the midi clock 
-		const TSharedPtr<const FMidiClock, ESPMode::NotThreadSafe> MidiClock = MidiStreamInPin->GetClock();
 		
 		while (FramesRequired > 0)
 		{
