[Harmonix] Update song maps when the midi clock used by a MidiSongPosVertexAnalyzer changes
#tests cooked build

diff --git a/Engine/Plugins/Runtime/Harmonix/Source/HarmonixMetasound/Private/HarmonixMetasound/Analysis/MidiSongPosVertexAnalyzer.cpp b/Engine/Plugins/Runtime/Harmonix/Source/HarmonixMetasound/Private/HarmonixMetasound/Analysis/MidiSongPosVertexAnalyzer.cpp
--- a/Engine/Plugins/Runtime/Harmonix/Source/HarmonixMetasound/Private/HarmonixMetasound/Analysis/MidiSongPosVertexAnalyzer.cpp
+++ b/Engine/Plugins/Runtime/Harmonix/Source/HarmonixMetasound/Private/HarmonixMetasound/Analysis/MidiSongPosVertexAnalyzer.cpp
@@ -59,6 +59,7 @@ namespace HarmonixMetasound::Analysis
 		TSharedPtr<const FSongMapChain> MapChain = History->GetLatestMapsForProducer();
 
 		if (!MapChain->SongMaps || Clock.GetSongMapsChangedInBlock() ||
+			(LastClock.IsValid() && LastClock != Clock.AsWeak()) ||
 			Clock.GetFirstTickInLoop() != MapChain->FirstTickInLoop ||
 			Clock.GetLoopLengthTicks() != MapChain->LoopLengthTicks)
 		{
@@ -73,6 +74,8 @@ namespace HarmonixMetasound::Analysis
 			History->UpdateMaps(Maps, Clock.GetFirstTickInLoop(), Clock.GetLoopLengthTicks());
 		}
 
+		LastClock = Clock.AsWeak();
+
 		LastMidiClockSongPos->CurrentSpeed = Clock.GetSpeedAtEndOfBlock();
 		LastMidiClockSongPos->CurrentTransportState = Clock.GetTransportStateAtStartOfBlock();
 
diff --git a/Engine/Plugins/Runtime/Harmonix/Source/HarmonixMetasound/Public/HarmonixMetasound/Analysis/MidiSongPosVertexAnalyzer.h b/Engine/Plugins/Runtime/Harmonix/Source/HarmonixMetasound/Public/HarmonixMetasound/Analysis/MidiSongPosVertexAnalyzer.h
--- a/Engine/Plugins/Runtime/Harmonix/Source/HarmonixMetasound/Public/HarmonixMetasound/Analysis/MidiSongPosVertexAnalyzer.h
+++ b/Engine/Plugins/Runtime/Harmonix/Source/HarmonixMetasound/Public/HarmonixMetasound/Analysis/MidiSongPosVertexAnalyzer.h
@@ -41,6 +41,7 @@ namespace HarmonixMetasound::Analysis
 		virtual void Execute() override;
 
 	private:
+		TWeakPtr<const FMidiClock, ESPMode::NotThreadSafe> LastClock;
 		FMidiClockSongPositionWriteRef LastMidiClockSongPos;
 		TSharedPtr<FMidiClockSongPositionHistory> History;
 		bool LastPosWasSeekOrLoop {false};
