[UE-226551] Fix Crash in Harmonix::Midi::Constants::SubdivisionToBeats()

No need to "assert" when passing in an invalid Quantization enum. We return "0" which can be handled gracefully.

However, handling this "gracefully" requires a bit of extra work. The PulseGenerator needs to not send (or stop sending) pulses when then Interval is invalid. Also, if the interval changes from valid to invalid, we need to stop active pulses.

Also, handle some edge cases if the Time Signature is invalid as well.

#tests Harmonix Automation Tests & Editor Tests
#rnx

diff --git a/Engine/Plugins/Runtime/Harmonix/Source/HarmonixMetasound/Private/HarmonixMetasound/DataTypes/MusicTimeInterval.cpp b/Engine/Plugins/Runtime/Harmonix/Source/HarmonixMetasound/Private/HarmonixMetasound/DataTypes/MusicTimeInterval.cpp
--- a/Engine/Plugins/Runtime/Harmonix/Source/HarmonixMetasound/Private/HarmonixMetasound/DataTypes/MusicTimeInterval.cpp
+++ b/Engine/Plugins/Runtime/Harmonix/Source/HarmonixMetasound/Private/HarmonixMetasound/DataTypes/MusicTimeInterval.cpp
@@ -43,6 +43,12 @@ namespace Harmonix
 	float GetOffsetInBeats(const FMusicTimeInterval& Interval, const FTimeSignature& TimeSignature)
 	{
 		const float IntervalBeats = GetIntervalInBeats(Interval, TimeSignature);
+		
+		if (FMath::IsNearlyZero(IntervalBeats))
+		{
+			return 0.0f;
+		}
+		
 		const float OffsetBeats = Harmonix::Midi::Constants::SubdivisionToBeats(Interval.Offset, TimeSignature) * Interval.OffsetMultiplier;
 		return FMath::Fmod(OffsetBeats, IntervalBeats);
 	}
diff --git a/Engine/Plugins/Runtime/Harmonix/Source/HarmonixMetasound/Private/HarmonixMetasound/MidiOps/PulseGenerator.cpp b/Engine/Plugins/Runtime/Harmonix/Source/HarmonixMetasound/Private/HarmonixMetasound/MidiOps/PulseGenerator.cpp
--- a/Engine/Plugins/Runtime/Harmonix/Source/HarmonixMetasound/Private/HarmonixMetasound/MidiOps/PulseGenerator.cpp
+++ b/Engine/Plugins/Runtime/Harmonix/Source/HarmonixMetasound/Private/HarmonixMetasound/MidiOps/PulseGenerator.cpp
@@ -24,17 +24,21 @@ namespace Harmonix::Midi::Ops
 
 	void FPulseGenerator::Process(const HarmonixMetasound::FMidiClock& MidiClock, const TFunctionRef<void(const FPulseInfo&)>& OnPulse)
 	{
+		bool bIntervalIsValid = Interval.Interval != EMidiClockSubdivisionQuantization::None;
+		bool bTimeSigIsValid = CurrentTimeSignature.Numerator > 0 && CurrentTimeSignature.Denominator > 0;
+		
 		// ensure the pulse generator is lined up with the current clock phase
 		if (NextPulseTimestamp == FMusicTimestamp(-1, -1))
 		{
 			CurrentTimeSignature = *MidiClock.GetSongMapEvaluator().GetTimeSignatureAtTick(MidiClock.GetLastProcessedMidiTick());
-
+			bTimeSigIsValid = CurrentTimeSignature.Numerator > 0 && CurrentTimeSignature.Denominator > 0;
+			
 			// Find the next pulse and line up phase with the current bar
 			const FMusicTimestamp ClockCurrentTimestamp = MidiClock.GetMusicTimestampAtBlockOffset(0);
 			NextPulseTimestamp.Bar = ClockCurrentTimestamp.Bar;
 			NextPulseTimestamp.Beat = 1;
 			IncrementTimestampByOffset(NextPulseTimestamp, Interval, CurrentTimeSignature);
-			while (NextPulseTimestamp < ClockCurrentTimestamp)
+			while (bIntervalIsValid && bTimeSigIsValid && NextPulseTimestamp < ClockCurrentTimestamp)
 			{
 				IncrementTimestampByInterval(NextPulseTimestamp, Interval, CurrentTimeSignature);
 			}
@@ -54,7 +58,7 @@ namespace Harmonix::Midi::Ops
 				
 				int32 NextPulseTick = MidiClock.GetSongMapEvaluator().MusicTimestampToTick(NextPulseTimestamp);
 
-				while (AsAdvance->LastTickToProcess() >= NextPulseTick)
+				while (bIntervalIsValid && bTimeSigIsValid && AsAdvance->LastTickToProcess() >= NextPulseTick)
 				{
 					OnPulse({ ClockEvent.BlockFrameIndex, NextPulseTick });
 
@@ -66,6 +70,7 @@ namespace Harmonix::Midi::Ops
 			else if (const FTimeSignatureChange* AsTimeSigChange = ClockEvent.TryGet<FTimeSignatureChange>())
 			{
 				CurrentTimeSignature = AsTimeSigChange->TimeSignature;
+				bTimeSigIsValid = CurrentTimeSignature.Numerator > 0 && CurrentTimeSignature.Denominator > 0;
 				
 				// Time sig changes will come on the downbeat, and if we change time signature,
 				// we want to reset the pulse, so the next pulse is now plus the offset
@@ -85,7 +90,7 @@ namespace Harmonix::Midi::Ops
 				NextPulseTimestamp.Bar = ClockCurrentTimestamp.Bar;
 				NextPulseTimestamp.Beat = 1;
 				IncrementTimestampByOffset(NextPulseTimestamp, Interval, CurrentTimeSignature);
-				while (NextPulseTimestamp < ClockCurrentTimestamp)
+				while (bIntervalIsValid && bTimeSigIsValid && NextPulseTimestamp < ClockCurrentTimestamp)
 				{
 					IncrementTimestampByInterval(NextPulseTimestamp, Interval, CurrentTimeSignature);
 				}
@@ -104,6 +109,26 @@ namespace Harmonix::Midi::Ops
 	{
 		OutStream.PrepareBlock();
 
+		// kill any notes if the interval becomes invalid
+		if (GetInterval().Interval == EMidiClockSubdivisionQuantization::None && LastNoteOn.IsSet())
+		{
+			check(LastNoteOn->MidiMessage.IsNoteOn());
+			const int32 NoteOffSample = 0;
+
+			// Trigger the note off one tick before the note on
+			const int32 NoteOffTick = MidiClock.GetNextTickToProcessAtBlockFrame(0);
+
+			FMidiMsg Msg{ FMidiMsg::CreateAllNotesOff() };
+			HarmonixMetasound::FMidiStreamEvent Event{ &VoiceGenerator, Msg };
+			Event.BlockSampleFrameIndex = NoteOffSample;
+			Event.AuthoredMidiTick = NoteOffTick;
+			Event.CurrentMidiTick = NoteOffTick;
+			Event.TrackIndex = 1;
+			OutStream.InsertMidiEvent(Event);
+
+			LastNoteOn.Reset();	
+		}
+
 		FPulseGenerator::Process(MidiClock, [this, &OutStream](const FPulseInfo& Pulse)
 		{
 			AddPulseNote(Pulse.BlockFrameIndex, Pulse.Tick, OutStream);
diff --git a/Engine/Plugins/Runtime/Harmonix/Source/HarmonixMidi/Private/HarmonixMidi/MidiConstants.cpp b/Engine/Plugins/Runtime/Harmonix/Source/HarmonixMidi/Private/HarmonixMidi/MidiConstants.cpp
--- a/Engine/Plugins/Runtime/Harmonix/Source/HarmonixMidi/Private/HarmonixMidi/MidiConstants.cpp
+++ b/Engine/Plugins/Runtime/Harmonix/Source/HarmonixMidi/Private/HarmonixMidi/MidiConstants.cpp
@@ -81,7 +81,6 @@ namespace Harmonix::Midi::Constants
 		case EMidiClockSubdivisionQuantization::HalfNoteTriplet:
 			return BeatsPerQuarter * 4 / 3;
 		default:
-			checkNoEntry();
 			return 0;
 		}
 	}
