MidiFileFactory returns nullptr if MidiFile fails to load.

Update MidiFile api to return `ELoadResult`

#tests Harmonix Editor Tests
#rnx

diff --git a/Engine/Plugins/Runtime/Harmonix/Source/HarmonixMidi/Private/HarmonixMidi/MidiFile.cpp b/Engine/Plugins/Runtime/Harmonix/Source/HarmonixMidi/Private/HarmonixMidi/MidiFile.cpp
--- a/Engine/Plugins/Runtime/Harmonix/Source/HarmonixMidi/Private/HarmonixMidi/MidiFile.cpp
+++ b/Engine/Plugins/Runtime/Harmonix/Source/HarmonixMidi/Private/HarmonixMidi/MidiFile.cpp
@@ -130,7 +130,7 @@ void UMidiFile::PostInitProperties()
 #endif
 }
 
-void UMidiFile::LoadStdMidiFile(
+UMidiFile::ELoadResult UMidiFile::LoadStdMidiFile(
 	const FString& FilePath,
 	int32 DesiredTicksPerQuarterNote,
 	Harmonix::Midi::Constants::EMidiTextEventEncoding InTextEncoding,
@@ -140,8 +140,13 @@ void UMidiFile::LoadStdMidiFile(
 	IPlatformFile& PlatformFileApi = FPlatformFileManager::Get().GetPlatformFile();
 	IFileHandle* FileHandle = PlatformFileApi.OpenRead(*FilePath);
 	TSharedPtr<FArchiveFileReaderGeneric> Archive = MakeShared<FArchiveFileReaderGeneric>(FileHandle, *Filename, FileHandle->Size());
-	LoadStdMidiFile(Archive, Filename, DesiredTicksPerQuarterNote, InTextEncoding, EnsureNotFailed);
+	const ELoadResult Result = LoadStdMidiFile(Archive, Filename, DesiredTicksPerQuarterNote, InTextEncoding, EnsureNotFailed);
 
+	if (Result != ELoadResult::Success)
+	{
+		return Result;
+	}
+	
 #if WITH_EDITORONLY_DATA
 	if (!AssetImportData)
 	{
@@ -149,9 +154,11 @@ void UMidiFile::LoadStdMidiFile(
 	}
 	AssetImportData->Update(FilePath);
 #endif
+
+	return Result;
 }
 
-void UMidiFile::LoadStdMidiFile(
+UMidiFile::ELoadResult UMidiFile::LoadStdMidiFile(
 	void* Buffer,
 	int32 BufferSize,
 	const FString& Filename,
@@ -160,19 +167,23 @@ void UMidiFile::LoadStdMidiFile(
 	bool EnsureNotFailed)
 {
 	TSharedPtr<FBufferReader> BufferArchive = MakeShared<FBufferReader>(Buffer, BufferSize, false);
-	LoadStdMidiFile(BufferArchive, Filename, DesiredTicksPerQuarterNote, InTextEncoding, EnsureNotFailed);
+	return LoadStdMidiFile(BufferArchive, Filename, DesiredTicksPerQuarterNote, InTextEncoding, EnsureNotFailed);
 }
 
-void UMidiFile::LoadStdMidiFile(
+UMidiFile::ELoadResult UMidiFile::LoadStdMidiFile(
 	TSharedPtr<FArchive> Archive,
 	const FString& Filename,
 	int32 DesiredTicksPerQuarterNote,
 	Harmonix::Midi::Constants::EMidiTextEventEncoding InTextEncoding,
 	bool EnsureNotFailed)
 {
-	TheMidiData.TicksPerQuarterNote = DesiredTicksPerQuarterNote;
-	TheMidiData.SongMaps.Init(TheMidiData.TicksPerQuarterNote);
-	TheMidiData.Tracks.Empty();
+	// make a copy of the midi data and modify that with the file reader
+	// if the file reader fails, we can retain the original data before overwriting it
+	FMidiFileData MidiDataCopy = TheMidiData;
+	
+	MidiDataCopy.TicksPerQuarterNote = DesiredTicksPerQuarterNote;
+	MidiDataCopy.SongMaps.Init(MidiDataCopy.TicksPerQuarterNote);
+	MidiDataCopy.Tracks.Empty();
 	// Create a receiver list...
 	FMidiReceiverList AllReceivers;
 	// Add the event receiver that will populate the midi tracks with midi events...
@@ -183,7 +194,7 @@ void UMidiFile::LoadStdMidiFile(
 	AllReceivers.Add(&MapReceiver);
 
 	// Create the reader and read all the data...
-	FStdMidiFileReader reader(Archive, Filename, &AllReceivers, TheMidiData.TicksPerQuarterNote, InTextEncoding);
+	FStdMidiFileReader reader(Archive, Filename, &AllReceivers, MidiDataCopy.TicksPerQuarterNote, InTextEncoding);
 	reader.ReadAllTracks();
 
 	if (reader.IsFailed())
@@ -197,11 +208,16 @@ void UMidiFile::LoadStdMidiFile(
 			UE_LOG(LogMIDI, Error, TEXT("%s MIDI import failed. Midi data is malformed. Check the log for details."), *Filename);
 		}
 		TheMidiData.Empty();
+		return ELoadResult::Failed;
 	}
 
+	// update the actual midi data with the loaded data
+	TheMidiData = MidiDataCopy;
+	
 	// we are now "dirty"... so make sure any new requests for renderable data 
 	// get a new copy...
 	RenderableCopyOfMidiFileData = nullptr;
+	return ELoadResult::Success;
 }
 
 void UMidiFile::SaveStdMidiFile(const FString& FilePath) const
diff --git a/Engine/Plugins/Runtime/Harmonix/Source/HarmonixMidi/Public/HarmonixMidi/MidiFile.h b/Engine/Plugins/Runtime/Harmonix/Source/HarmonixMidi/Public/HarmonixMidi/MidiFile.h
--- a/Engine/Plugins/Runtime/Harmonix/Source/HarmonixMidi/Public/HarmonixMidi/MidiFile.h
+++ b/Engine/Plugins/Runtime/Harmonix/Source/HarmonixMidi/Public/HarmonixMidi/MidiFile.h
@@ -120,6 +120,12 @@ class HARMONIXMIDI_API UMidiFile : public UObject, public IAudioProxyDataFactory
 public:
 	using FMidiTrackList = TArray<FMidiTrack>;
 
+	enum class ELoadResult
+	{
+		Failed = 0,
+		Success
+	};
+
 	UMidiFile();
 
 	// A comparison operator. This allows for differences in members related to how
@@ -131,13 +137,13 @@ public:
 	virtual void PostInitProperties() override;
 
 	/** A method for importing a standard midi file, with the option of providing a pointer to an FSongMaps instance that will be populated during the load. */
-	void LoadStdMidiFile(
+	ELoadResult LoadStdMidiFile(
 		const FString& FilePath,
 		int32 DesiredTicksPerQuarterNote = Harmonix::Midi::Constants::GTicksPerQuarterNoteInt,
 		Harmonix::Midi::Constants::EMidiTextEventEncoding InTextEncoding = Harmonix::Midi::Constants::EMidiTextEventEncoding::UTF8,
 		bool EnsureNotFailed = true);
 	/** A method for importing a standard midi file, with the option of providing a pointer to an FSongMaps instance that will be populated during the load. */
-	void LoadStdMidiFile(
+	ELoadResult LoadStdMidiFile(
 		void* Buffer,
 		int32 BufferSize,
 		const FString& FileName,
@@ -145,7 +151,7 @@ public:
 		Harmonix::Midi::Constants::EMidiTextEventEncoding InTextEncoding = Harmonix::Midi::Constants::EMidiTextEventEncoding::UTF8,
 		bool EnsureNotFailed = true);
 	/** A method for importing a standard midi file, with the option of providing a pointer to an FSongMaps instance that will be populated during the load. */
-	void LoadStdMidiFile(
+	ELoadResult LoadStdMidiFile(
 		TSharedPtr<FArchive> Archive,
 		const FString& Filename,
 		int32 DesiredTicksPerQuarterNote = Harmonix::Midi::Constants::GTicksPerQuarterNoteInt,
diff --git a/Engine/Plugins/Runtime/Harmonix/Source/HarmonixMidiEditor/Private/MidiFileFactory.cpp b/Engine/Plugins/Runtime/Harmonix/Source/HarmonixMidiEditor/Private/MidiFileFactory.cpp
--- a/Engine/Plugins/Runtime/Harmonix/Source/HarmonixMidiEditor/Private/MidiFileFactory.cpp
+++ b/Engine/Plugins/Runtime/Harmonix/Source/HarmonixMidiEditor/Private/MidiFileFactory.cpp
@@ -45,7 +45,12 @@ UObject* UMidiFileFactory::FactoryCreateFile(UClass* InClass, UObject* InParent,
 	{
 		MidiFileAsset = NewObject<UMidiFile>(InParent, InClass, InName, Flags);
 	}
-	MidiFileAsset->LoadStdMidiFile(Filename, Harmonix::Midi::Constants::GTicksPerQuarterNoteInt);
+
+	UMidiFile::ELoadResult Result = MidiFileAsset->LoadStdMidiFile(Filename, Harmonix::Midi::Constants::GTicksPerQuarterNoteInt);
+	if (Result != UMidiFile::ELoadResult::Success)
+	{
+		return nullptr;
+	}
 	
 	//add the current file to the file array (for handling multi-batch import)
 	ImportedFiles.Add(MidiFileAsset);
@@ -96,7 +101,12 @@ EReimportResult::Type UMidiFileFactory::Reimport(UObject* Obj)
 		return EReimportResult::Failed;
 	}
 
-	AsMidiFile->LoadStdMidiFile(ReimportPath, Harmonix::Midi::Constants::GTicksPerQuarterNoteInt);
+	UMidiFile::ELoadResult Result = AsMidiFile->LoadStdMidiFile(ReimportPath, Harmonix::Midi::Constants::GTicksPerQuarterNoteInt);
+	if (Result != UMidiFile::ELoadResult::Success)
+	{
+		return EReimportResult::Failed;
+	}
+	
 	ImportedFiles.AddUnique(AsMidiFile);
 
 	return EReimportResult::Succeeded;
