Fix MIDI Files failing to load correctly again.

#rnx
#tests Harmonix Editor Tests

diff --git a/Engine/Plugins/Runtime/Harmonix/Source/HarmonixMidi/Private/HarmonixMidi/MidiFile.cpp b/Engine/Plugins/Runtime/Harmonix/Source/HarmonixMidi/Private/HarmonixMidi/MidiFile.cpp
--- a/Engine/Plugins/Runtime/Harmonix/Source/HarmonixMidi/Private/HarmonixMidi/MidiFile.cpp
+++ b/Engine/Plugins/Runtime/Harmonix/Source/HarmonixMidi/Private/HarmonixMidi/MidiFile.cpp
@@ -177,13 +177,10 @@ UMidiFile::ELoadResult UMidiFile::LoadStdMidiFile(
 	Harmonix::Midi::Constants::EMidiTextEventEncoding InTextEncoding,
 	bool EnsureNotFailed)
 {
-	// make a copy of the midi data and modify that with the file reader
-	// if the file reader fails, we can retain the original data before overwriting it
-	FMidiFileData MidiDataCopy = TheMidiData;
 	
-	MidiDataCopy.TicksPerQuarterNote = DesiredTicksPerQuarterNote;
-	MidiDataCopy.SongMaps.Init(MidiDataCopy.TicksPerQuarterNote);
-	MidiDataCopy.Tracks.Empty();
+	TheMidiData.TicksPerQuarterNote = DesiredTicksPerQuarterNote;
+	TheMidiData.SongMaps.Init(TheMidiData.TicksPerQuarterNote);
+	TheMidiData.Tracks.Empty();
 	// Create a receiver list...
 	FMidiReceiverList AllReceivers;
 	// Add the event receiver that will populate the midi tracks with midi events...
@@ -194,7 +191,7 @@ UMidiFile::ELoadResult UMidiFile::LoadStdMidiFile(
 	AllReceivers.Add(&MapReceiver);
 
 	// Create the reader and read all the data...
-	FStdMidiFileReader reader(Archive, Filename, &AllReceivers, MidiDataCopy.TicksPerQuarterNote, InTextEncoding);
+	FStdMidiFileReader reader(Archive, Filename, &AllReceivers, TheMidiData.TicksPerQuarterNote, InTextEncoding);
 	reader.ReadAllTracks();
 
 	if (reader.IsFailed())
@@ -210,9 +207,6 @@ UMidiFile::ELoadResult UMidiFile::LoadStdMidiFile(
 		TheMidiData.Empty();
 		return ELoadResult::Failed;
 	}
-
-	// update the actual midi data with the loaded data
-	TheMidiData = MidiDataCopy;
 	
 	// we are now "dirty"... so make sure any new requests for renderable data 
 	// get a new copy...
