MidiFileFactory, more robust reimport logic.

Use "ImportObject" so that we can reuse the base import logic without duplicating code.

We have to trigger the prompt for the user about conforming midi ticks within this reimport function because the cleanup logic doesn't get called automatically. Plus, we want to mark the package as dirty after Reimport && Reconforming, otherwise reimported assets won't be conformed.

#tests Harmonix Tests
#rnx

diff --git a/Engine/Plugins/Runtime/Harmonix/Source/HarmonixMidiEditor/Private/MidiFileFactory.cpp b/Engine/Plugins/Runtime/Harmonix/Source/HarmonixMidiEditor/Private/MidiFileFactory.cpp
--- a/Engine/Plugins/Runtime/Harmonix/Source/HarmonixMidiEditor/Private/MidiFileFactory.cpp
+++ b/Engine/Plugins/Runtime/Harmonix/Source/HarmonixMidiEditor/Private/MidiFileFactory.cpp
@@ -101,14 +101,35 @@ EReimportResult::Type UMidiFileFactory::Reimport(UObject* Obj)
 		return EReimportResult::Failed;
 	}
 
-	UMidiFile::ELoadResult Result = AsMidiFile->LoadStdMidiFile(ReimportPath, Harmonix::Midi::Constants::GTicksPerQuarterNoteInt);
-	if (Result != UMidiFile::ELoadResult::Success)
+	bool OutCanceled = false;
+	UObject* ImportedObject = ImportObject(AsMidiFile->GetClass(), AsMidiFile->GetOuter(), *AsMidiFile->GetName(), RF_Public | RF_Standalone, ReimportPath, nullptr, OutCanceled);
+	if (ImportedObject && !OutCanceled)
+	{
+		// the imported object should be the existing midi file
+		ensure(ImportedObject == AsMidiFile);
+
+		CheckImportedFilesForValidLengths();
+		
+		if (AsMidiFile->GetOuter())
+		{
+			AsMidiFile->GetOuter()->MarkPackageDirty();
+		}
+		else
+		{
+			AsMidiFile->MarkPackageDirty();
+		}
+	}
+	else if (OutCanceled)
+	{
+		UE_LOG(LogHarmonixMidiEditor, Warning, TEXT("import canceled"));
+		return EReimportResult::Cancelled;
+	}
+	else
 	{
+		UE_LOG(LogHarmonixMidiEditor, Warning, TEXT("import failed"));
 		return EReimportResult::Failed;
 	}
 	
-	ImportedFiles.AddUnique(AsMidiFile);
-
 	return EReimportResult::Succeeded;
 }
 
