MidiFileData Validation

Checks if MidiFile data is Empty

#tests Harmonix Tests
#rnx

diff --git a/Engine/Plugins/Runtime/Harmonix/Source/HarmonixMidi/Private/HarmonixMidi/MidiFile.cpp b/Engine/Plugins/Runtime/Harmonix/Source/HarmonixMidi/Private/HarmonixMidi/MidiFile.cpp
--- a/Engine/Plugins/Runtime/Harmonix/Source/HarmonixMidi/Private/HarmonixMidi/MidiFile.cpp
+++ b/Engine/Plugins/Runtime/Harmonix/Source/HarmonixMidi/Private/HarmonixMidi/MidiFile.cpp
@@ -15,12 +15,17 @@
 #if WITH_EDITORONLY_DATA
 #include "EditorFramework/AssetImportData.h"
 #endif
+#if WITH_EDITOR
+#include "Misc/DataValidation.h"
+#endif
 
 //////////////////////////////////////////////////////////////
 //
 // private helper class for reading a std midi file into a UMidiFile
 //
 
+#define LOCTEXT_NAMESPACE "Harmonix_MidiFile"
+
 class FMidiEventReceiver : public IMidiReceiver
 {
 public:
@@ -571,6 +576,20 @@ void UMidiFile::PostEditChangeProperty(FPropertyChangedEvent& PropertyChangedEve
 }
 #endif
 
+#if WITH_EDITOR
+
+EDataValidationResult UMidiFile::IsDataValid(FDataValidationContext& Context) const
+{
+	if (TheMidiData.IsEmpty())
+	{
+		Context.AddError(LOCTEXT("MidiFile_DatIsEmpty", "MidiFile data is empty"));
+		return EDataValidationResult::Invalid;
+	}
+	return EDataValidationResult::Valid;
+}
+
+#endif
+
 FMidiEventReceiver::FMidiEventReceiver(UMidiFile& file)
 	: File(file)
 {
@@ -988,3 +1007,5 @@ bool FMidiFileData::LengthIsAPerfectSubdivision() const
 	return SongMaps.LengthIsAPerfectSubdivision();
 }
 
+
+#undef LOCTEXT_NAMESPACE // Harmonix_MidiFile
\ No newline at end of file
diff --git a/Engine/Plugins/Runtime/Harmonix/Source/HarmonixMidi/Public/HarmonixMidi/MidiFile.h b/Engine/Plugins/Runtime/Harmonix/Source/HarmonixMidi/Public/HarmonixMidi/MidiFile.h
--- a/Engine/Plugins/Runtime/Harmonix/Source/HarmonixMidi/Public/HarmonixMidi/MidiFile.h
+++ b/Engine/Plugins/Runtime/Harmonix/Source/HarmonixMidi/Public/HarmonixMidi/MidiFile.h
@@ -216,6 +216,12 @@ public:
 
 	//StartBar UPROPERTY callback 
 	virtual void PostEditChangeProperty(FPropertyChangedEvent& PropertyChangedEvent) override;
+#endif
+	
+#if WITH_EDITOR
+
+	virtual EDataValidationResult IsDataValid(class FDataValidationContext& Context) const override;
+
 #endif
 
 	int32 GetStartBar() { return TheMidiData.SongMaps.GetStartBar(); };
