[Engine] Add File Duration in seconds field to the midi file details.

#tests Content Browser in Editor.
#rnx

diff --git a/Engine/Plugins/Runtime/Harmonix/Source/HarmonixMidiEditor/Private/MidiFileDetailCustomization.cpp b/Engine/Plugins/Runtime/Harmonix/Source/HarmonixMidiEditor/Private/MidiFileDetailCustomization.cpp
--- a/Engine/Plugins/Runtime/Harmonix/Source/HarmonixMidiEditor/Private/MidiFileDetailCustomization.cpp
+++ b/Engine/Plugins/Runtime/Harmonix/Source/HarmonixMidiEditor/Private/MidiFileDetailCustomization.cpp
@@ -49,6 +49,23 @@ void FMidiFileDetailCustomization::CustomizeDetails(IDetailLayoutBuilder& Detail
 
 	BuildLengthRow(MidiFileBeingEdited.Get());
 
+	//Show FileDuration in seconds
+	FDetailWidgetRow* DurationRow = &MidiFileLengthCategory.AddCustomRow(FText::FromString("Midi File Duration"));
+	DurationRow->NameContent()
+		[
+			SNew(STextBlock)
+			.Text(FText::FromString("File Duration"))
+			.Font(IDetailLayoutBuilder::GetDetailFont())
+		]; 
+	DurationRow->ValueContent()
+	[
+		SNew(STextBlock)
+			.Text(FText::Format(LOCTEXT("MIDIFileDuration", "{0} secs"),  MidiFileBeingEdited->GetSongMaps()->GetSongLengthMs() / 1000.0f))
+			.Font(IDetailLayoutBuilder::GetDetailFont())
+			.ColorAndOpacity(FLinearColor::White)
+			.ToolTipText(LOCTEXT("MIDIFileDurationTooltip", "Duration of the midi file in seconds."))
+	];
+	
 	IDetailCategoryBuilder& TrackInfoCategory = DetailLayout.EditCategory(TEXT("Track Info"), LOCTEXT("TrackInfo", "Tracks:"));
 	const UMidiFile::FMidiTrackList& Tracks = MidiFileBeingEdited->GetTracks();
 	for (const FMidiTrack& Track : Tracks)
@@ -101,6 +118,8 @@ void FMidiFileDetailCustomization::CustomizeDetails(IDetailLayoutBuilder& Detail
 
 void FMidiFileDetailCustomization::BuildLengthRow(UMidiFile* TheMidiFile)
 {
+	auto const RoundedLengthToolTipText = LOCTEXT("MIDIFileLengthTooltip", "Length of the midi file in bars | beats."); 
+	
 	bool bLengthIsAPerfectSubdivision = TheMidiFile->GetSongMaps()->LengthIsAPerfectSubdivision();
 	if (bLengthIsAPerfectSubdivision)
 	{
diff --git a/Engine/Plugins/Runtime/Harmonix/Source/HarmonixMidiEditor/Private/MidiFileDetailCustomization.h b/Engine/Plugins/Runtime/Harmonix/Source/HarmonixMidiEditor/Private/MidiFileDetailCustomization.h
--- a/Engine/Plugins/Runtime/Harmonix/Source/HarmonixMidiEditor/Private/MidiFileDetailCustomization.h
+++ b/Engine/Plugins/Runtime/Harmonix/Source/HarmonixMidiEditor/Private/MidiFileDetailCustomization.h
@@ -35,7 +35,6 @@ public:
 private:
 	TSharedPtr<IDetailLayoutBuilder> DetailBuilder;
 	FText FileLengthText;
-	FText RoundedLengthToolTipText;
 	TSharedPtr<STextBlock> FileLengthTextBlock;
 	TSharedPtr<SHorizontalBox> LengthBox;
 	FDetailWidgetRow* LengthRow = nullptr;
