[Harmonix] Add -1 as a track wildcard to the fusion sampler node
#tests PIE

diff --git a/Engine/Plugins/Runtime/Harmonix/Source/HarmonixMetasound/Private/HarmonixMetasound/Nodes/FusionSamplerNode.cpp b/Engine/Plugins/Runtime/Harmonix/Source/HarmonixMetasound/Private/HarmonixMetasound/Nodes/FusionSamplerNode.cpp
--- a/Engine/Plugins/Runtime/Harmonix/Source/HarmonixMetasound/Private/HarmonixMetasound/Nodes/FusionSamplerNode.cpp
+++ b/Engine/Plugins/Runtime/Harmonix/Source/HarmonixMetasound/Private/HarmonixMetasound/Nodes/FusionSamplerNode.cpp
@@ -30,6 +30,7 @@ namespace HarmonixMetasound
 
 	namespace FusionSamplerNodePinNames
 	{
+		METASOUND_PARAM(MidiTrackNumber, "Track Number", "Track number (1 based). 0 will play all tracks (track-omni).")
 		METASOUND_PARAM(Lfo0DepthOverride, "LFO 0 Depth Override", "Overrides the patch's lfo0 depth setting if connected AND if the value is >= 0.")
 		METASOUND_PARAM(Lfo0FrequencyOverride, "LFO 0 Frequency Override", "Overrides the patch's lfo0 frequency setting if connected AND if the vaue is > 0.")
 		METASOUND_PARAM(Lfo1DepthOverride, "LFO 1 Depth Override", "Overrides the patch's lfo1 depth setting if connected AND if the value is >= 0..")
@@ -81,7 +82,7 @@ namespace HarmonixMetasound
 				, InSettings(& InParams.OperatorSettings)
 				, InEnabled(InParams.InputData.GetOrCreateDefaultDataReadReference<bool>(METASOUND_GET_PARAM_NAME(CommonPinNames::Inputs::Enable), InParams.OperatorSettings))
 				, InMidiStream(InParams.InputData.GetOrConstructDataReadReference<FMidiStream>(METASOUND_GET_PARAM_NAME(CommonPinNames::Inputs::MidiStream)))
-				, InTrackNumber(InParams.InputData.GetOrCreateDefaultDataReadReference<int32>(METASOUND_GET_PARAM_NAME(CommonPinNames::Inputs::MidiTrackNumber), InParams.OperatorSettings))
+				, InTrackNumber(InParams.InputData.GetOrCreateDefaultDataReadReference<int32>(METASOUND_GET_PARAM_NAME(FusionSamplerNodePinNames::MidiTrackNumber), InParams.OperatorSettings))
 				, InTransposition(InParams.InputData.GetOrCreateDefaultDataReadReference<int32>(METASOUND_GET_PARAM_NAME(CommonPinNames::Inputs::Transposition), InParams.OperatorSettings))
 				, InLfo0RateOverride(InParams.InputData.GetOrCreateDefaultDataReadReference<float>(METASOUND_GET_PARAM_NAME(FusionSamplerNodePinNames::Lfo0FrequencyOverride), InParams.OperatorSettings))
 				, InLfo0DepthOverride(InParams.InputData.GetOrCreateDefaultDataReadReference<float>(METASOUND_GET_PARAM_NAME(FusionSamplerNodePinNames::Lfo0DepthOverride), InParams.OperatorSettings))
@@ -104,7 +105,7 @@ namespace HarmonixMetasound
 			static const FInputVertexInterface InputInterface = FInputVertexInterface(
 				TInputDataVertex<bool>(METASOUND_GET_PARAM_NAME_AND_METADATA(Inputs::Enable), true),
 				TInputDataVertex<FMidiStream>(METASOUND_GET_PARAM_NAME_AND_METADATA(Inputs::MidiStream)),
-				TInputDataVertex<int32>(METASOUND_GET_PARAM_NAME_AND_METADATA(Inputs::MidiTrackNumber), 1),
+				TInputDataVertex<int32>(METASOUND_GET_PARAM_NAME_AND_METADATA(FusionSamplerNodePinNames::MidiTrackNumber), 1),
 				TInputDataVertex<int32>(METASOUND_GET_PARAM_NAME_AND_METADATA(Inputs::Transposition), 0),
 				TInputDataVertex<FFusionPatchAsset>(METASOUND_GET_PARAM_NAME_AND_METADATA(Inputs::SynthPatch)),
 				TInputDataVertex<float>(METASOUND_GET_PARAM_NAME_AND_METADATA(FusionSamplerNodePinNames::Lfo0FrequencyOverride), 0.0f),
@@ -158,7 +159,7 @@ namespace HarmonixMetasound
 			using namespace CommonPinNames;
 			InVertexData.BindReadVertex(METASOUND_GET_PARAM_NAME(Inputs::Enable), EnableInPin);
 			InVertexData.BindReadVertex(METASOUND_GET_PARAM_NAME(Inputs::MidiStream), MidiStreamInPin);
-			InVertexData.BindReadVertex(METASOUND_GET_PARAM_NAME(Inputs::MidiTrackNumber), TrackNumberInPin);
+			InVertexData.BindReadVertex(METASOUND_GET_PARAM_NAME(FusionSamplerNodePinNames::MidiTrackNumber), TrackNumberInPin);
 			InVertexData.BindReadVertex(METASOUND_GET_PARAM_NAME(Inputs::Transposition), TranspositionInPin);
 			InVertexData.BindReadVertex(METASOUND_GET_PARAM_NAME(FusionSamplerNodePinNames::Lfo0FrequencyOverride), Lfo0RateInPin);
 			InVertexData.BindReadVertex(METASOUND_GET_PARAM_NAME(FusionSamplerNodePinNames::Lfo0DepthOverride), Lfo0DepthInPin);
@@ -631,8 +632,12 @@ namespace HarmonixMetasound
 				if ((*MidiEventIterator).BlockSampleFrameIndex <= CurrentBlockFrameIndex)
 				{
 					const FMidiMsg& MidiMessage = (*MidiEventIterator).MidiMessage;
-					if (MidiMessage.IsStd() && (*MidiEventIterator).TrackIndex == CurrentTrackNumber)
+					if (MidiMessage.IsStd() && (CurrentTrackNumber == 0 || (*MidiEventIterator).TrackIndex == CurrentTrackNumber))
 					{
+						// If using -1 for the track number to play all tracks, There's currently an issue
+						// With voice IDs that playing the same note and channel from the same source but with different tracks
+						// Will create the same ID. This means that notes with the same pitch and data will have the same ID,
+						// So they won't play as separate fusion voices.
 						HandleMidiMessage(
 							(*MidiEventIterator).GetVoiceId(),
 							MidiMessage.GetStdStatus(),
