Fix unneeded API export macro

#rnx

diff --git a/Engine/Plugins/Runtime/Harmonix/Source/HarmonixDsp/Public/HarmonixDsp/ObjectPool.h b/Engine/Plugins/Runtime/Harmonix/Source/HarmonixDsp/Public/HarmonixDsp/ObjectPool.h
--- a/Engine/Plugins/Runtime/Harmonix/Source/HarmonixDsp/Public/HarmonixDsp/ObjectPool.h
+++ b/Engine/Plugins/Runtime/Harmonix/Source/HarmonixDsp/Public/HarmonixDsp/ObjectPool.h
@@ -20,7 +20,7 @@ namespace Harmonix
 	 * Implemented to be thread safe
 	 */
 	template<typename TObject>
-	class HARMONIXDSP_API TDynamicObjectPool
+	class TDynamicObjectPool
 	{
 	public:
 
