Make sure FusionPatchAssetFactory only reimports files with the correct extensions

#tests Harmonix Dev
#rnx

diff --git a/Engine/Plugins/Runtime/Harmonix/Source/HarmonixDspEditor/Private/FusionPatchAssetFactory.cpp b/Engine/Plugins/Runtime/Harmonix/Source/HarmonixDspEditor/Private/FusionPatchAssetFactory.cpp
--- a/Engine/Plugins/Runtime/Harmonix/Source/HarmonixDspEditor/Private/FusionPatchAssetFactory.cpp
+++ b/Engine/Plugins/Runtime/Harmonix/Source/HarmonixDspEditor/Private/FusionPatchAssetFactory.cpp
@@ -44,10 +44,20 @@ bool UFusionPatchAssetFactory::FactoryCanImport(const FString& Filename)
 
 bool UFusionPatchAssetFactory::CanReimport(UObject* Obj, TArray<FString>& OutFilenames)
 {
-	UFusionPatch* FusionPatch = Cast<UFusionPatch>(Obj);
-	if (FusionPatch && FusionPatch->AssetImportData)
+	if (UFusionPatch* FusionPatch = Cast<UFusionPatch>(Obj))
 	{
-		FusionPatch->AssetImportData->ExtractFilenames(OutFilenames);
+		// the PreferredReimportPath is the path of the new file, which can differ from the original SrcFilePath
+		// if the file has the same name but a different extension, don't import it!
+		if (!IsSupportedFileExtension(PreferredReimportPath))
+		{
+			UE_LOG(LogFusionPatchAssetFactory, Warning, TEXT("%s: Failed to reimport with new file. Invalid extension: %s"), *FusionPatch->GetPathName(), *FPaths::GetExtension(PreferredReimportPath));
+			return false;
+		}
+
+		if  (FusionPatch->AssetImportData)
+		{
+			FusionPatch->AssetImportData->ExtractFilenames(OutFilenames);
+		}
 		return true;
 	}
 	return false;
