Make sure MidiFileFactory only reimports files with the correct extensions

#tests Harmonix Dev Tests
[RN] Fix Harmonix Module MIDI File Reimport stability

diff --git a/Engine/Plugins/Runtime/Harmonix/Source/HarmonixMidiEditor/Private/MidiFileFactory.cpp b/Engine/Plugins/Runtime/Harmonix/Source/HarmonixMidiEditor/Private/MidiFileFactory.cpp
--- a/Engine/Plugins/Runtime/Harmonix/Source/HarmonixMidiEditor/Private/MidiFileFactory.cpp
+++ b/Engine/Plugins/Runtime/Harmonix/Source/HarmonixMidiEditor/Private/MidiFileFactory.cpp
@@ -74,6 +74,15 @@ bool UMidiFileFactory::CanReimport(UObject* Obj, TArray<FString>& OutFilenames)
 			UE_LOG(LogHarmonixMidiEditor, Warning, TEXT("Couldn't find source path for %s!"), *AsMidiFile->GetPathName());
 			return false;
 		}
+
+		// the PreferredReimportPath is the path of the new file, which can differ from the original SrcFilePath
+		// if the file has the same name but a different extension, don't import it!
+		if (!IsSupportedFileExtension(PreferredReimportPath))
+		{
+			UE_LOG(LogHarmonixMidiEditor, Warning, TEXT("%s: Failed to reimport with new file. Invalid extension: %s"), *AsMidiFile->GetPathName(), *FPaths::GetExtension(PreferredReimportPath));
+			return false;
+		}
+		
 		OutFilenames.Push(FilePath);
 		return true;
 	}
