Tick MusicTempometer Early With High Priority

diff --git a/Engine/Plugins/Runtime/Harmonix/Source/HarmonixMetasound/Private/HarmonixMetasound/Components/MusicTempometerComponent.cpp b/Engine/Plugins/Runtime/Harmonix/Source/HarmonixMetasound/Private/HarmonixMetasound/Components/MusicTempometerComponent.cpp
--- a/Engine/Plugins/Runtime/Harmonix/Source/HarmonixMetasound/Private/HarmonixMetasound/Components/MusicTempometerComponent.cpp
+++ b/Engine/Plugins/Runtime/Harmonix/Source/HarmonixMetasound/Private/HarmonixMetasound/Components/MusicTempometerComponent.cpp
@@ -8,7 +8,10 @@
 
 UMusicTempometerComponent::UMusicTempometerComponent()
 {
+	SetTickGroup(TG_PrePhysics);
+
 	PrimaryComponentTick.bCanEverTick = true;
+	PrimaryComponentTick.SetPriorityIncludingPrerequisites(true);
 }
 
 void UMusicTempometerComponent::TickComponent(float DeltaTime, enum ELevelTick TickType, FActorComponentTickFunction* ThisTickFunction)
