- Don't crash when calculating a song position when the song maps are invalid/currupt.
- Fix inproper calculation of "Beat Type" (Unused in any current games.)

diff --git a/Engine/Plugins/Runtime/Harmonix/Source/HarmonixMidi/Private/HarmonixMidi/MidiSongPos.cpp b/Engine/Plugins/Runtime/Harmonix/Source/HarmonixMidi/Private/HarmonixMidi/MidiSongPos.cpp
--- a/Engine/Plugins/Runtime/Harmonix/Source/HarmonixMidi/Private/HarmonixMidi/MidiSongPos.cpp
+++ b/Engine/Plugins/Runtime/Harmonix/Source/HarmonixMidi/Private/HarmonixMidi/MidiSongPos.cpp
@@ -125,6 +125,8 @@ void FMidiSongPos::SetByTick(float InTick, const ISongMapEvaluator& Map)
 
 void FMidiSongPos::SetByTimeAndTick(float InMs, float InTick, const ISongMapEvaluator& Map)
 {
+	Timestamp = Map.TickToMusicTimestamp(InTick);
+
 	const FTimeSignature* TimeSig = Map.GetTimeSignatureAtTick(int32(InTick));
 	SecondsIncludingCountIn = InMs / 1000.0f;
 	SecondsFromBarOne       = SecondsIncludingCountIn - Map.GetCountInSeconds();
@@ -135,7 +137,7 @@ void FMidiSongPos::SetByTimeAndTick(float InMs, float InTick, const ISongMapEval
 
 	const FSongSection* SongSection = Map.GetSectionAtTick(int32(InTick));
 	CurrentSongSection = SongSection ? FSongSection(SongSection->Name, SongSection->StartTick, SongSection->LengthTicks) : FSongSection();
-	
+
 	int32 BeatPointIndex = 0;
 	const FBeatMapPoint* BeatPoint = Map.GetBeatPointInfoAtTick(int32(InTick), &BeatPointIndex);
 	if (BeatPoint != nullptr)
@@ -148,11 +150,16 @@ void FMidiSongPos::SetByTimeAndTick(float InMs, float InTick, const ISongMapEval
 	else
 	{
 		const FTimeSignaturePoint* TimeSigPoint = Map.GetTimeSignaturePointAtTick(int32(InTick));
-		check (TimeSigPoint);
-		float BarAtTimeSig = BarsIncludingCountIn - TimeSigPoint->BarIndex;
-		float FractPart = FMath::Fractional(BarAtTimeSig);
-		BeatsIncludingCountIn = TimeSigPoint->BeatIndex + (BarAtTimeSig * TimeSigPoint->TimeSignature.Numerator);
+		if (TimeSigPoint)
+		{
+			float BarAtTimeSig = BarsIncludingCountIn - TimeSigPoint->BarIndex;
+			float FractPart = FMath::Fractional(BarAtTimeSig);
+			BeatsIncludingCountIn = TimeSigPoint->BeatIndex + (BarAtTimeSig * TimeSigPoint->TimeSignature.Numerator);
+		}
+		else
+		{
+			BeatsIncludingCountIn = BarsIncludingCountIn * 4; // 4/4
+		}
 		BeatType = FMath::IsNearlyEqual(Timestamp.Beat, 1.0f) ? EMusicalBeatType::Downbeat : EMusicalBeatType::Normal;
 	}
-	Timestamp = Map.TickToMusicTimestamp(InTick);
 }
