Per-keyzone ADSR on FusionPatches

Additional logic to set initial values to the existing Default patch settings ADSR settings.

#tests Harmonix Dev
[RN] Add support for Per-Keyzone ADSR on Fusion Patches

diff --git a/Engine/Plugins/Runtime/Harmonix/Source/HarmonixDsp/Private/HarmonixDsp/FusionSampler/FusionPatch.cpp b/Engine/Plugins/Runtime/Harmonix/Source/HarmonixDsp/Private/HarmonixDsp/FusionSampler/FusionPatch.cpp
--- a/Engine/Plugins/Runtime/Harmonix/Source/HarmonixDsp/Private/HarmonixDsp/FusionSampler/FusionPatch.cpp
+++ b/Engine/Plugins/Runtime/Harmonix/Source/HarmonixDsp/Private/HarmonixDsp/FusionSampler/FusionPatch.cpp
@@ -89,6 +89,15 @@ void UFusionPatch::PostEditChangeChainProperty(FPropertyChangedChainEvent& Prope
 		return;
 	}
 
+	bool bSettingsChanged = false;
+	const FProperty* ChildChanged = PropertyChangedChainEvent.PropertyChain.GetHead()->GetNextNode()->GetValue();
+	if (ChildChanged && ChildChanged->GetFName() == GET_MEMBER_NAME_CHECKED(FFusionPatchData, Settings))
+	{
+		UE_LOG(LogFusionPatch, Log, TEXT("Fusion patch settings changed"));
+		bSettingsChanged = true;
+	}
+
+	const FFusionPatchSettings& DefaultSettings = FusionPatchData.Settings;
 	for (FKeyzoneSettings& Keyzone : FusionPatchData.Keyzones)
 	{
 		int8 MinNote = Keyzone.MinNote;
@@ -109,6 +118,14 @@ void UFusionPatch::PostEditChangeChainProperty(FPropertyChangedChainEvent& Prope
 		{
 			Keyzone.MinVelocity = MaxVelocity;
 		}
+
+		// update keyzone ADSR default settings for non-overridden keyzones
+		// maintains continuity when deciding to override ADSR settings later.
+		// only update the keyzones if the default settings have changed
+		if (bSettingsChanged && !Keyzone.bOverrideAdsr)
+		{
+			Keyzone.Adsr = DefaultSettings.Adsr[0];
+		}
 	}
 
 	// don't need to do anything if there isn't a RenderableFusionPatchData to copy data to
@@ -242,6 +259,18 @@ void UFusionPatch::Serialize(FArchive& Ar)
 void UFusionPatch::PostLoad()
 {
 	Super::PostLoad();
+
+	// apply the default settings to each keyzone ADSR on post load
+	// only do this for keyzones that are _NOT_ currently overridden
+	// that way, when overriding a keyzone, it'll initialize with the current patch
+	const FFusionPatchSettings& DefaultSettings = FusionPatchData.Settings;
+	for (FKeyzoneSettings& Keyzone : FusionPatchData.Keyzones)
+	{
+		if (!Keyzone.bOverrideAdsr)
+		{
+			Keyzone.Adsr = DefaultSettings.Adsr[0];
+		}
+	}
 }
 
 void FFusionPatchData::InitProxyData(const Audio::FProxyDataInitParams& InitParams)
diff --git a/Engine/Plugins/Runtime/Harmonix/Source/HarmonixDsp/Private/HarmonixDsp/FusionSampler/FusionVoice.cpp b/Engine/Plugins/Runtime/Harmonix/Source/HarmonixDsp/Private/HarmonixDsp/FusionSampler/FusionVoice.cpp
--- a/Engine/Plugins/Runtime/Harmonix/Source/HarmonixDsp/Private/HarmonixDsp/FusionSampler/FusionVoice.cpp
+++ b/Engine/Plugins/Runtime/Harmonix/Source/HarmonixDsp/Private/HarmonixDsp/FusionSampler/FusionVoice.cpp
@@ -88,7 +88,15 @@ bool FFusionVoice::AssignIDs(FFusionSampler* InSampler, const FKeyzoneSettings*
 	MySampler = InSampler;
 	PitchShifter = Shifter;
 
-	AdsrVolume.UseSettings(&(MySampler->AdsrVolumeSettings));
+	if (InKeyZone->bOverrideAdsr)
+	{
+		AdsrVolume.UseSettings(&(InKeyZone->Adsr));
+	}
+	else
+	{
+		AdsrVolume.UseSettings(&(MySampler->AdsrVolumeSettings));
+	}
+	
 	AdsrAssignable.UseSettings(&(MySampler->AdsrAssignableSettings));
 	VelocityGain = 0.0f;
 
diff --git a/Engine/Plugins/Runtime/Harmonix/Source/HarmonixDsp/Private/HarmonixDsp/FusionSampler/Settings/KeyzoneSettings.cpp b/Engine/Plugins/Runtime/Harmonix/Source/HarmonixDsp/Private/HarmonixDsp/FusionSampler/Settings/KeyzoneSettings.cpp
--- a/Engine/Plugins/Runtime/Harmonix/Source/HarmonixDsp/Private/HarmonixDsp/FusionSampler/Settings/KeyzoneSettings.cpp
+++ b/Engine/Plugins/Runtime/Harmonix/Source/HarmonixDsp/Private/HarmonixDsp/FusionSampler/Settings/KeyzoneSettings.cpp
@@ -10,6 +10,11 @@
 
 DEFINE_LOG_CATEGORY(LogKeyzoneSettings);
 
+FKeyzoneSettings::FKeyzoneSettings()
+{
+	Adsr.Target = EAdsrTarget::Volume;
+}
+
 FKeyzoneSettings::~FKeyzoneSettings()
 {
 	if (SoundWaveProxy && SoundWaveProxy->GetLoadingBehavior() == ESoundWaveLoadingBehavior::ForceInline)
diff --git a/Engine/Plugins/Runtime/Harmonix/Source/HarmonixDsp/Public/HarmonixDsp/FusionSampler/Settings/KeyzoneSettings.h b/Engine/Plugins/Runtime/Harmonix/Source/HarmonixDsp/Public/HarmonixDsp/FusionSampler/Settings/KeyzoneSettings.h
--- a/Engine/Plugins/Runtime/Harmonix/Source/HarmonixDsp/Public/HarmonixDsp/FusionSampler/Settings/KeyzoneSettings.h
+++ b/Engine/Plugins/Runtime/Harmonix/Source/HarmonixDsp/Public/HarmonixDsp/FusionSampler/Settings/KeyzoneSettings.h
@@ -1,6 +1,7 @@
 // Copyright Epic Games, Inc. All Rights Reserved.
 #pragma once
 
+#include "HarmonixDsp/Modulators/Settings/AdsrSettings.h"
 #include "HarmonixDsp/PannerDetails.h"
 #include "HarmonixDsp/StretcherAndPitchShifter.h"
 #include "HarmonixDsp/Streaming/TrackChannelInfo.h"
@@ -40,7 +41,7 @@ struct HARMONIXDSP_API FKeyzoneSettings
 
 	GENERATED_BODY()
 
-	FKeyzoneSettings() {}
+	FKeyzoneSettings();
 
 	FKeyzoneSettings(const FKeyzoneArgs& InArgs)
 	{
@@ -97,7 +98,13 @@ public:
 	uint8 Priority = 50;
 	
 	UPROPERTY(EditDefaultsOnly, Category="Settings")
-	FPannerDetails Pan;
+	FPannerDetails Pan;	
+
+	UPROPERTY(EditDefaultsOnly, Category="Settings")
+	bool bOverrideAdsr = false;
+	
+	UPROPERTY(EditDefaultsOnly, Category = "Settings", meta=(EditCondition="bOverrideAdsr", EditConditionHides=true))
+	FAdsrSettings Adsr;
 	
 	UPROPERTY(EditDefaultsOnly, Category="Settings", Meta = (UIMin = -100, UIMax = 100, ClampMin = -100, ClampMax = 100))
 	float FineTuneCents = 0.0f;
