Fix AdrSettings curve names and comments

AdsrSetting Curve values are unitless parameters that determine the shape of the curve.

#rnx
[RN]

diff --git a/Engine/Plugins/Runtime/Harmonix/Source/HarmonixDsp/Public/HarmonixDsp/Modulators/Settings/AdsrSettings.h b/Engine/Plugins/Runtime/Harmonix/Source/HarmonixDsp/Public/HarmonixDsp/Modulators/Settings/AdsrSettings.h
--- a/Engine/Plugins/Runtime/Harmonix/Source/HarmonixDsp/Public/HarmonixDsp/Modulators/Settings/AdsrSettings.h
+++ b/Engine/Plugins/Runtime/Harmonix/Source/HarmonixDsp/Public/HarmonixDsp/Modulators/Settings/AdsrSettings.h
@@ -41,13 +41,31 @@ public:
 	UPROPERTY(EditDefaultsOnly, Category="Settings", Meta = (UIMin = "0.0", UIMax = "1.0", ClampMin = "0.0", ClampMax = "1"))
 	float Depth = 0.5f;
 
-	UPROPERTY(EditDefaultsOnly, Category="Settings", DisplayName = "Attack Curve (Seconds)", Meta = (UIMin = "-1", UIMax = "1", ClampMin = "-1", ClampMax = "1"))
+	/*
+	 * Control the shape of the attack curve:
+	 * 0: Linear
+	 * <0: Exponential
+	 * >0: Logarithmic
+	 **/
+	UPROPERTY(EditDefaultsOnly, Category="Settings", DisplayName = "Attack Curve", Meta = (UIMin = "-1", UIMax = "1", ClampMin = "-1", ClampMax = "1"))
 	float AttackCurve = 0.0f;
 
-	UPROPERTY(EditDefaultsOnly, Category="Settings", DisplayName = "Decay Curve (Seconds)", Meta = (UIMin = "-1", UIMax = "1", ClampMin = "-1", ClampMax = "1"))
+	/*
+	* Control the shape of the attack curve:
+	* 0: Linear
+	* <0: Exponential
+	* >0: Logarithmic
+	**/
+	UPROPERTY(EditDefaultsOnly, Category="Settings", DisplayName = "Decay Curve", Meta = (UIMin = "-1", UIMax = "1", ClampMin = "-1", ClampMax = "1"))
 	float DecayCurve = 0.0f;
 
-	UPROPERTY(EditDefaultsOnly, Category="Settings", DisplayName = "Release Curve (Seconds)", Meta = (UIMin = "-1", UIMax = "1", ClampMin = "-1", ClampMax = "1"))
+	/*
+	 * Control the shape of the attack curve:
+	 * 0: Linear
+	 * <0: Exponential
+	 * >0: Logarithmic
+	 **/
+	UPROPERTY(EditDefaultsOnly, Category="Settings", DisplayName = "Release Curve", Meta = (UIMin = "-1", UIMax = "1", ClampMin = "-1", ClampMax = "1"))
 	float ReleaseCurve = 0.0f;
 
 	static const uint32 kCurveTableSize = 20;
