Fix bug with HarmonixDsp ADSR when Attack curve is ~1.0f.

Use the Attack Time to determine when to switch to the sustian/decay stage instead of the Curve Value

#tests Harmonix Dsp Tests
#rnx

diff --git a/Engine/Plugins/Runtime/Harmonix/Source/HarmonixDsp/Private/HarmonixDsp/Modulators/Adsr.cpp b/Engine/Plugins/Runtime/Harmonix/Source/HarmonixDsp/Private/HarmonixDsp/Modulators/Adsr.cpp
--- a/Engine/Plugins/Runtime/Harmonix/Source/HarmonixDsp/Private/HarmonixDsp/Modulators/Adsr.cpp
+++ b/Engine/Plugins/Runtime/Harmonix/Source/HarmonixDsp/Private/HarmonixDsp/Modulators/Adsr.cpp
@@ -65,7 +65,10 @@ void FAdsr::Advance(uint32 InNumSamples)
 				Gain = Settings->LerpAttackCurve(TableIdx);
 			}
 		}
-		if (Gain >= 1.0)
+
+		Gain = FMath::Clamp(Gain, 0.0f, 1.0f);
+		int32 CurveEndFrame = Settings->AttackTime * SamplesPerSecond;
+		if (CurveAge >= CurveEndFrame)
 		{
 			if (Settings->HasDecayStage())
 			{
