Fix unhandled error in DTA Parser code.

Add context around parser code to provide detailed error messages

#rnx
#tests HarmonixDev

diff --git a/Engine/Plugins/Runtime/Harmonix/Source/HarmonixDspEditor/Private/Dta/DtaParser.cpp b/Engine/Plugins/Runtime/Harmonix/Source/HarmonixDspEditor/Private/Dta/DtaParser.cpp
--- a/Engine/Plugins/Runtime/Harmonix/Source/HarmonixDspEditor/Private/Dta/DtaParser.cpp
+++ b/Engine/Plugins/Runtime/Harmonix/Source/HarmonixDspEditor/Private/Dta/DtaParser.cpp
@@ -1,9 +1,13 @@
 // Copyright Epic Games, Inc. All Rights Reserved.
 #include "Dta/DtaParser.h"
 #include "Misc/Char.h"
+#include "Internationalization/Internationalization.h"
+#include "Internationalization/Text.h"
 
 DEFINE_LOG_CATEGORY(LogDtaParser)
 
+#define LOCTEXT_NAMESPACE "DtaParser"
+
 /**
 *
 * The Dta Grammar
@@ -209,17 +213,29 @@ public:
 };
 
 
-void FDtaParser::DtaStringToJsonString(const FString& InDtaString, FString& OutJsonString)
+bool FDtaParser::DtaStringToJsonString(const FString& InDtaString, FString& OutJsonString, FString& OutErrorMessage)
 {
 	TArray<FString> Tokens;
-	Tokenize(InDtaString, Tokens);
-	ParseToJson(Tokens, OutJsonString);
+	FParseContext Context;
+	if (!Tokenize(InDtaString, Tokens, Context))
+	{
+		OutErrorMessage = FText::Format(LOCTEXT("DtaTokenizerError", "Error parsing DTA at {0} - {1}"),
+			FText::FromString(FString::Printf(TEXT("(%d, %d): \'%c\'"), Context.LineIdx + 1, Context.CharIdx + 1, Context.CurrentChar)),
+			FText::FromString(Context.GetErrorAsString())).ToString();
+		return false;
+	}
+	
+	if (!ParseToJson(Tokens, OutJsonString))
+	{
+		OutErrorMessage = LOCTEXT("DtaParserError", "Error parsing DTA").ToString();
+		return false;
+	}
+
+	return true;
 }
 
-void FDtaParser::ParseToJson(const TArray<FString>& Tokens, FString& OutJsonString)
+bool FDtaParser::ParseToJson(const TArray<FString>& Tokens, FString& OutJsonString)
 {
-	int Num = Tokens.Num();
-
 	TArray<TSharedPtr<FParseNode>> Stack;
 	for (int idx = 0; idx < Tokens.Num(); ++idx)
 	{
@@ -232,12 +248,18 @@ void FDtaParser::ParseToJson(const TArray<FString>& Tokens, FString& OutJsonStri
 		else if (Token_0[0] == ')')
 		{
 			TArray<TSharedPtr<FParseNode>> TempStack;
-
-			while (Stack.Last()->GetType() != FParseNode::Type_Open)
+			
+			while (!Stack.IsEmpty() && Stack.Last()->GetType() != FParseNode::Type_Open)
 			{
 				TempStack.Add(Stack.Pop());
 			}
 
+			// this should be guaranteed since we're keeping track of parens in the tokenizer
+			if (!ensure(Stack.Num() > 0))
+			{
+				return false;
+			}
+
 			// pop off the open paren
 			Stack.Pop();
 
@@ -314,8 +336,8 @@ void FDtaParser::ParseToJson(const TArray<FString>& Tokens, FString& OutJsonStri
 			}
 			else
 			{
-				// this is an error state we should probably handle?
-				checkNoEntry();
+				// The parser has reached an invalid state!
+				return false;
 			}
 		}
 		else
@@ -334,13 +356,35 @@ void FDtaParser::ParseToJson(const TArray<FString>& Tokens, FString& OutJsonStri
 	}
 
 	RootObject->ToJson(OutJsonString);
+	return true;
 }
 
+void FDtaParser::FParseContext::Reset()
+{
+	LineIdx = 0;
+	CharIdx = 0;
+	CurrentChar = '\0';
+	PrevChar = '\0';
+}
 
+FString FDtaParser::FParseContext::GetErrorAsString() const
+{
+	switch (Error)
+	{
+	case ParenthesisMismatch:
+		return LOCTEXT("DtaParserError_ParenMismatch", "Paranthesis Mismatch").ToString();
+	case UnexpectedEndOfFile:
+		return LOCTEXT("DtaParserError_UnexpectedEOF", "Unexpected End Of File").ToString();
+	default:
+		return LOCTEXT("DtaParserError_Unexpected", "Unexpected Error").ToString();
+	}
+}
 
-void FDtaParser::Tokenize(const FString& InDtaString, TArray<FString>& OutTokens)
+
+bool FDtaParser::Tokenize(const FString& InDtaString, TArray<FString>& OutTokens, FParseContext& OutContext)
 {
 	OutTokens.Reset();
+	OutContext.Reset();
 
 	static const int kSlack = 10;
 	FString Token;
@@ -348,19 +392,27 @@ void FDtaParser::Tokenize(const FString& InDtaString, TArray<FString>& OutTokens
 
 	int Num = InDtaString.Len();
 	bool IsInCommentSection = false;
-	for (int Pos = 0; Pos < Num; ++Pos)
+	int OpenParenCount = 0;
+	for (int Pos = 0; Pos < Num; ++Pos, ++OutContext.CharIdx)
 	{
-		auto chr = InDtaString[Pos];
+		OutContext.PrevChar = OutContext.CurrentChar;
+		OutContext.CurrentChar = InDtaString[Pos];
+
+		if (OutContext.CurrentChar == '\n')
+		{
+			++OutContext.LineIdx;
+			OutContext.CharIdx = 0;
+		}
 
 		if (IsInCommentSection)
 		{
 			// break out of the comment section if we hit a newline.
-			if (chr == '\n')
+			if (OutContext.CurrentChar == '\n')
 			{
 				IsInCommentSection = false;
 			}
 		}
-		else if (chr == ';')
+		else if (OutContext.CurrentChar == ';')
 		{
 			// found a comment,
 			// add any token we were building
@@ -376,7 +428,7 @@ void FDtaParser::Tokenize(const FString& InDtaString, TArray<FString>& OutTokens
 			Token.Empty(kSlack);
 			continue;
 		}
-		else if (FChar::IsWhitespace(chr))
+		else if (FChar::IsWhitespace(OutContext.CurrentChar))
 		{
 			if (!Token.IsEmpty())
 			{
@@ -385,7 +437,7 @@ void FDtaParser::Tokenize(const FString& InDtaString, TArray<FString>& OutTokens
 			}
 			continue;
 		}
-		else if (chr == '(' || chr == ')')
+		else if (OutContext.CurrentChar == '(' || OutContext.CurrentChar == ')')
 		{
 			if (!Token.IsEmpty())
 			{
@@ -393,12 +445,36 @@ void FDtaParser::Tokenize(const FString& InDtaString, TArray<FString>& OutTokens
 				Token.Empty(kSlack);
 			}
 
+			if (OutContext.CurrentChar == '(')
+			{
+				++OpenParenCount;
+			}
+			else
+			{
+				--OpenParenCount;
+			}
+
+			if (OpenParenCount < 0)
+			{
+				OutContext.Error = FParseContext::ParenthesisMismatch;
+				return false;
+			}
 			// add the parenthesis as a token
-			OutTokens.Add(FString::Chr(chr));
+			OutTokens.Add(FString::Chr(OutContext.CurrentChar));
 		}
 		else
 		{
-			Token += chr;
+			Token += OutContext.CurrentChar;
 		}
 	}
+	
+	if (OpenParenCount > 0)
+	{
+		OutContext.Error = FParseContext::UnexpectedEndOfFile;
+		return false;
+	}
+
+	return true;
 };
+
+#undef LOCTEXT_NAMESPACE
\ No newline at end of file
diff --git a/Engine/Plugins/Runtime/Harmonix/Source/HarmonixDspEditor/Private/Dta/DtaParser.h b/Engine/Plugins/Runtime/Harmonix/Source/HarmonixDspEditor/Private/Dta/DtaParser.h
--- a/Engine/Plugins/Runtime/Harmonix/Source/HarmonixDspEditor/Private/Dta/DtaParser.h
+++ b/Engine/Plugins/Runtime/Harmonix/Source/HarmonixDspEditor/Private/Dta/DtaParser.h
@@ -2,7 +2,7 @@
 #pragma once
 
 #include "Containers/UnrealString.h"
-
+#include "HAL/Platform.h"
 #include "Templates/SharedPointer.h"
 
 
@@ -11,10 +11,38 @@ DECLARE_LOG_CATEGORY_EXTERN(LogDtaParser, Log, All);
 
 struct FDtaParser
 {
-	static void DtaStringToJsonString(const FString& InDtaString, FString& OutJsonString);
+	static bool DtaStringToJsonString(const FString& InDtaString, FString& OutJsonString, FString& OutErrorMessage);
 
 private:
 
-	static void Tokenize(const FString& InDtaString, TArray<FString>& OutTokens);
-	static void ParseToJson(const TArray<FString>& Tokens, FString& OutString);
+	struct FParseContext
+	{
+		enum EError
+		{
+			None,
+			ParenthesisMismatch,
+			UnexpectedEndOfFile
+		};
+		
+		// line index in the dta string being read
+		int LineIdx = 0;
+		
+		// character index within the line
+		int CharIdx = 0;
+
+		// the character being read
+		TCHAR CurrentChar = '\0';
+
+		// the previous "character" read 
+		TCHAR PrevChar = '\0';
+
+		EError Error = EError::None;
+		
+		void Reset();
+
+		FString GetErrorAsString() const;
+	};
+	
+	static bool Tokenize(const FString& InDtaString, TArray<FString>& OutTokens, FParseContext& OutContext);
+	static bool ParseToJson(const TArray<FString>& Tokens, FString& OutString);
 };
\ No newline at end of file
diff --git a/Engine/Plugins/Runtime/Harmonix/Source/HarmonixDspEditor/Private/FusionPatchAssetFactory.cpp b/Engine/Plugins/Runtime/Harmonix/Source/HarmonixDspEditor/Private/FusionPatchAssetFactory.cpp
--- a/Engine/Plugins/Runtime/Harmonix/Source/HarmonixDspEditor/Private/FusionPatchAssetFactory.cpp
+++ b/Engine/Plugins/Runtime/Harmonix/Source/HarmonixDspEditor/Private/FusionPatchAssetFactory.cpp
@@ -252,8 +252,18 @@ UObject* UFusionPatchAssetFactory::FactoryCreateText(UClass* InClass, UObject* I
 	const FString SourceFile = GetCurrentFilename();
 	FString JsonString;
 	const FString DtaString = FString::ConstructFromPtrSize(Buffer, BufferEnd - Buffer);
-
-	FDtaParser::DtaStringToJsonString(DtaString, JsonString);
+	FString DtaErrorMessage;
+	if (!FDtaParser::DtaStringToJsonString(DtaString, JsonString, DtaErrorMessage))
+	{
+		const FText ImportErrorMessage = FText::Format(
+			LOCTEXT("ImportFailed_Json", "Failed to import asset:\n'{0}'.\nFailed to read .fusion file - data malformed: {1}\n"),
+			FText::FromString(SourceFile), FText::FromString(DtaErrorMessage)
+		);
+		FMessageDialog::Open(EAppMsgType::Ok, ImportErrorMessage);
+		UE_LOG(LogFusionPatchAssetFactory, Error, TEXT("Failed to import .fusion asset: %s - %s"), *SourceFile, *DtaErrorMessage);
+		return nullptr;
+	}
+	
 	FString ErrorMessage;
 	TSharedPtr<FJsonObject> JsonObj = FJsonImporter::ParseJsonString(JsonString, ErrorMessage);
 	bool bImportSuccessful = false;
