Fix LOCTEXT Key conflict

#rnx

diff --git a/Engine/Plugins/Runtime/Harmonix/Source/HarmonixDspEditor/Private/FusionPatchAssetFactory.cpp b/Engine/Plugins/Runtime/Harmonix/Source/HarmonixDspEditor/Private/FusionPatchAssetFactory.cpp
--- a/Engine/Plugins/Runtime/Harmonix/Source/HarmonixDspEditor/Private/FusionPatchAssetFactory.cpp
+++ b/Engine/Plugins/Runtime/Harmonix/Source/HarmonixDspEditor/Private/FusionPatchAssetFactory.cpp
@@ -256,7 +256,7 @@ UObject* UFusionPatchAssetFactory::FactoryCreateText(UClass* InClass, UObject* I
 	if (!FDtaParser::DtaStringToJsonString(DtaString, JsonString, DtaErrorMessage))
 	{
 		const FText ImportErrorMessage = FText::Format(
-			LOCTEXT("ImportFailed_Json", "Failed to import asset:\n'{0}'.\nFailed to read .fusion file - data malformed: {1}\n"),
+			LOCTEXT("ImportFailed_DtaToJson", "Failed to import asset:\n'{0}'.\nFailed to read .fusion file - data malformed: {1}\n"),
 			FText::FromString(SourceFile), FText::FromString(DtaErrorMessage)
 		);
 		FMessageDialog::Open(EAppMsgType::Ok, ImportErrorMessage);
