Fix issues found by the static code analysis in CIS

diff --git a/Engine/Plugins/Runtime/Harmonix/Source/HarmonixMetasound/Private/HarmonixMetasound/Nodes/MusicSeekTargetBuilder.cpp b/Engine/Plugins/Runtime/Harmonix/Source/HarmonixMetasound/Private/HarmonixMetasound/Nodes/MusicSeekTargetBuilder.cpp
--- a/Engine/Plugins/Runtime/Harmonix/Source/HarmonixMetasound/Private/HarmonixMetasound/Nodes/MusicSeekTargetBuilder.cpp
+++ b/Engine/Plugins/Runtime/Harmonix/Source/HarmonixMetasound/Private/HarmonixMetasound/Nodes/MusicSeekTargetBuilder.cpp
@@ -268,9 +268,10 @@ namespace HarmonixMetasound::Nodes::MusicTimestampToSeekTarget
 {
 	using namespace Metasound;
 
-	FNodeClassName GetClassName()
+	const FNodeClassName& GetClassName()
 	{
-		return { HarmonixNodeNamespace, TEXT("MusicTimestampToSeekTarget"), TEXT("") };
+		static const FNodeClassName ClassName { HarmonixNodeNamespace, TEXT("MusicTimestampToSeekTarget"), TEXT("") };
+		return ClassName;
 	}
 
 	int32 GetMajorVersion()
diff --git a/Engine/Plugins/Runtime/Harmonix/Source/HarmonixMetasound/Public/HarmonixMetasound/Nodes/MusicSeekTargetBuilder.h b/Engine/Plugins/Runtime/Harmonix/Source/HarmonixMetasound/Public/HarmonixMetasound/Nodes/MusicSeekTargetBuilder.h
--- a/Engine/Plugins/Runtime/Harmonix/Source/HarmonixMetasound/Public/HarmonixMetasound/Nodes/MusicSeekTargetBuilder.h
+++ b/Engine/Plugins/Runtime/Harmonix/Source/HarmonixMetasound/Public/HarmonixMetasound/Nodes/MusicSeekTargetBuilder.h
@@ -8,18 +8,18 @@
 
 namespace HarmonixMetasound::Nodes::BarBeatToSeekTarget
 {
-	HARMONIXMETASOUND_API Metasound::FNodeClassName GetClassName();
+	HARMONIXMETASOUND_API const Metasound::FNodeClassName& GetClassName();
 	HARMONIXMETASOUND_API int32 GetCurrentMajorVersion();
 	
 	namespace  Inputs
 	{
-		DECLARE_METASOUND_PARAM_EXTERN(Bar);
-		DECLARE_METASOUND_PARAM_EXTERN(Beat);
+		DECLARE_METASOUND_PARAM_ALIAS(Bar);
+		DECLARE_METASOUND_PARAM_ALIAS(Beat);
 	}
 
 	namespace Outputs
 	{
-		DECLARE_METASOUND_PARAM_EXTERN(SeekTarget);
+		DECLARE_METASOUND_PARAM_ALIAS(SeekTarget);
 	}
 }
 
@@ -35,7 +35,7 @@ namespace HarmonixMetasound::Nodes::MusicTimeStampToSeekTarget
 
 	namespace Outputs
 	{
-		DECLARE_METASOUND_PARAM_EXTERN(SeekTarget);
+		DECLARE_METASOUND_PARAM_ALIAS(SeekTarget);
 	}
 }
 
@@ -51,6 +51,6 @@ namespace HarmonixMetasound::Nodes::TimeMsToSeekTarget
 
 	namespace Outputs
 	{
-		DECLARE_METASOUND_PARAM_EXTERN(SeekTarget);
+		DECLARE_METASOUND_PARAM_ALIAS(SeekTarget);
 	}
 }
\ No newline at end of file
