[Harmonix] Fixed omni mode fusion samplers skipping some midi notes
#tests PIE, cooked build

diff --git a/Engine/Plugins/Runtime/Harmonix/Source/HarmonixMetasound/Private/HarmonixMetasound/Nodes/FusionSamplerNode.cpp b/Engine/Plugins/Runtime/Harmonix/Source/HarmonixMetasound/Private/HarmonixMetasound/Nodes/FusionSamplerNode.cpp
--- a/Engine/Plugins/Runtime/Harmonix/Source/HarmonixMetasound/Private/HarmonixMetasound/Nodes/FusionSamplerNode.cpp
+++ b/Engine/Plugins/Runtime/Harmonix/Source/HarmonixMetasound/Private/HarmonixMetasound/Nodes/FusionSamplerNode.cpp
@@ -711,7 +711,7 @@ namespace HarmonixMetasound
 		while (MidiEventIterator != MidiEvents.end())
 		{
 			const FMidiMsg& MidiMessage = (*MidiEventIterator).MidiMessage;
-			if (MidiMessage.IsStd() && (*MidiEventIterator).TrackIndex == CurrentTrackNumber)
+			if (MidiMessage.IsStd() && (CurrentTrackNumber == 0 || (*MidiEventIterator).TrackIndex == CurrentTrackNumber))
 			{
 				HandleMidiMessage(
 					(*MidiEventIterator).GetVoiceId(),
