[Harmonix] Add HasValidSongPos() to MusicClockComponent to let users guard against grabbing the SongPos between Starting the MusicClock and the first tick, and a BeatsPerSecond() accessor to MidiSongPos for those who want BPS rather than BPM

#tests PIE Jam in Lobby.
#rnx

diff --git a/Engine/Plugins/Runtime/Harmonix/Source/HarmonixMetasound/Private/HarmonixMetasound/Components/MusicClockComponent.cpp b/Engine/Plugins/Runtime/Harmonix/Source/HarmonixMetasound/Private/HarmonixMetasound/Components/MusicClockComponent.cpp
--- a/Engine/Plugins/Runtime/Harmonix/Source/HarmonixMetasound/Private/HarmonixMetasound/Components/MusicClockComponent.cpp
+++ b/Engine/Plugins/Runtime/Harmonix/Source/HarmonixMetasound/Private/HarmonixMetasound/Components/MusicClockComponent.cpp
@@ -609,6 +609,12 @@ const ISongMapEvaluator& UMusicClockComponent::GetSongMaps() const
 	return SongMaps ? *SongMaps : DefaultMaps;
 }
 
+bool UMusicClockComponent::HasValidSongPos() const
+{
+	return State != EMusicClockState::Stopped && LastBroadcastBeat != -1;
+}
+
+
 const FMidiSongPos& UMusicClockComponent::GetSongPos(ECalibratedMusicTimebase Timebase) const
 {
 	EnsureClockIsValidForGameFrame();
diff --git a/Engine/Plugins/Runtime/Harmonix/Source/HarmonixMetasound/Public/HarmonixMetasound/Components/MusicClockComponent.h b/Engine/Plugins/Runtime/Harmonix/Source/HarmonixMetasound/Public/HarmonixMetasound/Components/MusicClockComponent.h
--- a/Engine/Plugins/Runtime/Harmonix/Source/HarmonixMetasound/Public/HarmonixMetasound/Components/MusicClockComponent.h
+++ b/Engine/Plugins/Runtime/Harmonix/Source/HarmonixMetasound/Public/HarmonixMetasound/Components/MusicClockComponent.h
@@ -210,6 +210,10 @@ public:
 	UFUNCTION(BlueprintPure, Category = "MusicClock")
 	float GetDeltaBeat(ECalibratedMusicTimebase Timebase = ECalibratedMusicTimebase::VideoRenderTime) const;
 
+	/** Returns true if Running or Paused and has been ticked at least Once.  If false, all values will be at their defaults */
+	UFUNCTION(BlueprintPure, Category = "MusicClock")
+	bool HasValidSongPos() const;
+
 	// Note: Not const as it might cause the clock to update from its source.
 	UFUNCTION(BlueprintPure, Category = "MusicClock")
 	const FMidiSongPos& GetSongPos(ECalibratedMusicTimebase Timebase = ECalibratedMusicTimebase::VideoRenderTime) const;
diff --git a/Engine/Plugins/Runtime/Harmonix/Source/HarmonixMidi/Public/HarmonixMidi/MidiSongPos.h b/Engine/Plugins/Runtime/Harmonix/Source/HarmonixMidi/Public/HarmonixMidi/MidiSongPos.h
--- a/Engine/Plugins/Runtime/Harmonix/Source/HarmonixMidi/Public/HarmonixMidi/MidiSongPos.h
+++ b/Engine/Plugins/Runtime/Harmonix/Source/HarmonixMidi/Public/HarmonixMidi/MidiSongPos.h
@@ -61,7 +61,7 @@ public:
 	// "midi tick system" for advancing song position.
 	void SetByTick(float Tick, const ISongMapEvaluator& Map);
 
-	void Reset()	
+	void Reset()
 	{
 		SecondsFromBarOne = 0.0f;
 		SecondsIncludingCountIn = 0.0f;
@@ -75,6 +75,8 @@ public:
 		CurrentSongSection = FSongSection();
 	}
 
+	float GetBeatsPerSecond() const { return Tempo / 60.0f; }
+
 	bool IsZero() const	{ return FMath::IsNearlyZero(BeatsIncludingCountIn, UE_KINDA_SMALL_NUMBER); }
 
 	// operators
