Fix crash in Lfo code...

Initialize LfoBaseFrequency in Delay.h. Leaving it unitialized would sometimes result in NaN values getting through.

#tests Harmonix Dev
#rnx

diff --git a/Engine/Plugins/Runtime/Harmonix/Source/HarmonixDsp/Public/HarmonixDsp/Effects/Delay.h b/Engine/Plugins/Runtime/Harmonix/Source/HarmonixDsp/Public/HarmonixDsp/Effects/Delay.h
--- a/Engine/Plugins/Runtime/Harmonix/Source/HarmonixDsp/Public/HarmonixDsp/Effects/Delay.h
+++ b/Engine/Plugins/Runtime/Harmonix/Source/HarmonixDsp/Public/HarmonixDsp/Effects/Delay.h
@@ -202,7 +202,7 @@ namespace Harmonix::Dsp::Effects
 
 		Harmonix::Dsp::Modulators::FLfo Lfo;
 		FLfoSettings LfoSettings;
-		float LfoBaseFrequency;
+		float LfoBaseFrequency = UE_SMALL_NUMBER;
 
 		// info for bouncing between speakers
 		// belongs here rather than GainTable.cpp because this is specific to the
