Backout change that made MusicClockComponent usable for Seuqnce Editor auditioning. Caused problems with exiting PIE. Will find another way to handle it.

diff --git a/Engine/Plugins/Runtime/Harmonix/Source/HarmonixMetasound/Private/HarmonixMetasound/Components/MusicClockComponent.cpp b/Engine/Plugins/Runtime/Harmonix/Source/HarmonixMetasound/Private/HarmonixMetasound/Components/MusicClockComponent.cpp
--- a/Engine/Plugins/Runtime/Harmonix/Source/HarmonixMetasound/Private/HarmonixMetasound/Components/MusicClockComponent.cpp
+++ b/Engine/Plugins/Runtime/Harmonix/Source/HarmonixMetasound/Private/HarmonixMetasound/Components/MusicClockComponent.cpp
@@ -229,28 +229,19 @@ void UMusicClockComponent::BeginPlay()
 		CreateClockDriver();
 	}
 	UMidiClockUpdateSubsystem::TrackMusicClockComponent(this);
-	if (IsRegistered())
-	{
-		Super::BeginPlay();
-	}
+	Super::BeginPlay();
 }
 
 void UMusicClockComponent::EndPlay(const EEndPlayReason::Type EndPlayReason)
 {
-	if (IsRegistered())
-	{
-		Super::EndPlay(EndPlayReason);
-	}
+	Super::EndPlay(EndPlayReason);
 	UMidiClockUpdateSubsystem::StopTrackingMusicClockComponent(this);
 	DisconnectFromClockDriver();
 }
 
 void UMusicClockComponent::TickComponent(float DeltaTime, enum ELevelTick TickType, FActorComponentTickFunction* ThisTickFunction)
 {
-	if (IsRegistered())
-	{
-		Super::TickComponent(DeltaTime, TickType, ThisTickFunction);
-	}
+	Super::TickComponent(DeltaTime, TickType, ThisTickFunction);
 	EnsureClockIsValidForGameFrame();
 	BroadcastSongPosChanges();
 	BroadcastSeekLoopDetections();
diff --git a/Engine/Plugins/Runtime/Harmonix/Source/HarmonixMetasound/Private/HarmonixMetasound/MusicEnvironmentSupport/HarmonixMusicEnvironmentMetronome.cpp b/Engine/Plugins/Runtime/Harmonix/Source/HarmonixMetasound/Private/HarmonixMetasound/MusicEnvironmentSupport/HarmonixMusicEnvironmentMetronome.cpp
--- a/Engine/Plugins/Runtime/Harmonix/Source/HarmonixMetasound/Private/HarmonixMetasound/MusicEnvironmentSupport/HarmonixMusicEnvironmentMetronome.cpp
+++ b/Engine/Plugins/Runtime/Harmonix/Source/HarmonixMetasound/Private/HarmonixMetasound/MusicEnvironmentSupport/HarmonixMusicEnvironmentMetronome.cpp
@@ -19,12 +19,18 @@ namespace HarmonixMovieSceneMetronome
 
 bool UHarmonixMusicEnvironmentMetronome::Initialize(UWorld* InWorld)
 {
+	// TEMP HACK TO ELIMINATE ENSURE IN MAIN ABOUT MUSIC CLOCK COMPONENT NOT CALLING SUPER
+	// IN END PLAY. THIS FUNCTIONALITY IS CURRENTLY NOT USED BY ANYONE, SO WILL ADDRESS
+	// THIS ANOTHER WAY LATER. 
+	return false;
+#if 0	
 	if (!bMetasoundIsPlaying)
 	{
 		if (!BuildAndStartMetasound(InWorld))
 			return false;
 	}
 	return true;
+#endif
 }
 
 void UHarmonixMusicEnvironmentMetronome::Tick(float DeltaSecs)
@@ -37,12 +43,18 @@ void UHarmonixMusicEnvironmentMetronome::Tick(float DeltaSecs)
 
 void UHarmonixMusicEnvironmentMetronome::Start(double FromSeconds)
 {
+	// TEMP HACK TO ELIMINATE ENSURE IN MAIN ABOUT MUSIC CLOCK COMPONENT NOT CALLING SUPER
+	// IN END PLAY. THIS FUNCTIONALITY IS CURRENTLY NOT USED BY ANYONE, SO WILL ADDRESS
+	// THIS ANOTHER WAY LATER. 
+	return;
+#if 0	
 	UE_LOG(LogMIDI, Log, TEXT("UHarmonixMusicEnvironmentMetronome: Starting Metronome at %f seconds"), FromSeconds);
 	FAudioParameter SeekTargetMs(TEXT("SeekTarget"), static_cast<float>(FromSeconds * 1000.0f));
 	AudioComponent->Activate();
 	AudioComponent->SetParameter(MoveTemp(SeekTargetMs));
 	AudioComponent->SetTriggerParameter(TEXT("PlayMetronome"));
 	MusicClockComponent->Start();
+#endif
 }
 
 void UHarmonixMusicEnvironmentMetronome::Seek(double ToSeconds)
@@ -51,20 +63,38 @@ void UHarmonixMusicEnvironmentMetronome::Seek(double ToSeconds)
 
 void UHarmonixMusicEnvironmentMetronome::Stop()
 {
+	// TEMP HACK TO ELIMINATE ENSURE IN MAIN ABOUT MUSIC CLOCK COMPONENT NOT CALLING SUPER
+	// IN END PLAY. THIS FUNCTIONALITY IS CURRENTLY NOT USED BY ANYONE, SO WILL ADDRESS
+	// THIS ANOTHER WAY LATER. 
+	return;
+#if 0
 	UE_LOG(LogMIDI, Log, TEXT("UHarmonixMusicEnvironmentMetronome: Stopping Metronome"));
 	AudioComponent->SetTriggerParameter(TEXT("StopMetronome"));
 	MusicClockComponent->Stop();
 	AudioComponent->Deactivate();
+#endif
 }
 
 void UHarmonixMusicEnvironmentMetronome::Pause()
 {
+	// TEMP HACK TO ELIMINATE ENSURE IN MAIN ABOUT MUSIC CLOCK COMPONENT NOT CALLING SUPER
+	// IN END PLAY. THIS FUNCTIONALITY IS CURRENTLY NOT USED BY ANYONE, SO WILL ADDRESS
+	// THIS ANOTHER WAY LATER. 
+	return;
+#if 0
 	AudioComponent->SetTriggerParameter(TEXT("PauseMetronome"));
+#endif
 }
 
 void UHarmonixMusicEnvironmentMetronome::Resume()
 {
+	// TEMP HACK TO ELIMINATE ENSURE IN MAIN ABOUT MUSIC CLOCK COMPONENT NOT CALLING SUPER
+	// IN END PLAY. THIS FUNCTIONALITY IS CURRENTLY NOT USED BY ANYONE, SO WILL ADDRESS
+	// THIS ANOTHER WAY LATER. 
+	return;
+#if 0
 	AudioComponent->SetTriggerParameter(TEXT("ContinueMetronome"));
+#endif
 }
 
 float UHarmonixMusicEnvironmentMetronome::GetCurrentTempo()
@@ -79,6 +109,11 @@ double UHarmonixMusicEnvironmentMetronome::GetCurrentPositionSeconds()
 
 FMusicalTime UHarmonixMusicEnvironmentMetronome::GetMusicalTime() const
 {
+	// TEMP HACK TO ELIMINATE ENSURE IN MAIN ABOUT MUSIC CLOCK COMPONENT NOT CALLING SUPER
+	// IN END PLAY. THIS FUNCTIONALITY IS CURRENTLY NOT USED BY ANYONE, SO WILL ADDRESS
+	// THIS ANOTHER WAY LATER. 
+	return FMusicalTime();
+#if 0
 	const FMidiSongPos& SongPos = MusicClockComponent->GetSongPos();
 	FMusicalTime MusicalTime;
 	MusicalTime.TicksPerBeat = MusicalTime::TicksPerQuarterNote * 4 / SongPos.TimeSigDenominator;
@@ -86,6 +121,7 @@ FMusicalTime UHarmonixMusicEnvironmentMetronome::GetMusicalTime() const
 	MusicalTime.Bar = FMath::FloorToInt32(SongPos.BarsIncludingCountIn);
 	MusicalTime.TickInBar = FMath::FloorToInt32(static_cast<float>(MusicalTime.TicksPerBar) * FMath::Fractional(SongPos.BarsIncludingCountIn));
 	return MusicalTime;
+#endif
 }
 
 void UHarmonixMusicEnvironmentMetronome::OnMusicMapSet()
@@ -94,6 +130,11 @@ void UHarmonixMusicEnvironmentMetronome::OnMusicMapSet()
 
 bool UHarmonixMusicEnvironmentMetronome::OnSetTempo(float Bpm)
 {
+	// TEMP HACK TO ELIMINATE ENSURE IN MAIN ABOUT MUSIC CLOCK COMPONENT NOT CALLING SUPER
+	// IN END PLAY. THIS FUNCTIONALITY IS CURRENTLY NOT USED BY ANYONE, SO WILL ADDRESS
+	// THIS ANOTHER WAY LATER. 
+	return false;
+#if 0
 	if (!AudioComponent)
 	{
 		return false;
@@ -102,10 +143,16 @@ bool UHarmonixMusicEnvironmentMetronome::OnSetTempo(float Bpm)
 	AudioComponent->SetParameter(MoveTemp(TempoParam));
 	CurrentTempo = Bpm;
 	return true;
+#endif
 }
 
 bool UHarmonixMusicEnvironmentMetronome::BuildAndStartMetasound(UWorld* InWorld)
 {
+	// TEMP HACK TO ELIMINATE ENSURE IN MAIN ABOUT MUSIC CLOCK COMPONENT NOT CALLING SUPER
+	// IN END PLAY. THIS FUNCTIONALITY IS CURRENTLY NOT USED BY ANYONE, SO WILL ADDRESS
+	// THIS ANOTHER WAY LATER. 
+	return false;
+#if 0
 	if (!BuildMetasound())
 	{
 		return false;
@@ -127,12 +174,18 @@ bool UHarmonixMusicEnvironmentMetronome::BuildAndStartMetasound(UWorld* InWorld)
 	
 	bMetasoundIsPlaying = true;
 	return true;
+#endif
 }
 
 bool UHarmonixMusicEnvironmentMetronome::BuildMetasound()
 {
 	using namespace HarmonixMetasound::Nodes;
-	
+
+	// TEMP HACK TO ELIMINATE ENSURE IN MAIN ABOUT MUSIC CLOCK COMPONENT NOT CALLING SUPER
+	// IN END PLAY. THIS FUNCTIONALITY IS CURRENTLY NOT USED BY ANYONE, SO WILL ADDRESS
+	// THIS ANOTHER WAY LATER. 
+	return false;
+#if 0
 	if (UMetaSoundBuilderSubsystem* MetaSoundBuilder = GEngine->GetEngineSubsystem<UMetaSoundBuilderSubsystem>())
 	{
 		FMetaSoundBuilderNodeOutputHandle Out;
@@ -243,6 +296,7 @@ bool UHarmonixMusicEnvironmentMetronome::BuildMetasound()
 		return Result == EMetaSoundBuilderResult::Succeeded; 
 	}
 	return false;
+#endif
 }
 
 FMetaSoundBuilderNodeOutputHandle UHarmonixMusicEnvironmentMetronome::AddGraphInput(const FName& InputName, const float Value) const
