Tweaks to MusicClockComponent and metronomes to allow use of MusicClockComponent in the Sequence Editor.

diff --git a/Engine/Plugins/Runtime/Harmonix/Source/HarmonixMetasound/Private/HarmonixMetasound/Components/MusicClockComponent.cpp b/Engine/Plugins/Runtime/Harmonix/Source/HarmonixMetasound/Private/HarmonixMetasound/Components/MusicClockComponent.cpp
--- a/Engine/Plugins/Runtime/Harmonix/Source/HarmonixMetasound/Private/HarmonixMetasound/Components/MusicClockComponent.cpp
+++ b/Engine/Plugins/Runtime/Harmonix/Source/HarmonixMetasound/Private/HarmonixMetasound/Components/MusicClockComponent.cpp
@@ -223,18 +223,28 @@ FMidiSongPos UMusicClockComponent::CalculateSongPosWithOffset(float MsOffset, EC
 }
 
 void UMusicClockComponent::BeginPlay()
+{
+	BeginPlayInternal();
+	Super::BeginPlay();
+}
+
+void UMusicClockComponent::BeginPlayInternal()
 {
 	if (!ClockDriver)
 	{
 		CreateClockDriver();
 	}
 	UMidiClockUpdateSubsystem::TrackMusicClockComponent(this);
-	Super::BeginPlay();
 }
 
 void UMusicClockComponent::EndPlay(const EEndPlayReason::Type EndPlayReason)
 {
 	Super::EndPlay(EndPlayReason);
+	EndPlayInternal();
+}
+
+void UMusicClockComponent::EndPlayInternal()
+{
 	UMidiClockUpdateSubsystem::StopTrackingMusicClockComponent(this);
 	DisconnectFromClockDriver();
 }
@@ -242,6 +252,11 @@ void UMusicClockComponent::EndPlay(const EEndPlayReason::Type EndPlayReason)
 void UMusicClockComponent::TickComponent(float DeltaTime, enum ELevelTick TickType, FActorComponentTickFunction* ThisTickFunction)
 {
 	Super::TickComponent(DeltaTime, TickType, ThisTickFunction);
+	TickComponentInternal();
+}
+
+void UMusicClockComponent::TickComponentInternal()
+{
 	EnsureClockIsValidForGameFrame();
 	BroadcastSongPosChanges();
 	BroadcastSeekLoopDetections();
diff --git a/Engine/Plugins/Runtime/Harmonix/Source/HarmonixMetasound/Private/HarmonixMetasound/MusicEnvironmentSupport/HarmonixMusicEnvironmentMetronome.cpp b/Engine/Plugins/Runtime/Harmonix/Source/HarmonixMetasound/Private/HarmonixMetasound/MusicEnvironmentSupport/HarmonixMusicEnvironmentMetronome.cpp
--- a/Engine/Plugins/Runtime/Harmonix/Source/HarmonixMetasound/Private/HarmonixMetasound/MusicEnvironmentSupport/HarmonixMusicEnvironmentMetronome.cpp
+++ b/Engine/Plugins/Runtime/Harmonix/Source/HarmonixMetasound/Private/HarmonixMetasound/MusicEnvironmentSupport/HarmonixMusicEnvironmentMetronome.cpp
@@ -19,42 +19,36 @@ namespace HarmonixMovieSceneMetronome
 
 bool UHarmonixMusicEnvironmentMetronome::Initialize(UWorld* InWorld)
 {
-	// TEMP HACK TO ELIMINATE ENSURE IN MAIN ABOUT MUSIC CLOCK COMPONENT NOT CALLING SUPER
-	// IN END PLAY. THIS FUNCTIONALITY IS CURRENTLY NOT USED BY ANYONE, SO WILL ADDRESS
-	// THIS ANOTHER WAY LATER. 
-	return false;
-#if 0	
 	if (!bMetasoundIsPlaying)
 	{
 		if (!BuildAndStartMetasound(InWorld))
 			return false;
 	}
 	return true;
-#endif
 }
 
 void UHarmonixMusicEnvironmentMetronome::Tick(float DeltaSecs)
 {
 	if (MusicClockComponent)
 	{
-		MusicClockComponent->TickComponent(DeltaSecs, LEVELTICK_TimeOnly, nullptr);
+		MusicClockComponent->TickComponentInternal();
 	}
 }
 
 void UHarmonixMusicEnvironmentMetronome::Start(double FromSeconds)
 {
-	// TEMP HACK TO ELIMINATE ENSURE IN MAIN ABOUT MUSIC CLOCK COMPONENT NOT CALLING SUPER
-	// IN END PLAY. THIS FUNCTIONALITY IS CURRENTLY NOT USED BY ANYONE, SO WILL ADDRESS
-	// THIS ANOTHER WAY LATER. 
-	return;
-#if 0	
 	UE_LOG(LogMIDI, Log, TEXT("UHarmonixMusicEnvironmentMetronome: Starting Metronome at %f seconds"), FromSeconds);
+	
+	if (!AudioComponent || !MusicClockComponent)
+	{
+		return;
+	}
+
 	FAudioParameter SeekTargetMs(TEXT("SeekTarget"), static_cast<float>(FromSeconds * 1000.0f));
 	AudioComponent->Activate();
 	AudioComponent->SetParameter(MoveTemp(SeekTargetMs));
 	AudioComponent->SetTriggerParameter(TEXT("PlayMetronome"));
 	MusicClockComponent->Start();
-#endif
 }
 
 void UHarmonixMusicEnvironmentMetronome::Seek(double ToSeconds)
@@ -63,38 +57,33 @@ void UHarmonixMusicEnvironmentMetronome::Seek(double ToSeconds)
 
 void UHarmonixMusicEnvironmentMetronome::Stop()
 {
-	// TEMP HACK TO ELIMINATE ENSURE IN MAIN ABOUT MUSIC CLOCK COMPONENT NOT CALLING SUPER
-	// IN END PLAY. THIS FUNCTIONALITY IS CURRENTLY NOT USED BY ANYONE, SO WILL ADDRESS
-	// THIS ANOTHER WAY LATER. 
-	return;
-#if 0
 	UE_LOG(LogMIDI, Log, TEXT("UHarmonixMusicEnvironmentMetronome: Stopping Metronome"));
-	AudioComponent->SetTriggerParameter(TEXT("StopMetronome"));
-	MusicClockComponent->Stop();
-	AudioComponent->Deactivate();
-#endif
+	
+	if (MusicClockComponent)
+	{
+		MusicClockComponent->Stop();
+	}
+	if (AudioComponent)
+	{
+		AudioComponent->SetTriggerParameter(TEXT("StopMetronome"));
+		AudioComponent->Deactivate();
+	}
 }
 
 void UHarmonixMusicEnvironmentMetronome::Pause()
 {
-	// TEMP HACK TO ELIMINATE ENSURE IN MAIN ABOUT MUSIC CLOCK COMPONENT NOT CALLING SUPER
-	// IN END PLAY. THIS FUNCTIONALITY IS CURRENTLY NOT USED BY ANYONE, SO WILL ADDRESS
-	// THIS ANOTHER WAY LATER. 
-	return;
-#if 0
-	AudioComponent->SetTriggerParameter(TEXT("PauseMetronome"));
-#endif
+	if (AudioComponent)
+	{
+		AudioComponent->SetTriggerParameter(TEXT("PauseMetronome"));
+	}
 }
 
 void UHarmonixMusicEnvironmentMetronome::Resume()
 {
-	// TEMP HACK TO ELIMINATE ENSURE IN MAIN ABOUT MUSIC CLOCK COMPONENT NOT CALLING SUPER
-	// IN END PLAY. THIS FUNCTIONALITY IS CURRENTLY NOT USED BY ANYONE, SO WILL ADDRESS
-	// THIS ANOTHER WAY LATER. 
-	return;
-#if 0
-	AudioComponent->SetTriggerParameter(TEXT("ContinueMetronome"));
-#endif
+	if (AudioComponent)
+	{
+		AudioComponent->SetTriggerParameter(TEXT("ContinueMetronome"));
+	}
 }
 
 float UHarmonixMusicEnvironmentMetronome::GetCurrentTempo()
@@ -109,19 +98,25 @@ double UHarmonixMusicEnvironmentMetronome::GetCurrentPositionSeconds()
 
 FMusicalTime UHarmonixMusicEnvironmentMetronome::GetMusicalTime() const
 {
-	// TEMP HACK TO ELIMINATE ENSURE IN MAIN ABOUT MUSIC CLOCK COMPONENT NOT CALLING SUPER
-	// IN END PLAY. THIS FUNCTIONALITY IS CURRENTLY NOT USED BY ANYONE, SO WILL ADDRESS
-	// THIS ANOTHER WAY LATER. 
-	return FMusicalTime();
-#if 0
-	const FMidiSongPos& SongPos = MusicClockComponent->GetSongPos();
 	FMusicalTime MusicalTime;
-	MusicalTime.TicksPerBeat = MusicalTime::TicksPerQuarterNote * 4 / SongPos.TimeSigDenominator;
-	MusicalTime.TicksPerBar = MusicalTime.TicksPerBeat * SongPos.TimeSigNumerator;
-	MusicalTime.Bar = FMath::FloorToInt32(SongPos.BarsIncludingCountIn);
-	MusicalTime.TickInBar = FMath::FloorToInt32(static_cast<float>(MusicalTime.TicksPerBar) * FMath::Fractional(SongPos.BarsIncludingCountIn));
+	if (MusicClockComponent)
+	{
+		const FMidiSongPos& SongPos = MusicClockComponent->GetSongPos();
+		MusicalTime.TicksPerBeat = MusicalTime::TicksPerQuarterNote * 4 / SongPos.TimeSigDenominator;
+		MusicalTime.TicksPerBar = MusicalTime.TicksPerBeat * SongPos.TimeSigNumerator;
+		MusicalTime.Bar = FMath::FloorToInt32(SongPos.BarsIncludingCountIn);
+		MusicalTime.TickInBar = FMath::FloorToInt32(static_cast<float>(MusicalTime.TicksPerBar) * FMath::Fractional(SongPos.BarsIncludingCountIn));
+	}
 	return MusicalTime;
-#endif
+}
+
+void UHarmonixMusicEnvironmentMetronome::BeginDestroy()
+{
+	if (MusicClockComponent)
+	{
+		MusicClockComponent->EndPlayInternal();
+	}
+	UObject::BeginDestroy();
 }
 
 void UHarmonixMusicEnvironmentMetronome::OnMusicMapSet()
@@ -130,11 +125,6 @@ void UHarmonixMusicEnvironmentMetronome::OnMusicMapSet()
 
 bool UHarmonixMusicEnvironmentMetronome::OnSetTempo(float Bpm)
 {
-	// TEMP HACK TO ELIMINATE ENSURE IN MAIN ABOUT MUSIC CLOCK COMPONENT NOT CALLING SUPER
-	// IN END PLAY. THIS FUNCTIONALITY IS CURRENTLY NOT USED BY ANYONE, SO WILL ADDRESS
-	// THIS ANOTHER WAY LATER. 
-	return false;
-#if 0
 	if (!AudioComponent)
 	{
 		return false;
@@ -143,29 +133,23 @@ bool UHarmonixMusicEnvironmentMetronome::OnSetTempo(float Bpm)
 	AudioComponent->SetParameter(MoveTemp(TempoParam));
 	CurrentTempo = Bpm;
 	return true;
-#endif
 }
 
 bool UHarmonixMusicEnvironmentMetronome::BuildAndStartMetasound(UWorld* InWorld)
 {
-	// TEMP HACK TO ELIMINATE ENSURE IN MAIN ABOUT MUSIC CLOCK COMPONENT NOT CALLING SUPER
-	// IN END PLAY. THIS FUNCTIONALITY IS CURRENTLY NOT USED BY ANYONE, SO WILL ADDRESS
-	// THIS ANOTHER WAY LATER. 
-	return false;
-#if 0
 	if (!BuildMetasound())
 	{
 		return false;
 	}
-	AudioComponent = NewObject<UAudioComponent>(this); //(Actor->AddComponentByClass(UAudioComponent::StaticClass(), false, FTransform(),false));
+	AudioComponent = NewObject<UAudioComponent>(this);
 	AudioComponent->SetUISound(true);
 	SourceBuilder->Audition(this, AudioComponent, {});
 
-	MusicClockComponent = NewObject<UMusicClockComponent>(this); //(Actor->AddComponentByClass(UMusicClockComponent::StaticClass(), false, FTransform(),false));
+	MusicClockComponent = NewObject<UMusicClockComponent>(this); 
 	MusicClockComponent->MetasoundOutputName = TEXT("MIDI Clock");
 	MusicClockComponent->ConnectToMetasoundOnAudioComponent(AudioComponent);
 	MusicClockComponent->Start();
-	MusicClockComponent->BeginPlay();
+	MusicClockComponent->BeginPlayInternal();
 
 	if (MusicMap)
 	{
@@ -174,18 +158,12 @@ bool UHarmonixMusicEnvironmentMetronome::BuildAndStartMetasound(UWorld* InWorld)
 	
 	bMetasoundIsPlaying = true;
 	return true;
-#endif
 }
 
 bool UHarmonixMusicEnvironmentMetronome::BuildMetasound()
 {
 	using namespace HarmonixMetasound::Nodes;
 
-	// TEMP HACK TO ELIMINATE ENSURE IN MAIN ABOUT MUSIC CLOCK COMPONENT NOT CALLING SUPER
-	// IN END PLAY. THIS FUNCTIONALITY IS CURRENTLY NOT USED BY ANYONE, SO WILL ADDRESS
-	// THIS ANOTHER WAY LATER. 
-	return false;
-#if 0
 	if (UMetaSoundBuilderSubsystem* MetaSoundBuilder = GEngine->GetEngineSubsystem<UMetaSoundBuilderSubsystem>())
 	{
 		FMetaSoundBuilderNodeOutputHandle Out;
@@ -296,7 +274,6 @@ bool UHarmonixMusicEnvironmentMetronome::BuildMetasound()
 		return Result == EMetaSoundBuilderResult::Succeeded; 
 	}
 	return false;
-#endif
 }
 
 FMetaSoundBuilderNodeOutputHandle UHarmonixMusicEnvironmentMetronome::AddGraphInput(const FName& InputName, const float Value) const
diff --git a/Engine/Plugins/Runtime/Harmonix/Source/HarmonixMetasound/Private/HarmonixMetasound/MusicEnvironmentSupport/HarmonixMusicEnvironmentMetronome.h b/Engine/Plugins/Runtime/Harmonix/Source/HarmonixMetasound/Private/HarmonixMetasound/MusicEnvironmentSupport/HarmonixMusicEnvironmentMetronome.h
--- a/Engine/Plugins/Runtime/Harmonix/Source/HarmonixMetasound/Private/HarmonixMetasound/MusicEnvironmentSupport/HarmonixMusicEnvironmentMetronome.h
+++ b/Engine/Plugins/Runtime/Harmonix/Source/HarmonixMetasound/Private/HarmonixMetasound/MusicEnvironmentSupport/HarmonixMusicEnvironmentMetronome.h
@@ -36,6 +36,8 @@ public:
 	virtual float GetCurrentTempo() override;
 	virtual double GetCurrentPositionSeconds() override;
 	virtual FMusicalTime GetMusicalTime() const override;
+
+	virtual void BeginDestroy() override;
 	
 protected:
 	virtual void OnMusicMapSet() override;
diff --git a/Engine/Plugins/Runtime/Harmonix/Source/HarmonixMetasound/Public/HarmonixMetasound/Components/MusicClockComponent.h b/Engine/Plugins/Runtime/Harmonix/Source/HarmonixMetasound/Public/HarmonixMetasound/Components/MusicClockComponent.h
--- a/Engine/Plugins/Runtime/Harmonix/Source/HarmonixMetasound/Public/HarmonixMetasound/Components/MusicClockComponent.h
+++ b/Engine/Plugins/Runtime/Harmonix/Source/HarmonixMetasound/Public/HarmonixMetasound/Components/MusicClockComponent.h
@@ -59,7 +59,9 @@ public:
 	UMusicClockComponent();
 
 	void BeginPlay() override;
+	void BeginPlayInternal();
 	void EndPlay(const EEndPlayReason::Type EndPlayReason) override;
+	void EndPlayInternal();
 
 	UPROPERTY(EditDefaultsOnly, Category = "MusicClock")
 	EMusicClockDriveMethod DriveMethod = EMusicClockDriveMethod::MetaSound;
@@ -384,7 +386,8 @@ public:
 
 public:
 	virtual void TickComponent(float DeltaTime, enum ELevelTick TickType, FActorComponentTickFunction* ThisTickFunction);
-
+	void TickComponentInternal();
+	
 private:
 	friend struct FMusicClockDriverBase;
 	friend struct FMetasoundMusicClockDriver;
