Compile HarmonixMetasound with AutoRTFM instrumentation enabled as it is reachable from Verse.

#rnx

diff --git a/Engine/Plugins/Runtime/Harmonix/Source/HarmonixMetasound/HarmonixMetasound.Build.cs b/Engine/Plugins/Runtime/Harmonix/Source/HarmonixMetasound/HarmonixMetasound.Build.cs
--- a/Engine/Plugins/Runtime/Harmonix/Source/HarmonixMetasound/HarmonixMetasound.Build.cs
+++ b/Engine/Plugins/Runtime/Harmonix/Source/HarmonixMetasound/HarmonixMetasound.Build.cs
@@ -44,7 +44,5 @@ public class HarmonixMetasound : ModuleRules
 			PrivateDependencyModuleNames.Add("AssetRegistry");
 			PrivateDependencyModuleNames.Add("UnrealEd");
 		}
-
-		bDisableAutoRTFMInstrumentation = true;
 	}
 }
