[Harmonix] fix issue where Midi Clock State node was not correctly reporting tempo until a tempo change event was received, and equivalent behavior for other outputs

#tests tested in PIE

diff --git a/Engine/Plugins/Runtime/Harmonix/Source/HarmonixMetasound/Private/HarmonixMetasound/Nodes/MidiClockStateNode.cpp b/Engine/Plugins/Runtime/Harmonix/Source/HarmonixMetasound/Private/HarmonixMetasound/Nodes/MidiClockStateNode.cpp
--- a/Engine/Plugins/Runtime/Harmonix/Source/HarmonixMetasound/Private/HarmonixMetasound/Nodes/MidiClockStateNode.cpp
+++ b/Engine/Plugins/Runtime/Harmonix/Source/HarmonixMetasound/Private/HarmonixMetasound/Nodes/MidiClockStateNode.cpp
@@ -215,12 +215,12 @@ namespace HarmonixMetasound
 		const int32 NextTick = MidiClockInPin->GetNextMidiTickToProcess();
 
 		bool HasTempoSpeedOrTimeSigChange = false;
-		if (MidiClockInPin->HasTempoChangesInBlock())
+		if (MidiClockInPin->GetTempoAtEndOfBlock() != *CurrentTempoOutPin)
 		{
 			HasTempoSpeedOrTimeSigChange = true;
 			*CurrentTempoOutPin = FMath::Max(kTinyTempoSpeed, MidiClockInPin->GetTempoAtEndOfBlock());
 		}
-		if (MidiClockInPin->HasSpeedChangesInBlock())
+		if (MidiClockInPin->GetSpeedAtEndOfBlock() != *CurrentSpeedOutPin)
 		{
 			HasTempoSpeedOrTimeSigChange = true;
 			*CurrentSpeedOutPin = FMath::Max(kTinyTempoSpeed, MidiClockInPin->GetSpeedAtEndOfBlock());
