Get 'none' timesignature mapping mode working again.

diff --git a/Engine/Plugins/Runtime/Harmonix/Source/HarmonixMetasound/Private/HarmonixMetasound/Components/MusicClockComponent.cpp b/Engine/Plugins/Runtime/Harmonix/Source/HarmonixMetasound/Private/HarmonixMetasound/Components/MusicClockComponent.cpp
--- a/Engine/Plugins/Runtime/Harmonix/Source/HarmonixMetasound/Private/HarmonixMetasound/Components/MusicClockComponent.cpp
+++ b/Engine/Plugins/Runtime/Harmonix/Source/HarmonixMetasound/Private/HarmonixMetasound/Components/MusicClockComponent.cpp
@@ -804,13 +804,22 @@ FMusicalTime UMusicClockComponent::GetMusicalTime() const
 	return Result;
 }
 
+int32 UMusicClockComponent::GetAbsoluteTickPosition() const
+{
+	// Note: Since this is an override of a function in the music environment system we need to
+	// be sure we are using its standard ticks per quarter note units.
+	const ISongMapEvaluator& MapEvaluator = GetSongMaps();
+	float CurrentTick = MapEvaluator.FractionalBarIncludingCountInToTick(CurrentVideoRenderSongPos.BarsIncludingCountIn);
+	return FMath::FloorToInt32((CurrentTick / static_cast<float>(MapEvaluator.GetTicksPerQuarterNote())) * static_cast<float>(MusicalTime::TicksPerQuarterNote));
+}
+
+
 FMusicalTime UMusicClockComponent::GetMusicalTimeWithSourceSpaceOffset(const FMusicalTime& Offset) const
 {
 	FMusicalTime Result;
 	const ISongMapEvaluator& MapEvaluator = GetSongMaps();
 	float CurrentTick = MapEvaluator.FractionalBarIncludingCountInToTick(CurrentVideoRenderSongPos.BarsIncludingCountIn);
 	
-	//float AdjustedTick = CurrentTick - (static_cast<float>(Offset.TickInBar * MapEvaluator.GetTicksPerQuarterNote()) / static_cast<float>(MusicalTime::TicksPerQuarterNote));
 	float OffsetTick = MapEvaluator.FractionalBarIncludingCountInToTick(static_cast<float>(Offset.FractionalBar()));
 	float AdjustedTick = CurrentTick - OffsetTick;
 
@@ -823,6 +832,17 @@ FMusicalTime UMusicClockComponent::GetMusicalTimeWithSourceSpaceOffset(const FMu
 	return Result;
 }
 
+int32 UMusicClockComponent::GetAbsoluteTickPositionWithSourceSpaceOffset(const FMusicalTime& Offset) const
+{
+	// Note: Since this is an override of a function in the music environment system we need to
+	// be sure we are using its standard ticks per quarter note units.
+	const ISongMapEvaluator& MapEvaluator = GetSongMaps();
+	float CurrentTick = MapEvaluator.FractionalBarIncludingCountInToTick(CurrentVideoRenderSongPos.BarsIncludingCountIn);
+	float OffsetTick = MapEvaluator.FractionalBarIncludingCountInToTick(static_cast<float>(Offset.FractionalBar()));
+	float AdjustedTick = CurrentTick - OffsetTick;
+	return FMath::FloorToInt32((AdjustedTick / static_cast<float>(MapEvaluator.GetTicksPerQuarterNote())) * static_cast<float>(MusicalTime::TicksPerQuarterNote));
+}
+
 FMusicalTime UMusicClockComponent::Quantize(const FMusicalTime& MusicalTime, int32 QuantizationInterval, UFrameBasedMusicMap::EQuantizeDirection Direction) const
 {
 	FMusicalTime Result;
diff --git a/Engine/Plugins/Runtime/Harmonix/Source/HarmonixMetasound/Private/HarmonixMetasound/MusicEnvironmentSupport/HarmonixMusicEnvironmentMetronome.cpp b/Engine/Plugins/Runtime/Harmonix/Source/HarmonixMetasound/Private/HarmonixMetasound/MusicEnvironmentSupport/HarmonixMusicEnvironmentMetronome.cpp
--- a/Engine/Plugins/Runtime/Harmonix/Source/HarmonixMetasound/Private/HarmonixMetasound/MusicEnvironmentSupport/HarmonixMusicEnvironmentMetronome.cpp
+++ b/Engine/Plugins/Runtime/Harmonix/Source/HarmonixMetasound/Private/HarmonixMetasound/MusicEnvironmentSupport/HarmonixMusicEnvironmentMetronome.cpp
@@ -106,6 +106,15 @@ FMusicalTime UHarmonixMusicEnvironmentMetronome::GetMusicalTime() const
 	return MusicalTime;
 }
 
+int32 UHarmonixMusicEnvironmentMetronome::GetAbsoluteTickPosition() const
+{
+	if (MusicClockComponent)
+	{
+		return MusicClockComponent->GetAbsoluteTickPosition();
+	}
+	return 0;
+}
+
 FMusicalTime UHarmonixMusicEnvironmentMetronome::GetMusicalTimeWithSourceSpaceOffset(const FMusicalTime& Offset) const
 {
 	FMusicalTime MusicalTime;
@@ -116,6 +125,15 @@ FMusicalTime UHarmonixMusicEnvironmentMetronome::GetMusicalTimeWithSourceSpaceOf
 	return MusicalTime;
 }
 
+int32 UHarmonixMusicEnvironmentMetronome::GetAbsoluteTickPositionWithSourceSpaceOffset(const FMusicalTime& Offset) const
+{
+	if (MusicClockComponent)
+	{
+		return MusicClockComponent->GetAbsoluteTickPositionWithSourceSpaceOffset(Offset);
+	}
+	return 0;
+}
+
 FMusicalTime UHarmonixMusicEnvironmentMetronome::Quantize(const FMusicalTime& MusicalTime, int32 QuantizationInterval, UFrameBasedMusicMap::EQuantizeDirection Direction) const
 {
 	if (!MusicClockComponent)
diff --git a/Engine/Plugins/Runtime/Harmonix/Source/HarmonixMetasound/Private/HarmonixMetasound/MusicEnvironmentSupport/HarmonixMusicEnvironmentMetronome.h b/Engine/Plugins/Runtime/Harmonix/Source/HarmonixMetasound/Private/HarmonixMetasound/MusicEnvironmentSupport/HarmonixMusicEnvironmentMetronome.h
--- a/Engine/Plugins/Runtime/Harmonix/Source/HarmonixMetasound/Private/HarmonixMetasound/MusicEnvironmentSupport/HarmonixMusicEnvironmentMetronome.h
+++ b/Engine/Plugins/Runtime/Harmonix/Source/HarmonixMetasound/Private/HarmonixMetasound/MusicEnvironmentSupport/HarmonixMusicEnvironmentMetronome.h
@@ -36,10 +36,15 @@ public:
 	virtual float GetCurrentTempo() override;
 	virtual double GetCurrentPositionSeconds() override;
 
+
+	// vvv BEGIN  IMusicEnvironmentClockSource Implementation vvv
 	virtual FMusicalTime GetMusicalTime() const override;
+	virtual int32 GetAbsoluteTickPosition() const override;
 	virtual FMusicalTime GetMusicalTimeWithSourceSpaceOffset(const FMusicalTime& Offset) const override;
+	virtual int32 GetAbsoluteTickPositionWithSourceSpaceOffset(const FMusicalTime& Offset) const override;
 	virtual FMusicalTime Quantize(const FMusicalTime& MusicalTime, int32 QuantizationInterval, UFrameBasedMusicMap::EQuantizeDirection Direction = UFrameBasedMusicMap::EQuantizeDirection::Nearest) const override;
 	virtual bool CanAuditionInEditor() const override { return true; }
+	// ^^^ END IMusicEnvironmentClockSource Implementation ^^^
 
 	virtual void BeginDestroy() override;
 	
diff --git a/Engine/Plugins/Runtime/Harmonix/Source/HarmonixMetasound/Public/HarmonixMetasound/Components/MusicClockComponent.h b/Engine/Plugins/Runtime/Harmonix/Source/HarmonixMetasound/Public/HarmonixMetasound/Components/MusicClockComponent.h
--- a/Engine/Plugins/Runtime/Harmonix/Source/HarmonixMetasound/Public/HarmonixMetasound/Components/MusicClockComponent.h
+++ b/Engine/Plugins/Runtime/Harmonix/Source/HarmonixMetasound/Public/HarmonixMetasound/Components/MusicClockComponent.h
@@ -377,7 +377,9 @@ public:
 
 	// vvv BEGIN  IMusicEnvironmentClockSource Implementation vvv
 	virtual FMusicalTime GetMusicalTime() const override;
+	virtual int32 GetAbsoluteTickPosition() const override;
 	virtual FMusicalTime GetMusicalTimeWithSourceSpaceOffset(const FMusicalTime& Offset) const;
+	virtual int32 GetAbsoluteTickPositionWithSourceSpaceOffset(const FMusicalTime& Offset) const override;
 	virtual FMusicalTime Quantize(const FMusicalTime& MusicalTime, int32 QuantizationInterval, UFrameBasedMusicMap::EQuantizeDirection Direction = UFrameBasedMusicMap::EQuantizeDirection::Nearest) const;
 	virtual bool CanAuditionInEditor() const override { return false; }
 	// ^^^ END IMusicEnvironmentClockSource Implementation ^^^
diff --git a/Engine/Source/Runtime/MusicEnvironment/Private/FrameBasedMusicMap.cpp b/Engine/Source/Runtime/MusicEnvironment/Private/FrameBasedMusicMap.cpp
--- a/Engine/Source/Runtime/MusicEnvironment/Private/FrameBasedMusicMap.cpp
+++ b/Engine/Source/Runtime/MusicEnvironment/Private/FrameBasedMusicMap.cpp
@@ -177,6 +177,17 @@ FFrameTime UFrameBasedMusicMap::GetFrameTime(double FractionalBar) const
 	return GetFrameTime(Bar, Beat);	
 }
 
+FFrameTime UFrameBasedMusicMap::GetFrameTime(int32 AtTick) const
+{
+	if (IsEmpty())
+	{
+		return FFrameTime();
+	}
+
+	FMusicalTime MusicalTime = GetMusicalTime(AtTick);
+	return GetFrameTime(MusicalTime);
+}
+
 double UFrameBasedMusicMap::GetSeconds(const FMusicalTime& MusicalTime) const
 {
 	if (IsEmpty())
diff --git a/Engine/Source/Runtime/MusicEnvironment/Public/FrameBasedMusicMap.h b/Engine/Source/Runtime/MusicEnvironment/Public/FrameBasedMusicMap.h
--- a/Engine/Source/Runtime/MusicEnvironment/Public/FrameBasedMusicMap.h
+++ b/Engine/Source/Runtime/MusicEnvironment/Public/FrameBasedMusicMap.h
@@ -316,6 +316,8 @@ public:
 	FFrameTime GetFrameTime(const FMusicalTime& MusicalTime) const;
 	FFrameTime GetFrameTime(int32 AtBar, float AtBeat) const;
 	FFrameTime GetFrameTime(double FractionalBar) const;
+	FFrameTime GetFrameTime(int32 AtTick) const;
+
 	double GetSeconds(const FMusicalTime& MusicalTime) const;
 	double GetSeconds(int32 AtBar, float AtBeat = 0.0f) const;
 
diff --git a/Engine/Source/Runtime/MusicEnvironment/Public/MusicEnvironmentClockSource.h b/Engine/Source/Runtime/MusicEnvironment/Public/MusicEnvironmentClockSource.h
--- a/Engine/Source/Runtime/MusicEnvironment/Public/MusicEnvironmentClockSource.h
+++ b/Engine/Source/Runtime/MusicEnvironment/Public/MusicEnvironmentClockSource.h
@@ -25,8 +25,9 @@ class MUSICENVIRONMENT_API IMusicEnvironmentClockSource
 
 public:
 	virtual FMusicalTime GetMusicalTime() const = 0;
-	//virtual FMusicalTime GetBarOneBeatOneMusicalTime() const = 0; // returns bar 0 beat 0 But ticks per beat and ticks per bar are the initial tempo!
+	virtual int32 GetAbsoluteTickPosition() const = 0;
 	virtual FMusicalTime GetMusicalTimeWithSourceSpaceOffset(const FMusicalTime& Offset) const = 0;
+	virtual int32 GetAbsoluteTickPositionWithSourceSpaceOffset(const FMusicalTime& SourceOffset) const = 0;
 	virtual FMusicalTime Quantize(const FMusicalTime& MusicalTime, int32 QuantizationInterval, UFrameBasedMusicMap::EQuantizeDirection Direction = UFrameBasedMusicMap::EQuantizeDirection::Nearest) const = 0;
 	virtual bool CanAuditionInEditor() const = 0;
 };
