[Harmonix + Metasound] Fix issue where a metasound generator handle whose underlying metasound generator is stopped does not properly broadcast a message that the generator has detached.  Fix bug where in MetasoundClockDriver where CurrentMapChain wouldn't get reassigned from the new ClockHistory when ClockHistory changes.

#rnx
#tests PIE jam in lobby, mainstage

diff --git a/Engine/Plugins/Runtime/Harmonix/Source/HarmonixMetasound/Private/HarmonixMetasound/Components/MetasoundMusicClockDriver.cpp b/Engine/Plugins/Runtime/Harmonix/Source/HarmonixMetasound/Private/HarmonixMetasound/Components/MetasoundMusicClockDriver.cpp
--- a/Engine/Plugins/Runtime/Harmonix/Source/HarmonixMetasound/Private/HarmonixMetasound/Components/MetasoundMusicClockDriver.cpp
+++ b/Engine/Plugins/Runtime/Harmonix/Source/HarmonixMetasound/Private/HarmonixMetasound/Components/MetasoundMusicClockDriver.cpp
@@ -4,7 +4,6 @@
 #include "Components/AudioComponent.h"
 #include "HarmonixMetasound/Analysis/MidiSongPosVertexAnalyzer.h"
 #include "HarmonixMetasound/DataTypes/MidiClock.h"
-#include "MetasoundGeneratorHandle.h"
 #include "MetasoundGenerator.h"
 #include "Engine/World.h"
 #include "Async/Async.h"
@@ -280,7 +279,7 @@ void FMetasoundMusicClockDriver::OnStop()
 void FMetasoundMusicClockDriver::Disconnect()
 {
 	check(IsInGameThread());
-	ClockHistory = nullptr;
+	SetClockHistory(nullptr);
 	DetachAllCallbacks();
 	AudioComponentToWatch.Reset();
 	CurrentGeneratorHandle.Reset();
@@ -359,13 +358,13 @@ void FMetasoundMusicClockDriver::DetachAllCallbacks()
 		GraphChangedCallbackHandle.Reset();
 
 	}
-	ClockHistory = nullptr;
+	SetClockHistory(nullptr);
 }
 
 void FMetasoundMusicClockDriver::OnGeneratorAttached()
 {
 	WasEverConnected = true;
-	ClockHistory = UMidiClockUpdateSubsystem::GetOrCreateClockHistory(MidiSongPosAnalyzerAddress);
+	SetClockHistory(UMidiClockUpdateSubsystem::GetOrCreateClockHistory(MidiSongPosAnalyzerAddress));
 	SmoothedAudioRenderClockHistoryCursor = ClockHistory->CreateReadCursor();
 	SmoothedPlayerExperienceClockHistoryCursor = ClockHistory->CreateReadCursor();
 	SmoothedVideoRenderClockHistoryCursor = ClockHistory->CreateReadCursor();
@@ -374,7 +373,7 @@ void FMetasoundMusicClockDriver::OnGeneratorAttached()
 
 void FMetasoundMusicClockDriver::OnGraphSet()
 {
-	ClockHistory = UMidiClockUpdateSubsystem::GetOrCreateClockHistory(MidiSongPosAnalyzerAddress);
+	SetClockHistory(UMidiClockUpdateSubsystem::GetOrCreateClockHistory(MidiSongPosAnalyzerAddress));
 	SmoothedAudioRenderClockHistoryCursor = ClockHistory->CreateReadCursor();
 	SmoothedPlayerExperienceClockHistoryCursor = ClockHistory->CreateReadCursor();
 	SmoothedVideoRenderClockHistoryCursor = ClockHistory->CreateReadCursor();
@@ -387,7 +386,7 @@ void FMetasoundMusicClockDriver::OnGeneratorIOUpdatedWithChanges(const TArray<Me
 	{
 		for (auto& VertexInterfaceChange : VertexInterfaceChanges)
 		{
-			ClockHistory = UMidiClockUpdateSubsystem::GetOrCreateClockHistory(MidiSongPosAnalyzerAddress);
+			SetClockHistory(UMidiClockUpdateSubsystem::GetOrCreateClockHistory(MidiSongPosAnalyzerAddress));
 			SmoothedAudioRenderClockHistoryCursor = ClockHistory->CreateReadCursor();
 			SmoothedPlayerExperienceClockHistoryCursor = ClockHistory->CreateReadCursor();
 			SmoothedVideoRenderClockHistoryCursor = ClockHistory->CreateReadCursor();
@@ -411,7 +410,7 @@ void FMetasoundMusicClockDriver::OnGeneratorDetached()
 		SongPosOffsetMs = ClockComponent->CurrentSmoothedAudioRenderSongPos.SecondsIncludingCountIn * 1000.0f;
 		FreeRunStartTimeSecs = GetWallClockTime();
 	}
-	ClockHistory = nullptr;
+	SetClockHistory(nullptr);
 	SmoothedAudioRenderClockHistoryCursor = HarmonixMetasound::Analysis::FMidiClockSongPositionHistory::FReadCursor();
 	SmoothedPlayerExperienceClockHistoryCursor = HarmonixMetasound::Analysis::FMidiClockSongPositionHistory::FReadCursor();
 	SmoothedVideoRenderClockHistoryCursor = HarmonixMetasound::Analysis::FMidiClockSongPositionHistory::FReadCursor();
@@ -597,6 +596,17 @@ void FMetasoundMusicClockDriver::RefreshCurrentSongPosFromHistory()
 	}
 }
 
+void FMetasoundMusicClockDriver::SetClockHistory(const UMidiClockUpdateSubsystem::FClockHistoryPtr& NewHistory)
+{
+	if (NewHistory != ClockHistory)
+	{
+		ClockHistory = NewHistory;
+		// We null this out because ultimately we get the CurrentMapChain from the ClockHistory so if it ever
+		// changes, we need to get a new CurrentMapChain assigned.
+		CurrentMapChain = nullptr;
+	}
+}
+
 void FMetasoundMusicClockDriver::UpdateCurrentTicksForOffsetClock(float SmoothedTick, float SmoothedTempoMapTick)
 {
 	float SmoothedPositionMs = CurrentMapChain->SongMaps->TickToMs(SmoothedTick);
@@ -680,7 +690,7 @@ FMetasoundMusicClockDriver::EHistoryFailureType FMetasoundMusicClockDriver::Calc
 {
 	using namespace HarmonixMetasound::Analysis;
 
-	// A little book keeping for tracking...
+	// A little bookkeeping for tracking...
 	if (MetasoundMusicClockDriver::HighWaterNumDataAvailable < ReadCursor.NumDataAvailable())
 	{
 		MetasoundMusicClockDriver::HighWaterNumDataAvailable = ReadCursor.NumDataAvailable();
diff --git a/Engine/Plugins/Runtime/Harmonix/Source/HarmonixMetasound/Private/HarmonixMetasound/Components/MetasoundMusicClockDriver.h b/Engine/Plugins/Runtime/Harmonix/Source/HarmonixMetasound/Private/HarmonixMetasound/Components/MetasoundMusicClockDriver.h
--- a/Engine/Plugins/Runtime/Harmonix/Source/HarmonixMetasound/Private/HarmonixMetasound/Components/MetasoundMusicClockDriver.h
+++ b/Engine/Plugins/Runtime/Harmonix/Source/HarmonixMetasound/Private/HarmonixMetasound/Components/MetasoundMusicClockDriver.h
@@ -97,6 +97,8 @@ private:
 	void DetachAllCallbacks();
 	void RefreshCurrentSongPosFromWallClock();
 	void RefreshCurrentSongPosFromHistory();
+	// always go through this to change ClockHistory
+	void SetClockHistory(const UMidiClockUpdateSubsystem::FClockHistoryPtr&);
 
 	enum class EHistoryFailureType : uint8
 	{
diff --git a/Engine/Plugins/Runtime/Metasound/Source/MetasoundEngine/Private/MetasoundGeneratorHandle.cpp b/Engine/Plugins/Runtime/Metasound/Source/MetasoundEngine/Private/MetasoundGeneratorHandle.cpp
--- a/Engine/Plugins/Runtime/Metasound/Source/MetasoundEngine/Private/MetasoundGeneratorHandle.cpp
+++ b/Engine/Plugins/Runtime/Metasound/Source/MetasoundEngine/Private/MetasoundGeneratorHandle.cpp
@@ -224,9 +224,19 @@ namespace Metasound
 		
 		check(IsInGameThread());
 
-		// early out if the incoming generator is null and the current generator is null
+		// early out if the incoming generator is null and the current generator is already invalid...
 		if (!Generator.IsValid() && !InGenerator.IsValid())
 		{
+			// If the generator is invalid but WAS, at some point, pointing at a generator,
+			// we need to execute the OnGeneratorSet delegates so they know the generator 
+			// went away!
+			if (Generator.GetWeakPtrTypeHash() != 0 && OnGeneratorSet.IsBound())
+			{
+				OnGeneratorSet.Execute(TWeakPtr<FMetasoundGenerator>(InGenerator));
+			}
+
+			// Now reset the generator so it is back to nullptr...
+			Generator.Reset();
 			return;
 		}
 		
