[Harmonix] MidiSongPos better tracks whether it is invalid or not, so that clients of it can tell if the song has *actually* started playing and set MidiSongPos values yet.

#tests PIE jam in lobby, mainstage
#rnx

diff --git a/Engine/Plugins/Runtime/Harmonix/Source/HarmonixMetasound/Private/HarmonixMetasound/Components/MusicClockComponent.cpp b/Engine/Plugins/Runtime/Harmonix/Source/HarmonixMetasound/Private/HarmonixMetasound/Components/MusicClockComponent.cpp
--- a/Engine/Plugins/Runtime/Harmonix/Source/HarmonixMetasound/Private/HarmonixMetasound/Components/MusicClockComponent.cpp
+++ b/Engine/Plugins/Runtime/Harmonix/Source/HarmonixMetasound/Private/HarmonixMetasound/Components/MusicClockComponent.cpp
@@ -7,7 +7,6 @@
 #include "HarmonixMidi/MusicTimeSpan.h"
 #include "Components/AudioComponent.h"
 #include "Templates/SharedPointer.h"
-#include "Misc/Optional.h"
 
 #include UE_INLINE_GENERATED_CPP_BY_NAME(MusicClockComponent)
 
@@ -201,24 +200,7 @@ FMidiSongPos UMusicClockComponent::CalculateSongPosWithOffset(float MsOffset, EC
 	}
 
 	// otherwise, use our song maps copy
-	const FSongMaps* Maps = &DefaultMaps;
-	switch (Timebase)
-	{
-	case ECalibratedMusicTimebase::RawAudioRenderTime:
-		Result.SetByTime((CurrentRawAudioRenderSongPos.SecondsIncludingCountIn * 1000.0f) + MsOffset, *Maps);
-		break;
-	case ECalibratedMusicTimebase::AudioRenderTime:
-		Result.SetByTime((CurrentSmoothedAudioRenderSongPos.SecondsIncludingCountIn * 1000.0f) + MsOffset, *Maps);
-		break;
-	case ECalibratedMusicTimebase::ExperiencedTime:
-		Result.SetByTime((CurrentPlayerExperiencedSongPos.SecondsIncludingCountIn * 1000.0f) + MsOffset, *Maps);
-		break;
-	case ECalibratedMusicTimebase::VideoRenderTime:
-	default:
-		Result.SetByTime((CurrentVideoRenderSongPos.SecondsIncludingCountIn * 1000.0f) + MsOffset, *Maps);
-		break;
-	}
-
+	Result.SetByTime((GetCurrentSongPosInternal(Timebase).SecondsIncludingCountIn * 1000.0f) + MsOffset, DefaultMaps);
 	return Result;
 }
 
@@ -328,82 +310,39 @@ void UMusicClockComponent::Stop()
 
 float UMusicClockComponent::GetSecondsIncludingCountIn(ECalibratedMusicTimebase Timebase) const
 {
-	EnsureClockIsValidForGameFrame();
-	switch (Timebase)
-	{
-	case ECalibratedMusicTimebase::RawAudioRenderTime: return CurrentRawAudioRenderSongPos.SecondsIncludingCountIn;
-	case ECalibratedMusicTimebase::AudioRenderTime:	   return CurrentSmoothedAudioRenderSongPos.SecondsIncludingCountIn;
-	case ECalibratedMusicTimebase::ExperiencedTime:	   return CurrentPlayerExperiencedSongPos.SecondsIncludingCountIn;
-	case ECalibratedMusicTimebase::VideoRenderTime:
-	default:                                           return CurrentVideoRenderSongPos.SecondsIncludingCountIn;
-	}
+	return GetSongPos(Timebase).SecondsIncludingCountIn;
 }
 
 float UMusicClockComponent::GetSecondsFromBarOne(ECalibratedMusicTimebase Timebase) const
 {
-	EnsureClockIsValidForGameFrame();
-	switch (Timebase)
-	{
-	case ECalibratedMusicTimebase::RawAudioRenderTime: return CurrentRawAudioRenderSongPos.SecondsFromBarOne;
-	case ECalibratedMusicTimebase::AudioRenderTime:	   return CurrentSmoothedAudioRenderSongPos.SecondsFromBarOne;
-	case ECalibratedMusicTimebase::ExperiencedTime:	   return CurrentPlayerExperiencedSongPos.SecondsFromBarOne;
-	case ECalibratedMusicTimebase::VideoRenderTime:
-	default:                                           return CurrentVideoRenderSongPos.SecondsFromBarOne;
-	}
+	return GetSongPos(Timebase).SecondsFromBarOne;
 }
 
 float UMusicClockComponent::GetBarsIncludingCountIn(ECalibratedMusicTimebase Timebase) const
 {
-	EnsureClockIsValidForGameFrame();
-	switch (Timebase)
-	{
-	case ECalibratedMusicTimebase::RawAudioRenderTime: return CurrentRawAudioRenderSongPos.BarsIncludingCountIn;
-	case ECalibratedMusicTimebase::AudioRenderTime:	   return CurrentSmoothedAudioRenderSongPos.BarsIncludingCountIn;
-	case ECalibratedMusicTimebase::ExperiencedTime:	   return CurrentPlayerExperiencedSongPos.BarsIncludingCountIn;
-	case ECalibratedMusicTimebase::VideoRenderTime:
-	default:                                           return CurrentVideoRenderSongPos.BarsIncludingCountIn;
-	}
+	return GetSongPos(Timebase).BarsIncludingCountIn;
 }
 
 float UMusicClockComponent::GetBeatsIncludingCountIn(ECalibratedMusicTimebase Timebase) const
 {
-	EnsureClockIsValidForGameFrame();
-	switch (Timebase)
-	{
-	case ECalibratedMusicTimebase::RawAudioRenderTime: return CurrentRawAudioRenderSongPos.BarsIncludingCountIn;
-	case ECalibratedMusicTimebase::AudioRenderTime:    return CurrentSmoothedAudioRenderSongPos.BeatsIncludingCountIn;
-	case ECalibratedMusicTimebase::ExperiencedTime:	   return CurrentPlayerExperiencedSongPos.BeatsIncludingCountIn;
-	case ECalibratedMusicTimebase::VideoRenderTime:
-	default:                                           return CurrentVideoRenderSongPos.BeatsIncludingCountIn;
-	}
+	return GetSongPos(Timebase).BeatsIncludingCountIn;
 }
 
 float UMusicClockComponent::GetTicksFromBarOne(ECalibratedMusicTimebase Timebase) const
 {
-	EnsureClockIsValidForGameFrame();
-	float Seconds = GetSecondsFromBarOne(Timebase);
+	const float Seconds = GetSecondsFromBarOne(Timebase);
 	return GetSongMaps().MsToTick(Seconds * 1000.0f);
 }
 
 float UMusicClockComponent::GetTicksIncludingCountIn(ECalibratedMusicTimebase Timebase) const
 {
-	EnsureClockIsValidForGameFrame();
-	float Seconds = GetSecondsIncludingCountIn(Timebase);
+	const float Seconds = GetSecondsIncludingCountIn(Timebase);
 	return GetSongMaps().MsToTick(Seconds * 1000.0f);
 }
 
-
 FMusicTimestamp UMusicClockComponent::GetCurrentTimestamp(ECalibratedMusicTimebase Timebase) const
 {
-	EnsureClockIsValidForGameFrame();
-	switch (Timebase)
-	{
-	case ECalibratedMusicTimebase::RawAudioRenderTime: return CurrentRawAudioRenderSongPos.Timestamp;
-	case ECalibratedMusicTimebase::AudioRenderTime:	   return CurrentSmoothedAudioRenderSongPos.Timestamp;
-	case ECalibratedMusicTimebase::ExperiencedTime:	   return CurrentPlayerExperiencedSongPos.Timestamp;
-	case ECalibratedMusicTimebase::VideoRenderTime:
-	default:                                           return CurrentVideoRenderSongPos.Timestamp;
-	}
+	return GetSongPos(Timebase).Timestamp;
 }
 
 FString UMusicClockComponent::GetCurrentSectionName(ECalibratedMusicTimebase Timebase) const
@@ -615,46 +554,17 @@ const ISongMapEvaluator& UMusicClockComponent::GetSongMaps() const
 	return SongMaps ? *SongMaps : DefaultMaps;
 }
 
-bool UMusicClockComponent::HasValidSongPos() const
-{
-	return State != EMusicClockState::Stopped && LastBroadcastBeat != -1;
-}
-
 
 const FMidiSongPos& UMusicClockComponent::GetSongPos(ECalibratedMusicTimebase Timebase) const
 {
 	EnsureClockIsValidForGameFrame();
-
-	switch (Timebase)
-	{
-	case ECalibratedMusicTimebase::RawAudioRenderTime:
-		return CurrentRawAudioRenderSongPos;
-	case ECalibratedMusicTimebase::AudioRenderTime:
-		return CurrentSmoothedAudioRenderSongPos;
-	case ECalibratedMusicTimebase::ExperiencedTime:
-		return CurrentPlayerExperiencedSongPos;
-	case ECalibratedMusicTimebase::VideoRenderTime:
-	default:
-		return CurrentVideoRenderSongPos;
-	}
+	return GetCurrentSongPosInternal(Timebase);
 }
 
 const FMidiSongPos& UMusicClockComponent::GetPreviousSongPos(ECalibratedMusicTimebase Timebase) const
 {
 	EnsureClockIsValidForGameFrame();
-
-	switch (Timebase)
-	{
-	case ECalibratedMusicTimebase::RawAudioRenderTime:
-		return PrevRawAudioRenderSongPos;
-	case ECalibratedMusicTimebase::AudioRenderTime:
-		return PrevAudioRenderSongPos;
-	case ECalibratedMusicTimebase::ExperiencedTime:
-		return PrevPlayerExperiencedSongPos;
-	case ECalibratedMusicTimebase::VideoRenderTime:
-	default:
-		return PrevVideoRenderSongPos;
-	}
+	return GetPreviousSongPosInternal(Timebase);
 }
 
 FMidiSongPos UMusicClockComponent::GetCurrentSmoothedAudioRenderSongPos() const
@@ -695,58 +605,34 @@ FMidiSongPos UMusicClockComponent::GetCurrentRawAudioRenderSongPos() const
 float UMusicClockComponent::MeasureSpanProgress(const FMusicalTimeSpan& Span, ECalibratedMusicTimebase Timebase) const
 {
 	EnsureClockIsValidForGameFrame();
-
-	const ISongMapEvaluator* Maps = ClockDriver ? ClockDriver->GetCurrentSongMapEvaluator() : &DefaultMaps;
-	if (!Maps)
+	if (const ISongMapEvaluator* Maps = ClockDriver ? ClockDriver->GetCurrentSongMapEvaluator() : &DefaultMaps)
 	{
-		return 0.0f;
-	}
-
-	const FMidiSongPos* Basis = &CurrentVideoRenderSongPos;
-	switch (Timebase)
-	{
-	case ECalibratedMusicTimebase::RawAudioRenderTime:
-		Basis = &CurrentRawAudioRenderSongPos;
-		break;
-	case ECalibratedMusicTimebase::AudioRenderTime:
-		Basis = &CurrentSmoothedAudioRenderSongPos;
-		break;
-	case ECalibratedMusicTimebase::ExperiencedTime:
-		Basis = &CurrentPlayerExperiencedSongPos;
-		break;
+		return Span.CalcPositionInSpan(GetCurrentSongPosInternal(Timebase), *Maps);
 	}
-
-	return Span.CalcPositionInSpan(*Basis, *Maps);
+	return 0.0f;
 }
 
 void UMusicClockComponent::BroadcastSongPosChanges()
 {
-	const FMidiSongPos* Basis = &CurrentVideoRenderSongPos;
-	switch (TimebaseForBarAndBeatEvents)
-	{
-	case ECalibratedMusicTimebase::RawAudioRenderTime:
-		Basis = &CurrentRawAudioRenderSongPos;
-		break;
-	case ECalibratedMusicTimebase::AudioRenderTime:
-		Basis = &CurrentSmoothedAudioRenderSongPos;
-		break;
-	case ECalibratedMusicTimebase::ExperiencedTime:
-		Basis = &CurrentPlayerExperiencedSongPos;
-		break;
-	}
-	int32 CurrBar = FMath::FloorToInt32(Basis->BarsIncludingCountIn);
+	const FMidiSongPos& Basis = GetCurrentSongPosInternal(TimebaseForBarAndBeatEvents);
+
+	// BUG There's a bug here.  We should probably check Basis for being valid
+	// before taking any action.  Something like: if (!Basis->IsValid()) return;
+	// But that's potentially a big change so do it as a separate checkin -jwf
+	
+	const int32 CurrBar = FMath::FloorToInt32(Basis.BarsIncludingCountIn);
 	if (LastBroadcastBar != CurrBar)
 	{
-		BarEvent.Broadcast(Basis->Timestamp.Bar);
+		BarEvent.Broadcast(Basis.Timestamp.Bar);
 		LastBroadcastBar = CurrBar;
 	}
-	int32 CurrBeat = FMath::FloorToInt32(Basis->BeatsIncludingCountIn);
+	const int32 CurrBeat = FMath::FloorToInt32(Basis.BeatsIncludingCountIn);
 	if (LastBroadcastBeat != CurrBeat)
 	{
-		BeatEvent.Broadcast(CurrBeat, FMath::FloorToInt32(Basis->Timestamp.Beat));
+		BeatEvent.Broadcast(CurrBeat, FMath::FloorToInt32(Basis.Timestamp.Beat));
 		LastBroadcastBeat = CurrBeat;
 	}
-	const FSongSection& SongSection = Basis->CurrentSongSection;
+	const FSongSection& SongSection = Basis.CurrentSongSection;
 	if (LastBroadcastSongSection.StartTick != SongSection.StartTick || LastBroadcastSongSection.LengthTicks != SongSection.LengthTicks)
 	{
 		SectionEvent.Broadcast(SongSection.Name, SongSection.StartTick, SongSection.LengthTicks);
@@ -754,7 +640,7 @@ void UMusicClockComponent::BroadcastSongPosChanges()
 	}
 }
 
-void UMusicClockComponent::BroadcastSeekLoopDetections()
+void UMusicClockComponent::BroadcastSeekLoopDetections() const
 {
 	if(AudioRenderSeekDetected)
 	{
@@ -790,9 +676,45 @@ void UMusicClockComponent::MakeDefaultSongMap()
 	DefaultMaps.GetBarMap().AddTimeSignatureAtBarIncludingCountIn(0, TimeSignatureNum, TimeSignatureDenom);
 }
 
+const FMidiSongPos& UMusicClockComponent::GetCurrentSongPosInternal(ECalibratedMusicTimebase Timebase) const
+{
+	switch (Timebase)
+	{
+	case ECalibratedMusicTimebase::AudioRenderTime:
+		return CurrentSmoothedAudioRenderSongPos;
+	case ECalibratedMusicTimebase::ExperiencedTime:
+		return CurrentPlayerExperiencedSongPos;
+	case ECalibratedMusicTimebase::VideoRenderTime:
+		return CurrentVideoRenderSongPos;
+	case ECalibratedMusicTimebase::RawAudioRenderTime:
+		return CurrentRawAudioRenderSongPos;
+	}
+	check(false);
+	// minor robustness, return something even if there's an illegal value. 
+	return CurrentVideoRenderSongPos;
+}
+
+const FMidiSongPos& UMusicClockComponent::GetPreviousSongPosInternal(ECalibratedMusicTimebase Timebase) const
+{
+	switch (Timebase)
+	{
+	case ECalibratedMusicTimebase::AudioRenderTime:
+		return PrevAudioRenderSongPos;
+	case ECalibratedMusicTimebase::ExperiencedTime:
+		return PrevPlayerExperiencedSongPos;
+	case ECalibratedMusicTimebase::VideoRenderTime:
+		return PrevVideoRenderSongPos;
+	case ECalibratedMusicTimebase::RawAudioRenderTime:
+		return PrevRawAudioRenderSongPos;
+	}
+	check(false);
+	// minor robustness, return something even if there's an illegal value. 
+	return PrevVideoRenderSongPos;
+}
+
 FMusicalTime UMusicClockComponent::GetMusicalTime() const
 {
-	FMusicalTime Result;
+	FMusicalTime Result; 
 	if (CurrentVideoRenderSongPos.TimeSigDenominator < 1)
 	{
 		return Result;
@@ -890,3 +812,4 @@ void FMusicClockDriverBase::EnsureClockIsValidForGameFrame()
 		ClockComponent->LastUpdateFrame = GFrameCounter;
 	}
 }
+
diff --git a/Engine/Plugins/Runtime/Harmonix/Source/HarmonixMetasound/Public/HarmonixMetasound/Components/MusicClockComponent.h b/Engine/Plugins/Runtime/Harmonix/Source/HarmonixMetasound/Public/HarmonixMetasound/Components/MusicClockComponent.h
--- a/Engine/Plugins/Runtime/Harmonix/Source/HarmonixMetasound/Public/HarmonixMetasound/Components/MusicClockComponent.h
+++ b/Engine/Plugins/Runtime/Harmonix/Source/HarmonixMetasound/Public/HarmonixMetasound/Components/MusicClockComponent.h
@@ -213,10 +213,6 @@ public:
 	UFUNCTION(BlueprintPure, Category = "MusicClock")
 	float GetDeltaBeat(ECalibratedMusicTimebase Timebase = ECalibratedMusicTimebase::VideoRenderTime) const;
 
-	/** Returns true if Running or Paused and has been ticked at least Once.  If false, all values will be at their defaults */
-	UFUNCTION(BlueprintPure, Category = "MusicClock")
-	bool HasValidSongPos() const;
-
 	// Note: Not const as it might cause the clock to update from its source.
 	UFUNCTION(BlueprintPure, Category = "MusicClock")
 	const FMidiSongPos& GetSongPos(ECalibratedMusicTimebase Timebase = ECalibratedMusicTimebase::VideoRenderTime) const;
@@ -400,6 +396,7 @@ private:
 
 	FSongMaps DefaultMaps;
 
+	// TODO shouldn't these deltas be rolled into the MidiSongPos?
 	float RawAudioRenderDeltaBarF = 0.0f;
 	float RawAudioRenderDeltaBeatF = 0.0f;
 	float AudioRenderDeltaBarF = 0.0f;
@@ -424,11 +421,11 @@ private:
 	FMidiSongPos PrevPlayerExperiencedSongPos;
 	FMidiSongPos PrevVideoRenderSongPos;
 
-	uint64    LastUpdateFrame = 0;
+	uint64 LastUpdateFrame = 0;
 
 	void CreateClockDriver();
 	void BroadcastSongPosChanges();
-	void BroadcastSeekLoopDetections();
+	void BroadcastSeekLoopDetections() const;
 	void MakeDefaultSongMap();
 	bool ConnectToMetasound();
 	void ConnectToWallClock();
@@ -440,6 +437,8 @@ private:
 												// deleted and all of the call sites cleaned up... as only this next, 
 												// non-const "Ensure" function will be required...
 	void EnsureClockIsValidForGameFrameFromSubsystem();
+	const FMidiSongPos& GetCurrentSongPosInternal(ECalibratedMusicTimebase Timebase) const;
+	const FMidiSongPos& GetPreviousSongPosInternal(ECalibratedMusicTimebase Timebase) const;
 };
 
 struct FMusicClockDriverBase : public TSharedFromThis<FMusicClockDriverBase>
diff --git a/Engine/Plugins/Runtime/Harmonix/Source/HarmonixMidi/Private/HarmonixMidi/MidiSongPos.cpp b/Engine/Plugins/Runtime/Harmonix/Source/HarmonixMidi/Private/HarmonixMidi/MidiSongPos.cpp
--- a/Engine/Plugins/Runtime/Harmonix/Source/HarmonixMidi/Private/HarmonixMidi/MidiSongPos.cpp
+++ b/Engine/Plugins/Runtime/Harmonix/Source/HarmonixMidi/Private/HarmonixMidi/MidiSongPos.cpp
@@ -82,6 +82,8 @@ FMidiSongPos FMidiSongPos::Lerp(const FMidiSongPos& A, const FMidiSongPos& B, fl
 	{
 		Result.BeatType = B.BeatType;
 	}
+	// Could also imagine using &&, but I was worried about propagating not set forever -jwf 
+	Result.IsSet = A.IsSet || B.IsSet;
 	return Result;
 }
 
@@ -109,6 +111,7 @@ void FMidiSongPos::SetByTime(float InElapsedMs, float InBpm, int32 InTimeSigNume
 	BeatsIncludingCountIn   = TotalBeats;
 	Timestamp.Bar = FMath::FloorToInt32(BarsIncludingCountIn) + StartBar;
 	Timestamp.Beat = FMath::Fractional(BarsIncludingCountIn) * InTimeSigNumerator + 1.0f;
+	IsSet = true;
 }
 
 void FMidiSongPos::SetByTime(float InMs, const ISongMapEvaluator& Map)
@@ -134,6 +137,7 @@ void FMidiSongPos::SetByTimeAndTick(float InMs, float InTick, const ISongMapEval
 	TimeSigDenominator      = TimeSig ? TimeSig->Denominator : 4;
 	Tempo                   = Map.GetTempoAtTick(int32(InTick));
 	BarsIncludingCountIn    = Map.GetBarIncludingCountInAtTick(InTick);
+	IsSet = true;
 
 	const FSongSection* SongSection = Map.GetSectionAtTick(int32(InTick));
 	CurrentSongSection = SongSection ? FSongSection(SongSection->Name, SongSection->StartTick, SongSection->LengthTicks) : FSongSection();
diff --git a/Engine/Plugins/Runtime/Harmonix/Source/HarmonixMidi/Public/HarmonixMidi/MidiSongPos.h b/Engine/Plugins/Runtime/Harmonix/Source/HarmonixMidi/Public/HarmonixMidi/MidiSongPos.h
--- a/Engine/Plugins/Runtime/Harmonix/Source/HarmonixMidi/Public/HarmonixMidi/MidiSongPos.h
+++ b/Engine/Plugins/Runtime/Harmonix/Source/HarmonixMidi/Public/HarmonixMidi/MidiSongPos.h
@@ -73,11 +73,13 @@ public:
 		BeatsIncludingCountIn = 0.0f;
 		BeatType = EMusicalBeatType::Normal;
 		CurrentSongSection = FSongSection();
+		IsSet = false;
 	}
 
 	float GetBeatsPerSecond() const { return Tempo / 60.0f; }
 
-	bool IsZero() const	{ return FMath::IsNearlyZero(BeatsIncludingCountIn, UE_KINDA_SMALL_NUMBER); }
+	// Has this been set by a system yet, or is it in the default state? 
+	bool IsValid() const { return IsSet; }
 
 	// operators
 	bool operator<(const FMidiSongPos& rhs) const;
@@ -89,5 +91,9 @@ public:
 	static FMidiSongPos Lerp(const FMidiSongPos& A, const FMidiSongPos& B, float Alpha);
 
 	void SetByTimeAndTick(float Ms, float Tick, const ISongMapEvaluator& Map);
+	
+private:
+	// Has anyone explicitly set this, or is this just the defaults.
+	bool IsSet = false;
 };
 
