Make MidiPlayerNode header so we can expose metasound class name and version number for external use

#rnx

diff --git a/Engine/Plugins/Runtime/Harmonix/Source/HarmonixMetasound/Private/HarmonixMetasound/Nodes/MidiPlayerNode.cpp b/Engine/Plugins/Runtime/Harmonix/Source/HarmonixMetasound/Private/HarmonixMetasound/Nodes/MidiPlayerNode.cpp
--- a/Engine/Plugins/Runtime/Harmonix/Source/HarmonixMetasound/Private/HarmonixMetasound/Nodes/MidiPlayerNode.cpp
+++ b/Engine/Plugins/Runtime/Harmonix/Source/HarmonixMetasound/Private/HarmonixMetasound/Nodes/MidiPlayerNode.cpp
@@ -20,14 +20,37 @@
 
 #define LOCTEXT_NAMESPACE "HarmonixMetaSound"
 
-namespace HarmonixMetasound
+namespace HarmonixMetasound::Nodes::MidiPlayerNode
 {
 	using namespace Metasound;
 
-	namespace MidiPlayerNodePinNames
+	const FNodeClassName& GetClassName()
 	{
-		METASOUND_PARAM(KillVoicesOnSeek, "Kill Voices On Seek", "If true, a \"Kill All Voices\" MIDI message will be sent when seeking. Otherwise an \"All Notes Off\" will be sent which allows to ADSR release phases.")
-		METASOUND_PARAM(KillVoicesOnMidiChange, "Kill Voices On MIDI File Change", "If true, a \"Kill All Voices\" MIDI message will be sent when the MIDI file asset is changed. Otherwise an \"All Notes Off\" will be sent which allows to ADSR release phases.")
+		static const FNodeClassName ClassName { HarmonixNodeNamespace, TEXT("MIDIPlayer"), TEXT("") };
+		return ClassName;
+	}
+
+	int32 GetCurrentMajorVersion()
+	{
+		return 0;
+	}
+
+	namespace Inputs
+	{
+		DEFINE_METASOUND_PARAM_ALIAS(MidiFileAsset, CommonPinNames::Inputs::MidiFileAsset);
+		DEFINE_METASOUND_PARAM_ALIAS(Transport, CommonPinNames::Inputs::Transport);
+		DEFINE_METASOUND_PARAM_ALIAS(MidiClock, CommonPinNames::Inputs::MidiClock);
+		DEFINE_METASOUND_PARAM_ALIAS(Loop, CommonPinNames::Inputs::Loop);
+		DEFINE_METASOUND_PARAM_ALIAS(Speed, CommonPinNames::Inputs::Speed);
+		DEFINE_METASOUND_PARAM_ALIAS(PrerollBars, CommonPinNames::Inputs::PrerollBars);
+		DEFINE_INPUT_METASOUND_PARAM(KillVoicesOnSeek, "Kill Voices On Seek", "If true, a \"Kill All Voices\" MIDI message will be sent when seeking. Otherwise an \"All Notes Off\" will be sent which allows to ADSR release phases.")
+		DEFINE_INPUT_METASOUND_PARAM(KillVoicesOnMidiChange, "Kill Voices On MIDI File Change", "If true, a \"Kill All Voices\" MIDI message will be sent when the MIDI file asset is changed. Otherwise an \"All Notes Off\" will be sent which allows to ADSR release phases.")
+	}
+
+	namespace Outputs
+	{
+		DEFINE_METASOUND_PARAM_ALIAS(MidiStream, CommonPinNames::Outputs::MidiStream);
+		DEFINE_METASOUND_PARAM_ALIAS(MidiClock, CommonPinNames::Outputs::MidiClock);
 	}
 
 	class FMidiPlayerOperator : public TExecutableOperator<FMidiPlayerOperator>, public FMidiCursor::FReceiver, public FMidiVoiceGeneratorBase, public FMusicTransportControllable
@@ -164,8 +187,8 @@ namespace HarmonixMetasound
 		auto InitNodeInfo = []() -> FNodeClassMetadata
 		{
 			FNodeClassMetadata Info;
-			Info.ClassName        = { HarmonixNodeNamespace, TEXT("MIDIPlayer"), TEXT("") };
-			Info.MajorVersion     = 0;
+			Info.ClassName        = GetClassName();
+			Info.MajorVersion     = GetCurrentMajorVersion();
 			Info.MinorVersion     = 1;
 			Info.DisplayName      = METASOUND_LOCTEXT("MIDIPlayerNode_DisplayName", "MIDI Player");
 			Info.Description      = METASOUND_LOCTEXT("MIDIPlayerNode_Description", "Plays a standard MIDI file.");
@@ -184,8 +207,6 @@ namespace HarmonixMetasound
 
 	const FVertexInterface& FMidiPlayerOperator::GetVertexInterface()
 	{
-		using namespace CommonPinNames;
-
 		static const FVertexInterface Interface(
 			FInputVertexInterface(
 				TInputDataVertex<FMidiAsset>(METASOUND_GET_PARAM_NAME_AND_METADATA(Inputs::MidiFileAsset)),
@@ -194,8 +215,8 @@ namespace HarmonixMetasound
 				TInputDataVertex<bool>(METASOUND_GET_PARAM_NAME_AND_METADATA(Inputs::Loop), false),
 				TInputDataVertex<float>(METASOUND_GET_PARAM_NAME_AND_METADATA(Inputs::Speed), 1.0f),
 				TInputConstructorVertex<int32>(METASOUND_GET_PARAM_NAME_AND_METADATA(Inputs::PrerollBars), 8),
-				TInputConstructorVertex<bool>(METASOUND_GET_PARAM_NAME_AND_METADATA(MidiPlayerNodePinNames::KillVoicesOnSeek), false),
-				TInputConstructorVertex<bool>(METASOUND_GET_PARAM_NAME_AND_METADATA(MidiPlayerNodePinNames::KillVoicesOnMidiChange), false)
+				TInputConstructorVertex<bool>(METASOUND_GET_PARAM_NAME_AND_METADATA(Inputs::KillVoicesOnSeek), false),
+				TInputConstructorVertex<bool>(METASOUND_GET_PARAM_NAME_AND_METADATA(Inputs::KillVoicesOnMidiChange), false)
 			),
 			FOutputVertexInterface(
 				TOutputDataVertex<FMidiStream>(METASOUND_GET_PARAM_NAME_AND_METADATA(Outputs::MidiStream)),
@@ -208,8 +229,6 @@ namespace HarmonixMetasound
 
 	TUniquePtr<IOperator> FMidiPlayerOperator::CreateOperator(const FBuildOperatorParams& InParams, FBuildResults& OutResults)
 	{
-		using namespace CommonPinNames;
-
 		const FMidiPlayerNode& MidiPlayerNode = static_cast<const FMidiPlayerNode&>(InParams.Node);
 
 		const FInputVertexInterfaceData& InputData = InParams.InputData;
@@ -221,8 +240,8 @@ namespace HarmonixMetasound
 		FBoolReadRef InLoop = InputData.GetOrConstructDataReadReference<bool>(METASOUND_GET_PARAM_NAME(Inputs::Loop), false);
 		FFloatReadRef InSpeed = InputData.GetOrConstructDataReadReference<float>(METASOUND_GET_PARAM_NAME(Inputs::Speed), 1.0f);
 		int32 InPrerollBars = InputData.GetOrCreateDefaultValue<int32>(METASOUND_GET_PARAM_NAME(Inputs::PrerollBars), Settings);
-		bool bInKillVoicesOnSeek = InputData.GetOrCreateDefaultValue<bool>(METASOUND_GET_PARAM_NAME(MidiPlayerNodePinNames::KillVoicesOnSeek), Settings);
-		bool bInKillVoicesOnMidiChange = InputData.GetOrCreateDefaultValue<bool>(METASOUND_GET_PARAM_NAME(MidiPlayerNodePinNames::KillVoicesOnMidiChange), Settings);
+		bool bInKillVoicesOnSeek = InputData.GetOrCreateDefaultValue<bool>(METASOUND_GET_PARAM_NAME(Inputs::KillVoicesOnSeek), Settings);
+		bool bInKillVoicesOnMidiChange = InputData.GetOrCreateDefaultValue<bool>(METASOUND_GET_PARAM_NAME(Inputs::KillVoicesOnMidiChange), Settings);
 		if (InputData.IsVertexBound(METASOUND_GET_PARAM_NAME(Inputs::MidiClock)))
 		{
 			FMidiClockReadRef InMidiClock = InputData.GetOrConstructDataReadReference<FMidiClock>(METASOUND_GET_PARAM_NAME(Inputs::MidiClock), Settings);
@@ -289,7 +308,6 @@ namespace HarmonixMetasound
 	void FExternallyClockedMidiPlayerOperator::BindInputs(FInputVertexInterfaceData& InVertexData)
 	{
 		FMidiPlayerOperator::BindInputs(InVertexData);
-		using namespace CommonPinNames;
 		InVertexData.BindReadVertex(METASOUND_GET_PARAM_NAME(Inputs::MidiClock), MidiClockIn);
 		MidiClockOut->SetDrivingClock(MidiClockIn->AsShared().ToSharedPtr());
 	}
@@ -484,21 +502,19 @@ namespace HarmonixMetasound
 
 	void FMidiPlayerOperator::BindInputs(FInputVertexInterfaceData& InVertexData)
 	{
-		using namespace CommonPinNames;
 		InVertexData.BindReadVertex(METASOUND_GET_PARAM_NAME(Inputs::MidiFileAsset), MidiAssetInPin);
 		InVertexData.BindReadVertex(METASOUND_GET_PARAM_NAME(Inputs::Transport), TransportInPin);
 		InVertexData.BindReadVertex(METASOUND_GET_PARAM_NAME(Inputs::Loop), LoopInPin);
 		InVertexData.BindReadVertex(METASOUND_GET_PARAM_NAME(Inputs::Speed), SpeedMultInPin);
 		InVertexData.SetValue(METASOUND_GET_PARAM_NAME(Inputs::PrerollBars), PrerollBars);
-		InVertexData.SetValue(METASOUND_GET_PARAM_NAME(MidiPlayerNodePinNames::KillVoicesOnSeek), bKillVoicesOnSeek);
-		InVertexData.SetValue(METASOUND_GET_PARAM_NAME(MidiPlayerNodePinNames::KillVoicesOnMidiChange), bKillVoicesOnMidiChange);
+		InVertexData.SetValue(METASOUND_GET_PARAM_NAME(Inputs::KillVoicesOnSeek), bKillVoicesOnSeek);
+		InVertexData.SetValue(METASOUND_GET_PARAM_NAME(Inputs::KillVoicesOnMidiChange), bKillVoicesOnMidiChange);
 
 		NeedsTransportInit = true;
 	}
 
 	void FMidiPlayerOperator::BindOutputs(FOutputVertexInterfaceData& InVertexData)
 	{
-		using namespace CommonPinNames;
 		InVertexData.BindReadVertex(METASOUND_GET_PARAM_NAME(Outputs::MidiStream), MidiOutPin);
 		InVertexData.BindReadVertex(METASOUND_GET_PARAM_NAME(Outputs::MidiClock), MidiClockOut);
 
diff --git a/Engine/Plugins/Runtime/Harmonix/Source/HarmonixMetasound/Public/HarmonixMetasound/Nodes/MidiPlayerNode.h b/Engine/Plugins/Runtime/Harmonix/Source/HarmonixMetasound/Public/HarmonixMetasound/Nodes/MidiPlayerNode.h
new file mode 100644
--- /dev/null
+++ b/Engine/Plugins/Runtime/Harmonix/Source/HarmonixMetasound/Public/HarmonixMetasound/Nodes/MidiPlayerNode.h
@@ -0,0 +1,35 @@
+﻿// Copyright Epic Games, Inc. All Rights Reserved.
+
+#pragma once
+
+#include "HAL/Platform.h"
+
+#include "MetasoundParamHelper.h"
+#include "HarmonixMetasound/Common.h"
+#include "MetasoundNodeInterface.h"
+
+namespace HarmonixMetasound::Nodes::MidiPlayerNode
+{
+	using namespace Metasound;
+
+	const HARMONIXMETASOUND_API FNodeClassName& GetClassName();
+	HARMONIXMETASOUND_API int32 GetCurrentMajorVersion();
+
+	namespace Inputs
+	{
+		DECLARE_METASOUND_PARAM_ALIAS(MidiFileAsset);
+		DECLARE_METASOUND_PARAM_ALIAS(Transport);
+		DECLARE_METASOUND_PARAM_ALIAS(MidiClock);
+		DECLARE_METASOUND_PARAM_ALIAS(Loop);
+		DECLARE_METASOUND_PARAM_ALIAS(Speed);
+		DECLARE_METASOUND_PARAM_ALIAS(PrerollBars);
+		DECLARE_METASOUND_PARAM_EXTERN(KillVoicesOnSeek);
+		DECLARE_METASOUND_PARAM_EXTERN(KillVoicesOnMidiChange);
+	}
+
+	namespace Outputs
+	{
+		DECLARE_METASOUND_PARAM_ALIAS(MidiStream);
+		DECLARE_METASOUND_PARAM_ALIAS(MidiClock);
+	}
+}
\ No newline at end of file
