Fix metronome not using correct music map when auditioning in music mode in sequencer.

Rewrite the MusicEnvironmentMetronome to use a "MidiPlayerNode" instead of a "MetronomeNode" in the metasound so that we can pass in the "MusicMap" being used by the sequencer.

Update the music map on the time controller when creating it

#rnx

diff --git a/Engine/Plugins/Runtime/Harmonix/Source/HarmonixMetasound/Private/HarmonixMetasound/MusicEnvironmentSupport/HarmonixMusicEnvironmentMetronome.cpp b/Engine/Plugins/Runtime/Harmonix/Source/HarmonixMetasound/Private/HarmonixMetasound/MusicEnvironmentSupport/HarmonixMusicEnvironmentMetronome.cpp
--- a/Engine/Plugins/Runtime/Harmonix/Source/HarmonixMetasound/Private/HarmonixMetasound/MusicEnvironmentSupport/HarmonixMusicEnvironmentMetronome.cpp
+++ b/Engine/Plugins/Runtime/Harmonix/Source/HarmonixMetasound/Private/HarmonixMetasound/MusicEnvironmentSupport/HarmonixMusicEnvironmentMetronome.cpp
@@ -7,10 +7,13 @@
 #include "Components/AudioComponent.h"
 #include "Engine/Engine.h"
 #include "HarmonixMetasound/Components/MusicClockComponent.h"
+#include "HarmonixMetasound/DataTypes/MidiAsset.h"
 #include "HarmonixMetasound/Nodes/MetronomeNode.h"
+#include "HarmonixMetasound/Nodes/MidiPlayerNode.h"
 #include "HarmonixMetasound/Nodes/MidiClockSubdivisionTriggerNode.h"
 #include "HarmonixMetasound/Nodes/MusicSeekTargetBuilder.h"
 #include "HarmonixMetasound/Nodes/TriggerToTransportNode.h"
+#include "HarmonixMidi/MidiConstants.h"
 
 namespace HarmonixMovieSceneMetronome
 {
@@ -44,8 +47,12 @@ void UHarmonixMusicEnvironmentMetronome::Start(double FromSeconds)
 		return;
 	}
 
+	RebuildMidiFile();
+	
+	FAudioParameter MidiFileParam(TEXT("MIDI File"), MidiFile);
 	FAudioParameter SeekTargetMs(TEXT("SeekTarget"), static_cast<float>(FromSeconds * 1000.0f));
 	AudioComponent->Activate();
+	AudioComponent->SetParameter(MoveTemp(MidiFileParam));
 	AudioComponent->SetParameter(MoveTemp(SeekTargetMs));
 	AudioComponent->SetTriggerParameter(TEXT("PlayMetronome"));
 	MusicClockComponent->Start();
@@ -157,26 +164,69 @@ void UHarmonixMusicEnvironmentMetronome::BeginDestroy()
 
 void UHarmonixMusicEnvironmentMetronome::OnMusicMapSet()
 {
+	RebuildMidiFile();
+
+	if (AudioComponent)
+	{
+		FAudioParameter MidiFileParam(TEXT("MIDI File"), MidiFile);
+		AudioComponent->SetParameter(MoveTemp(MidiFileParam));
+	}
 }
 
 bool UHarmonixMusicEnvironmentMetronome::OnSetTempo(float Bpm)
 {
+	RebuildMidiFile();
+	
 	if (!AudioComponent)
 	{
 		return false;
 	}
+
+	FAudioParameter MidiFileParam(TEXT("MIDI File"), MidiFile);
+	AudioComponent->SetParameter(MoveTemp(MidiFileParam));
+	
 	FAudioParameter TempoParam(TEXT("Tempo"), Bpm);
 	AudioComponent->SetParameter(MoveTemp(TempoParam));
 	CurrentTempo = Bpm;
 	return true;
 }
 
+bool UHarmonixMusicEnvironmentMetronome::RebuildMidiFile()
+{
+	MidiFile = NewObject<UMidiFile>(this, "MIDI File Music Map", RF_Transient);
+	if (MusicMap)
+	{
+		for (const FFrameBasedTimeSignaturePoint& TimeSigPoint : MusicMap->BarMap)
+		{
+			const int32 Num = TimeSigPoint.TimeSignature.Numerator;
+			const int32 Den = TimeSigPoint.TimeSignature.Denominator;
+			MidiFile->GetSongMaps()->AddTimeSignatureAtBarIncludingCountIn(TimeSigPoint.OnBar, Num, Den);
+		}
+
+		for (const FFrameBasedTempoPoint& TempoPoint : MusicMap->TempoMap)
+		{
+			int32 MidiTempo = TempoPoint.MicrosecondsPerQuarterNote;
+			MidiFile->GetSongMaps()->AddTempoInfoPoint(MidiTempo, TempoPoint.OnTick);
+		}
+	}
+	else
+	{
+		MidiFile->GetSongMaps()->AddTempoInfoPoint(Harmonix::Midi::Constants::BPMToMidiTempo(CurrentTempo), 0);
+		MidiFile->GetSongMaps()->AddTimeSignatureAtBarIncludingCountIn(0, 4, 4);
+	}
+	
+	MidiFile->BuildConductorTrack();
+	return true;
+}
+
+
 bool UHarmonixMusicEnvironmentMetronome::BuildAndStartMetasound(UWorld* InWorld)
 {
 	if (!BuildMetasound())
 	{
 		return false;
 	}
+	check(SourceBuilder);
 	AudioComponent = NewObject<UAudioComponent>(this, "MetronomeAudio", RF_Transient);
 	AudioComponent->SetUISound(true);
 	SourceBuilder->Audition(this, AudioComponent, {});
@@ -189,7 +239,7 @@ bool UHarmonixMusicEnvironmentMetronome::BuildAndStartMetasound(UWorld* InWorld)
 
 	if (MusicMap)
 	{
-		OnSetTempo(MusicMap->GetInitialTempo());
+		CurrentTempo = MusicMap->GetInitialTempo();
 	}
 	
 	bMetasoundIsPlaying = true;
@@ -215,7 +265,10 @@ bool UHarmonixMusicEnvironmentMetronome::BuildMetasound()
 
 		// Make the nodes and set their un-wired inputs to appropriate values...
 		if (!AddGraphInput(TEXT("BeepGain"), 0.25f).IsSet()) { return false; }
-		if (!AddGraphInput(TEXT("Tempo"), 120.0f).IsSet()) { return false; }
+		
+		// the midi file is used for the midi player node if there's a music map
+		if (!AddGraphInput(TEXT("MIDI File"), nullptr).IsSet()) { return false; }
+		
 		if (!AddGraphInput(TEXT("SeekTarget"), 0.0f).IsSet()) { return false; }
 		if (!AddGraphInputTrigger(TEXT("PlayMetronome")).IsSet()) { return false; }
 		if (!AddGraphInputTrigger(TEXT("StopMetronome")).IsSet()) { return false; }
@@ -231,7 +284,10 @@ bool UHarmonixMusicEnvironmentMetronome::BuildMetasound()
 		FMetaSoundNodeHandle SeekTargetConverter = SourceBuilder->AddNodeByClassName(TimeMsToSeekTarget::GetClassName(), Result, TimeMsToSeekTarget::GetCurrentMajorVersion());
 	
 		FMetaSoundNodeHandle TransportNode = SourceBuilder->AddNodeByClassName(TriggerToTransportNode::GetClassName(), Result, TriggerToTransportNode::GetCurrentMajorVersion());
-		FMetaSoundNodeHandle MetronomeNode = SourceBuilder->AddNodeByClassName(MetronomeNode::GetClassName(), Result, MetronomeNode::GetCurrentMajorVersion());
+
+		FMetaSoundNodeHandle MidiPlayerNode = SourceBuilder->AddNodeByClassName(MidiPlayerNode::GetClassName(), Result, MidiPlayerNode::GetCurrentMajorVersion());
+		FMetaSoundNodeHandle ClockSourceNode = MidiPlayerNode;
+		
 		FMetaSoundNodeHandle BeatSubdivision = SourceBuilder->AddNodeByClassName(MidiClockSubdivisionTriggerNode::GetClassName(), Result, MidiClockSubdivisionTriggerNode::GetCurrentMajorVersion());
 		
 		FMetaSoundNodeHandle BarSubdivision = SourceBuilder->AddNodeByClassName(MidiClockSubdivisionTriggerNode::GetClassName(), Result, MidiClockSubdivisionTriggerNode::GetCurrentMajorVersion());
@@ -278,7 +334,9 @@ bool UHarmonixMusicEnvironmentMetronome::BuildMetasound()
 		if (!ConnectGraphInputToNodeInput(TEXT("ContinueMetronome"), TransportNode, TEXT("Continue"))) { return false; }
 		if (!ConnectGraphInputToNodeInput(TEXT("StopMetronome"), TransportNode, TEXT("Stop"))) { return false; }
 		if (!ConnectGraphInputToNodeInput(TEXT("BeepGain"), GainMultiplier, TEXT("AdditionalOperands"))) { return false; }
-		if (!ConnectGraphInputToNodeInput(TEXT("Tempo"), MetronomeNode, TEXT("Tempo"))) { return false; }
+
+		if (!ConnectGraphInputToNodeInput(TEXT("MIDI File"), MidiPlayerNode,TEXT("MIDI File"))) { return false; }
+		
 		if (!ConnectGraphInputToNodeInput(TEXT("SeekTarget"), SeekTargetConverter, TEXT("Time (Ms)"))) { return false; }
 
 		// Interconnects...
@@ -286,12 +344,14 @@ bool UHarmonixMusicEnvironmentMetronome::BuildMetasound()
 		if (!ConnectNodes(SeekTriggerDelay, TEXT("Out"), TransportNode, TEXT("Trigger Seek"))) { return false; }
 		if (!ConnectNodes(SeekTargetConverter, TEXT("Seek Target"), TransportNode, TEXT("Seek Target"))) { return false; }
 		if (!ConnectNodes(PlayTriggerDelay, TEXT("Out"), TransportNode, TEXT("Play"))) { return false; }
-		if (!ConnectNodes(TransportNode, TEXT("Transport"), MetronomeNode, TEXT("Transport"))) { return false; }
-		if (!ConnectNodes(MetronomeNode, TEXT("MIDI Clock"), BeatSubdivision, TEXT("MIDI Clock"))) { return false; }
+
+		if (!ConnectNodes(TransportNode, TEXT("Transport"), ClockSourceNode, TEXT("Transport"))) { return false; }
+		if (!ConnectNodes(ClockSourceNode, TEXT("MIDI Clock"), BarSubdivision, TEXT("MIDI Clock"))) { return false; }
+		if (!ConnectNodes(ClockSourceNode, TEXT("MIDI Clock"), OffsetBarSubdivision, TEXT("MIDI Clock"))) { return false; }
+		if (!ConnectNodes(ClockSourceNode, TEXT("MIDI Clock"), BeatSubdivision, TEXT("MIDI Clock"))) { return false; }
+		
 		if (!ConnectNodes(BeatSubdivision, TEXT("Trigger Out"), Envelope, TEXT("Trigger"))) { return false; }
 		if (!ConnectNodes(Envelope, TEXT("Out Envelope"), EnvelopeMultiplier, TEXT("PrimaryOperand"))) { return false; }
-		if (!ConnectNodes(MetronomeNode, TEXT("MIDI Clock"), BarSubdivision, TEXT("MIDI Clock"))) { return false; }
-		if (!ConnectNodes(MetronomeNode, TEXT("MIDI Clock"), OffsetBarSubdivision, TEXT("MIDI Clock"))) { return false; }
 		if (!ConnectNodes(BarSubdivision, TEXT("Trigger Out"), TriggerCounter, TEXT("Reset"))) { return false; }
 		if (!ConnectNodes(OffsetBarSubdivision, TEXT("Trigger Out"), TriggerCounter, TEXT("In"))) { return false; }
 		if (!ConnectNodes(TriggerCounter, TEXT("On Trigger"), TriggerAny, TEXT("In 0"))) { return false; }
@@ -302,8 +362,9 @@ bool UHarmonixMusicEnvironmentMetronome::BuildMetasound()
 		if (!ConnectNodes(SineOsc, TEXT("Audio"), EnvelopeMultiplier, TEXT("AdditionalOperands"))) { return false; }
 		if (!ConnectNodes(EnvelopeMultiplier, TEXT("Out"), GainMultiplier, TEXT("PrimaryOperand"))) { return false; }
 
-		// Outputs... 
-		if (!ConnectNodeOutputToGraphOutput(MetronomeNode, TEXT("MIDI Clock"), TEXT("MIDI Clock"))) { return false; }
+		// Outputs...
+		if (!ConnectNodeOutputToGraphOutput(ClockSourceNode, TEXT("MIDI Clock"), TEXT("MIDI Clock"))) { return false; }
+		
 		Out = SourceBuilder->FindNodeOutputByName(GainMultiplier, TEXT("Out"), Result);
 		SourceBuilder->ConnectNodes(Out, AudioOutNodeInputs[0], Result);
 
@@ -312,8 +373,19 @@ bool UHarmonixMusicEnvironmentMetronome::BuildMetasound()
 	return false;
 }
 
+FMetaSoundBuilderNodeOutputHandle UHarmonixMusicEnvironmentMetronome::AddGraphInput(const FName& InputName, UMidiFile* InMidiFile) const
+{
+	check(SourceBuilder);
+	EMetaSoundBuilderResult Result;
+	
+	FMetasoundFrontendLiteral DefaultValue;
+	DefaultValue.Set(InMidiFile);
+	return SourceBuilder->AddGraphInputNode(InputName, Metasound::GetMetasoundDataTypeName<HarmonixMetasound::FMidiAsset>(), DefaultValue, Result);
+}
+
 FMetaSoundBuilderNodeOutputHandle UHarmonixMusicEnvironmentMetronome::AddGraphInput(const FName& InputName, const float Value) const
 {
+	check(SourceBuilder);
 	EMetaSoundBuilderResult Result;
 
 	FMetasoundFrontendLiteral DefaultValue;
@@ -323,12 +395,14 @@ FMetaSoundBuilderNodeOutputHandle UHarmonixMusicEnvironmentMetronome::AddGraphIn
 
 FMetaSoundBuilderNodeOutputHandle UHarmonixMusicEnvironmentMetronome::AddGraphInputTrigger(const FName& InputName) const
 {
+	check(SourceBuilder);
 	EMetaSoundBuilderResult Result;
 	return SourceBuilder->AddGraphInputNode(InputName, Metasound::GetMetasoundDataTypeName<Metasound::FTrigger>(), FMetasoundFrontendLiteral(), Result);
 }
 
 bool UHarmonixMusicEnvironmentMetronome::SetNodeInputDefault(const FMetaSoundNodeHandle& Node, const FName& InputName, const int32 Value) const
 {
+	check(SourceBuilder);
 	EMetaSoundBuilderResult Result;
 
 	FMetaSoundBuilderNodeInputHandle In = SourceBuilder->FindNodeInputByName(Node, InputName, Result);
@@ -341,6 +415,7 @@ bool UHarmonixMusicEnvironmentMetronome::SetNodeInputDefault(const FMetaSoundNod
 
 bool UHarmonixMusicEnvironmentMetronome::SetNodeInputDefault(const FMetaSoundNodeHandle& Node, const FName& InputName, const float Value) const
 {
+	check(SourceBuilder);	
 	EMetaSoundBuilderResult Result;
 
 	FMetaSoundBuilderNodeInputHandle In = SourceBuilder->FindNodeInputByName(Node, InputName, Result);
@@ -354,6 +429,7 @@ bool UHarmonixMusicEnvironmentMetronome::SetNodeInputDefault(const FMetaSoundNod
 bool UHarmonixMusicEnvironmentMetronome::SetNodeInputDefault(const FMetaSoundNodeHandle& Node, const FName& InputName,
 	const TArray<float>& Values) const
 {
+	check(SourceBuilder);
 	EMetaSoundBuilderResult Result;
 
 	FMetaSoundBuilderNodeInputHandle In = SourceBuilder->FindNodeInputByName(Node, InputName, Result);
@@ -367,6 +443,7 @@ bool UHarmonixMusicEnvironmentMetronome::SetNodeInputDefault(const FMetaSoundNod
 bool UHarmonixMusicEnvironmentMetronome::ConnectGraphInputToNodeInput(const FName& GraphInputName, const FMetaSoundNodeHandle& Node,
 	const FName& NodeInputName) const
 {
+	check(SourceBuilder);
 	EMetaSoundBuilderResult Result;
 
 	FMetaSoundBuilderNodeInputHandle NodeIn = SourceBuilder->FindNodeInputByName(Node, NodeInputName, Result);
@@ -378,6 +455,7 @@ bool UHarmonixMusicEnvironmentMetronome::ConnectGraphInputToNodeInput(const FNam
 bool UHarmonixMusicEnvironmentMetronome::ConnectNodeOutputToGraphOutput(const FMetaSoundNodeHandle& Node, const FName& NodeOutputName,
 	const FName& GraphOutputName) const
 {
+	check(SourceBuilder);
 	EMetaSoundBuilderResult Result;
 
 	FMetaSoundBuilderNodeOutputHandle NodeOut = SourceBuilder->FindNodeOutputByName(Node, NodeOutputName, Result);
@@ -389,6 +467,7 @@ bool UHarmonixMusicEnvironmentMetronome::ConnectNodeOutputToGraphOutput(const FM
 bool UHarmonixMusicEnvironmentMetronome::ConnectNodes(const FMetaSoundNodeHandle& SourceNode, const FName& OutName,
                                                       const FMetaSoundNodeHandle& DestinationNode, const FName& InName) const
 {
+	check(SourceBuilder);
 	EMetaSoundBuilderResult Result;
 
 	FMetaSoundBuilderNodeOutputHandle Out = SourceBuilder->FindNodeOutputByName(SourceNode, OutName, Result);
diff --git a/Engine/Plugins/Runtime/Harmonix/Source/HarmonixMetasound/Private/HarmonixMetasound/MusicEnvironmentSupport/HarmonixMusicEnvironmentMetronome.h b/Engine/Plugins/Runtime/Harmonix/Source/HarmonixMetasound/Private/HarmonixMetasound/MusicEnvironmentSupport/HarmonixMusicEnvironmentMetronome.h
--- a/Engine/Plugins/Runtime/Harmonix/Source/HarmonixMetasound/Private/HarmonixMetasound/MusicEnvironmentSupport/HarmonixMusicEnvironmentMetronome.h
+++ b/Engine/Plugins/Runtime/Harmonix/Source/HarmonixMetasound/Private/HarmonixMetasound/MusicEnvironmentSupport/HarmonixMusicEnvironmentMetronome.h
@@ -12,6 +12,7 @@
 class UAudioComponent;
 class UMetaSoundSourceBuilder;
 class UMusicClockComponent;
+class UMidiFile;
 class UWorld;
 class AActor;
 
@@ -58,6 +59,9 @@ private:
 	UPROPERTY(Transient)
 	TObjectPtr<UMusicClockComponent> MusicClockComponent;
 
+	UPROPERTY(Transient)
+	TObjectPtr<UMidiFile> MidiFile;
+
 	UPROPERTY(Transient)
 	TObjectPtr<UMetaSoundSourceBuilder> SourceBuilder;
 	UPROPERTY(Transient)
@@ -71,7 +75,9 @@ private:
 	float CurrentTempo = 120.0f;
 	
 	bool BuildAndStartMetasound(UWorld* InWorld);
+	bool RebuildMidiFile();
 	bool BuildMetasound();
+	FMetaSoundBuilderNodeOutputHandle AddGraphInput(const FName& InputName, UMidiFile* MidiFile) const;
 	FMetaSoundBuilderNodeOutputHandle AddGraphInput(const FName& InputName, const float Value) const;
 	FMetaSoundBuilderNodeOutputHandle AddGraphInputTrigger(const FName& InputName) const;
 	bool SetNodeInputDefault(const FMetaSoundNodeHandle& Node, const FName& InputName, const int32 Value) const;
