Separate stretcher and pitch shifter factories from the Fusion Voice Pool

[RN] Add pin to FusionSampler node to enable/disable Pitch Shifters

diff --git a/Engine/Plugins/Runtime/Harmonix/Source/HarmonixDsp/Private/HarmonixDsp/FusionSampler/FusionSampler.cpp b/Engine/Plugins/Runtime/Harmonix/Source/HarmonixDsp/Private/HarmonixDsp/FusionSampler/FusionSampler.cpp
--- a/Engine/Plugins/Runtime/Harmonix/Source/HarmonixDsp/Private/HarmonixDsp/FusionSampler/FusionSampler.cpp
+++ b/Engine/Plugins/Runtime/Harmonix/Source/HarmonixDsp/Private/HarmonixDsp/FusionSampler/FusionSampler.cpp
@@ -28,7 +28,7 @@ static FInterpEaseIn gFilterQInterp(0.1f, 10.0f, 4.03f);
 static FInterpEaseIn gLfoFreqInterp(0.1f, 30.0f, 2.74f);
 static FInterpEaseIn gDelayTimeInterp(0.0001f, 4.0f, 2.0f);
 
-FFusionSampler::FFusionSampler()
+FFusionSampler::FFusionSampler(float InSampleRateHz, bool InUsePitchShifters)
 	: RampCallRateHz(1)
 	, Speed(1.0f)
 	, MidiChannelVolume(Defaults::kVolumeDb)
@@ -47,6 +47,7 @@ FFusionSampler::FFusionSampler()
 	, MinPitchBendCents(Defaults::kMinPitchBend)
 	, MaxPitchBendCents(Defaults::kMaxPitchBend)
 	, VoicePool(nullptr)
+	, bUsePitchShifters(InUsePitchShifters)
 	, CurrentTempoBPM(120.0f)
 	, RawPitchMultiplier(1.0f)
 	, Transposition(0)
@@ -71,9 +72,14 @@ FFusionSampler::FFusionSampler()
 	ResetNoteActions(true);
 	ResetNoteStatus();
 
-	SetSampleRate(0.0f);
+	SetSampleRate(InSampleRateHz);
 	ResetInstrumentState();
 	ResetPatchRelatedState();
+
+	if (bUsePitchShifters)
+	{
+		IStretcherAndPitchShifterFactory::AddClient(this, InSampleRateHz);
+	}
 }
 
 FFusionSampler::~FFusionSampler()
@@ -83,6 +89,7 @@ FFusionSampler::~FFusionSampler()
 		FusionPatchData->DisconnectSampler(this);
 	}
 	SetVoicePool(nullptr, true);
+	IStretcherAndPitchShifterFactory::RemoveClient(this);
 }
 
 void FFusionSampler::SetVoicePool(FSharedFusionVoicePoolPtr InPool, bool NoCallbacks)
@@ -513,10 +520,19 @@ bool FFusionSampler::TryKeyOnZone(FMidiVoiceId InVoiceId, uint8 InTriggeredNote,
 
 	static float sVelocityNormalizer = 1.0f / 127.0f;
 
+	IStretcherAndPitchShifterFactory* PitchShifterFactory = nullptr;
+	if (bUsePitchShifters)
+	{
+		if (InKeyzone->TimeStretchConfig.bMaintainTime)
+		{
+			PitchShifterFactory = IStretcherAndPitchShifterFactory::FindBestFactory(InKeyzone->TimeStretchConfig);
+		}
+	}
+	
 	// We send in the original, triggered note so that the allocator
 	// properly stops other sounding voices that started as a result of
 	// a trigger on the same note...
-	FFusionVoice* Voice = VoicePool->GetFreeVoice(this, InVoiceId, InKeyzone, [this](FFusionVoice* Voice) { return RelinquishVoice(Voice); });
+	FFusionVoice* Voice = VoicePool->GetFreeVoice(this, InVoiceId, InKeyzone, [this](FFusionVoice* Voice) { return RelinquishVoice(Voice); }, PitchShifterFactory);
 	if (Voice == nullptr)
 	{
 		UE_LOG(LogFusionSampler, Warning, TEXT("Voice pool failed to create voice, can't play note!"), InTriggeredNote);
diff --git a/Engine/Plugins/Runtime/Harmonix/Source/HarmonixDsp/Private/HarmonixDsp/FusionSampler/FusionVoicePool.cpp b/Engine/Plugins/Runtime/Harmonix/Source/HarmonixDsp/Private/HarmonixDsp/FusionSampler/FusionVoicePool.cpp
--- a/Engine/Plugins/Runtime/Harmonix/Source/HarmonixDsp/Private/HarmonixDsp/FusionSampler/FusionVoicePool.cpp
+++ b/Engine/Plugins/Runtime/Harmonix/Source/HarmonixDsp/Private/HarmonixDsp/FusionSampler/FusionVoicePool.cpp
@@ -21,11 +21,11 @@ void FFusionVoicePool::SetSampleRate(float InSampleRate)
 
 	FScopeLock Lock(&PoolLock);
 	KillVoices();
-	FreeVoicesAndShifters();
+	FreeVoices();
 	SampleRate = InSampleRate;
 	if (!DynamicAllocAndFree || ClientSamplers.Num() > 0)
 	{
-		AllocVoicesAndShifters();
+		AllocVoices();
 	}
 }
 
@@ -109,14 +109,14 @@ FFusionVoicePool::~FFusionVoicePool()
 	//mPoolLock.Exit();
 }
 
-void FFusionVoicePool::HardAllocatateVoicesAndShifters()
+void FFusionVoicePool::HardAllocatateVoices()
 {
 	if (!DynamicAllocAndFree)
 	{
 		return; // already done!
 	}
 
-	AllocVoicesAndShifters();
+	AllocVoices();
 	DynamicAllocAndFree = false;
 }
 
@@ -133,7 +133,7 @@ void FFusionVoicePool::ReleaseHardAllocation()
 
 	if (ClientSamplers.Num() == 0)
 	{
-		FreeVoicesAndShifters();
+		FreeVoices();
 	}
 }
 
@@ -162,7 +162,7 @@ void FFusionVoicePool::AddClient(FFusionSampler* InSampler)
 	{
 		// got our first client.
 		// better make sure we have the voices.
-		AllocVoicesAndShifters();
+		AllocVoices();
 	}
 }
 
@@ -181,7 +181,7 @@ void FFusionVoicePool::RemoveClient(FFusionSampler* InSampler)
 
 	if (ClientSamplers.IsEmpty() && DynamicAllocAndFree)
 	{
-		FreeVoicesAndShifters();
+		FreeVoices();
 	}
 }
 
@@ -266,13 +266,14 @@ FFusionVoice* FFusionVoicePool::GetFreeVoice(
 	FMidiVoiceId InVoiceID, 
 	const FKeyzoneSettings* InKeyzone, 
 	TFunction<bool(FFusionVoice*)> Handler,
+	IStretcherAndPitchShifterFactory* PitchShifterFactory,
 	bool AllowAlias, 
 	bool IsRendererForAlias)
 {
 #if 0
 	if (AllowAlias && InKeyzone->IsSingleton())
 	{
-		return InKeyzone->SingletonFusionVoicePool->AllocateAlias(this, InSampler, InVoiceID, Handler);
+		return InKeyzone->SingletonFusionVoicePool->AllocateAlias(this, InSampler, InVoiceID, Handler, PitchShifterFactory);
 	}
 #endif
 
@@ -291,24 +292,17 @@ FFusionVoice* FFusionVoicePool::GetFreeVoice(
 		return nullptr;
 	}
 
-
-	IStretcherAndPitchShifterFactory* Factory = nullptr;
-
-	if (InKeyzone->TimeStretchConfig.bMaintainTime)
+	if (!InKeyzone->SoundWaveProxy)
 	{
-
-		Factory = IStretcherAndPitchShifterFactory::FindBestFactory(InKeyzone->TimeStretchConfig);
-		if (!Factory)
-		{
-			FName FactoryName = InKeyzone->TimeStretchConfig.PitchShifter.Name;
-			ensureMsgf(false, TEXT("Attempting to play a fusion InKeyzone that is set to \"maintain time\", but no shifter factory was available for assigned Pitch Shifter: %s. Check project configuration! (%s)"), *FactoryName.ToString(), *InKeyzone->SoundWaveProxy->GetFName().ToString());
-			return nullptr;
-		}
+		UE_LOG(LogFusionVoicePool, Warning, TEXT("Asked to allocate a InSampler voice for a sample resource that failed to load -> %s\n"), *InKeyzone->SoundWaveProxy->GetFName().ToString());
+		return nullptr;
 	}
 
-	if (!InKeyzone->SoundWaveProxy)
+	// if maintain time is true and we didn't get a pitch shitfer factory, then we can't do anything!
+	if (InKeyzone->TimeStretchConfig.bMaintainTime && !PitchShifterFactory)
 	{
-		UE_LOG(LogFusionVoicePool, Warning, TEXT("Asked to allocate a InSampler voice for a sample resource that failed to load -> %s\n"), *InKeyzone->SoundWaveProxy->GetFName().ToString());
+		FName FactoryName = InKeyzone->TimeStretchConfig.PitchShifter.Name;
+		ensureMsgf(false, TEXT("Attempting to play a fusion InKeyzone that is set to \"maintain time\", but no shifter factory was available for assigned Pitch Shifter: %s. Check project configuration! (%s)"), *FactoryName.ToString(), *InKeyzone->SoundWaveProxy->GetFName().ToString());
 		return nullptr;
 	}
 
@@ -444,7 +438,7 @@ FFusionVoice* FFusionVoicePool::GetFreeVoice(
 	}
 
 	// Now we have a best choice. Now... do we need a shifter?
-	if (InKeyzone->TimeStretchConfig.bMaintainTime && !BestVoice->GetPitchShifter() && !Factory->HasFreePitchShifters(InKeyzone->TimeStretchConfig))
+	if (InKeyzone->TimeStretchConfig.bMaintainTime && !BestVoice->GetPitchShifter() && !PitchShifterFactory->HasFreePitchShifters(InKeyzone->TimeStretchConfig))
 	{
 		UE_LOG(LogFusionVoicePool, Warning, TEXT("Can't get free voice! Voice needs a shifter but we don't have one!"));
 		return nullptr;
@@ -455,7 +449,7 @@ FFusionVoice* FFusionVoicePool::GetFreeVoice(
 	TSharedPtr<IStretcherAndPitchShifter, ESPMode::ThreadSafe> ShifterToUse = nullptr;
 	if (InKeyzone->TimeStretchConfig.bMaintainTime)
 	{
-		ShifterToUse = Factory->GetFreePitchShifter(InKeyzone->TimeStretchConfig);
+		ShifterToUse = PitchShifterFactory->GetFreePitchShifter(InKeyzone->TimeStretchConfig);
 		check(ShifterToUse);
 		ShifterToUse->ApplyOptions(InKeyzone->TimeStretchConfig.PitchShifterOptions);
 	}
@@ -637,7 +631,7 @@ FFusionVoice* FFusionVoicePool::GetVoice(uint32 VoiceIdx)
 }
 
 
-void FFusionVoicePool::AllocVoicesAndShifters()
+void FFusionVoicePool::AllocVoices()
 {
 	CreateVoices(NumVoicesSetting);
 	CreateShifters();
@@ -676,10 +670,9 @@ void FFusionVoicePool::CreateVoices(uint16 InNumToAllocate)
 void FFusionVoicePool::CreateShifters()
 {
 	KillVoices();
-	IStretcherAndPitchShifterFactory::AddClient(this, SampleRate);
 }
 
-void FFusionVoicePool::FreeVoicesAndShifters()
+void FFusionVoicePool::FreeVoices()
 {
 	bool FoundVoiceInUse = false;
 	for (uint32 VoiceIdx = 0; VoiceIdx < NumAllocatedVoices; ++VoiceIdx)
@@ -695,6 +688,5 @@ void FFusionVoicePool::FreeVoicesAndShifters()
 	//SAFE_ARRAY_DELETE(mVoices);
 	delete[]Voices;
 	Voices = nullptr;
-	IStretcherAndPitchShifterFactory::RemoveClient(this);
 }
 
diff --git a/Engine/Plugins/Runtime/Harmonix/Source/HarmonixDsp/Private/HarmonixDsp/FusionSampler/SingletonFusionVoicePool.cpp b/Engine/Plugins/Runtime/Harmonix/Source/HarmonixDsp/Private/HarmonixDsp/FusionSampler/SingletonFusionVoicePool.cpp
--- a/Engine/Plugins/Runtime/Harmonix/Source/HarmonixDsp/Private/HarmonixDsp/FusionSampler/SingletonFusionVoicePool.cpp
+++ b/Engine/Plugins/Runtime/Harmonix/Source/HarmonixDsp/Private/HarmonixDsp/FusionSampler/SingletonFusionVoicePool.cpp
@@ -44,7 +44,8 @@ FFusionVoice* FSingletonFusionVoicePool::AllocateAlias(
 	FFusionVoicePool* InPool, 
 	FFusionSampler* InSampler, 
 	FMidiVoiceId InVoiceID, 
-	TFunction<bool(FFusionVoice*)> Handler)
+	TFunction<bool(FFusionVoice*)> Handler,
+	IStretcherAndPitchShifterFactory* PitchShifterFactory)
 {
 	FScopeLock Lock(&PoolLock);
 
@@ -91,13 +92,16 @@ FFusionVoice* FSingletonFusionVoicePool::AllocateAlias(
 		}
 
 		PrimaryVoice = PrimaryTracker.Voice = InPool->GetFreeVoice(
-			InSampler, InVoiceID, &KeyzoneRef,
+			InSampler,
+			InVoiceID,
+			&KeyzoneRef,
 			[this](FFusionVoice* HandleVoice)
 			{
 				return Relinquish(HandleVoice);
 			},
-			false, true
-		);
+			PitchShifterFactory,
+			false,
+			true);
 
 		if (!ensure(PrimaryVoice))
 		{
diff --git a/Engine/Plugins/Runtime/Harmonix/Source/HarmonixDsp/Public/HarmonixDsp/FusionSampler/FusionSampler.h b/Engine/Plugins/Runtime/Harmonix/Source/HarmonixDsp/Public/HarmonixDsp/FusionSampler/FusionSampler.h
--- a/Engine/Plugins/Runtime/Harmonix/Source/HarmonixDsp/Public/HarmonixDsp/FusionSampler/FusionSampler.h
+++ b/Engine/Plugins/Runtime/Harmonix/Source/HarmonixDsp/Public/HarmonixDsp/FusionSampler/FusionSampler.h
@@ -30,7 +30,7 @@ using FSharedFusionVoicePoolPtr = TSharedPtr<FFusionVoicePool, ESPMode::ThreadSa
 class HARMONIXDSP_API FFusionSampler : public FVirtualInstrument
 {
 public:
-	FFusionSampler();
+	FFusionSampler(float InSampleRate, bool InUsePitchShifters = false);
 	virtual ~FFusionSampler();
 
 	void ResetPatchRelatedState();
@@ -304,6 +304,7 @@ private:
 	float MaxPitchBendCents = 200.0f;
 
 	FSharedFusionVoicePoolPtr VoicePool;
+	const bool bUsePitchShifters;
 
 	// if we get note-ons or note-offs on the main
 	// thread, or if we get multiple requests for
diff --git a/Engine/Plugins/Runtime/Harmonix/Source/HarmonixDsp/Public/HarmonixDsp/FusionSampler/FusionVoicePool.h b/Engine/Plugins/Runtime/Harmonix/Source/HarmonixDsp/Public/HarmonixDsp/FusionSampler/FusionVoicePool.h
--- a/Engine/Plugins/Runtime/Harmonix/Source/HarmonixDsp/Public/HarmonixDsp/FusionSampler/FusionVoicePool.h
+++ b/Engine/Plugins/Runtime/Harmonix/Source/HarmonixDsp/Public/HarmonixDsp/FusionSampler/FusionVoicePool.h
@@ -70,7 +70,7 @@ public:
 
 	virtual ~FFusionVoicePool();
 
-	void HardAllocatateVoicesAndShifters();
+	void HardAllocatateVoices();
 	void ReleaseHardAllocation();
 
 	void SetIsMultithreading(bool InIsMultithreaded);
@@ -147,10 +147,12 @@ public:
 		FMidiVoiceId InVoiceID, 
 		const FKeyzoneSettings* InKeyzone, 
 		TFunction<bool(FFusionVoice*)> Handler,
+		IStretcherAndPitchShifterFactory* PitchShifterFactory,
 		bool AllowAlias = true, 
 		bool IsRendererForAlias = false);
 
 	void ReleaseShifter(TSharedPtr<IStretcherAndPitchShifter, ESPMode::ThreadSafe> InShifter);
+	
 	/**
 	 * Fast releases voices that have exceeded the soft limit for the pool (or the channel).
 	 * If no channel is specified, then voices assigned to any channel are candidates for
@@ -192,10 +194,10 @@ private:
 
 	uint32     PeakVoiceUsage = 0;
 
-	void AllocVoicesAndShifters();
+	void AllocVoices();
 	void CreateVoices(uint16 InMaxPolyphony);
 	void CreateShifters();
-	void FreeVoicesAndShifters();
+	void FreeVoices();
 
 	TArray<FFusionSampler*> ClientSamplers;
 
diff --git a/Engine/Plugins/Runtime/Harmonix/Source/HarmonixDsp/Public/HarmonixDsp/FusionSampler/SingletonFusionVoicePool.h b/Engine/Plugins/Runtime/Harmonix/Source/HarmonixDsp/Public/HarmonixDsp/FusionSampler/SingletonFusionVoicePool.h
--- a/Engine/Plugins/Runtime/Harmonix/Source/HarmonixDsp/Public/HarmonixDsp/FusionSampler/SingletonFusionVoicePool.h
+++ b/Engine/Plugins/Runtime/Harmonix/Source/HarmonixDsp/Public/HarmonixDsp/FusionSampler/SingletonFusionVoicePool.h
@@ -18,6 +18,7 @@ class FFusionVoice;
 class FFusionVoicePool;
 class FFusionSampler;
 class IStretcherAndPitchShifter;
+class IStretcherAndPitchShifterFactory;
 
 //================================================================================================
 // SingletonFusionVoicePool - This class allocates a REAL fusion voice when first asked to 
@@ -36,7 +37,8 @@ public:
 		FFusionVoicePool* InPool,
 		FFusionSampler* InSampler,
 		FMidiVoiceId InVoiceID,
-		TFunction<bool(FFusionVoice*)> Handler);
+		TFunction<bool(FFusionVoice*)> Handler,
+		IStretcherAndPitchShifterFactory* PitchShifterFactory);
 	bool IsDriver(FAliasFusionVoice* InVoice) const { return InVoice == DriverVoice; }
 
 	void SamplerDisconnecting(const FFusionSampler* Sampler);
diff --git a/Engine/Plugins/Runtime/Harmonix/Source/HarmonixMetasound/Private/HarmonixMetasound/Nodes/FusionSamplerNode.cpp b/Engine/Plugins/Runtime/Harmonix/Source/HarmonixMetasound/Private/HarmonixMetasound/Nodes/FusionSamplerNode.cpp
--- a/Engine/Plugins/Runtime/Harmonix/Source/HarmonixMetasound/Private/HarmonixMetasound/Nodes/FusionSamplerNode.cpp
+++ b/Engine/Plugins/Runtime/Harmonix/Source/HarmonixMetasound/Private/HarmonixMetasound/Nodes/FusionSamplerNode.cpp
@@ -39,6 +39,7 @@ namespace HarmonixMetasound
 		METASOUND_PARAM(EnableMTFusion, "Multithreaded Rendering", "Turn on to allow Fusion rendering to be done across multiple tasks. NOTE: YOU MUST then connect this node's "
 		                                 "Audio Output AND Render Sync Output to a Fusion Synchronizer Node, and reference THAT NODE'S audio outputs elsewhere in the graph!")
 		METASOUND_PARAM(VoicePoolName, "Voice Pool Name", "The name of a specific voice pool for this node to use. Uses the default pool if the named pool does not exist, or if the name is \"None\"")
+		METASOUND_PARAM(UsePitchShifters, "Use Pitch Shifters", "Set whether to use pitch shifters. If Disabled, pitch shifting will be disabled, and voices that are set to \"Maintain Time\" won't play.")
 		METASOUND_PARAM(RenderSync, "Render Sync", "YOU MUST connect this AND this node's audio output to a Fusion Synchronizer Node if you have enabled Multithreaded Rendering for this node!")
 	}
 
@@ -75,6 +76,7 @@ namespace HarmonixMetasound
 			FBoolReadRef             InClockSpeedToPitch;
 			bool		             InMTRenderingEnable;
 			FString					 InVoicePoolName;
+			bool                     InUsePitchShifters;
 			EAudioBufferChannelLayout InOutputChannelLayout;
 
 			FConstructionArgs(const FBuildOperatorParams& InParams, EAudioBufferChannelLayout OutputChannelLayout)
@@ -93,6 +95,7 @@ namespace HarmonixMetasound
 				, InClockSpeedToPitch(InParams.InputData.GetOrCreateDefaultDataReadReference<bool>(METASOUND_GET_PARAM_NAME(CommonPinNames::Inputs::ClockSpeedToPitch), InParams.OperatorSettings))
 				, InMTRenderingEnable(InParams.InputData.GetOrCreateDefaultValue<bool>(METASOUND_GET_PARAM_NAME(FusionSamplerNodePinNames::EnableMTFusion), InParams.OperatorSettings))
 				, InVoicePoolName(InParams.InputData.GetOrCreateDefaultValue<FString>(METASOUND_GET_PARAM_NAME(FusionSamplerNodePinNames::VoicePoolName), InParams.OperatorSettings))
+			    , InUsePitchShifters(InParams.InputData.GetOrCreateDefaultValue<bool>(METASOUND_GET_PARAM_NAME(FusionSamplerNodePinNames::UsePitchShifters), InParams.OperatorSettings))
 				, InOutputChannelLayout(OutputChannelLayout)
 			{
 			}
@@ -115,14 +118,16 @@ namespace HarmonixMetasound
 				TInputDataVertex<float>(METASOUND_GET_PARAM_NAME_AND_METADATA(FusionSamplerNodePinNames::FineTuneCents), -1200.0f),
 				TInputDataVertex<bool>(METASOUND_GET_PARAM_NAME_AND_METADATA(Inputs::ClockSpeedToPitch), true),
 				TInputConstructorVertex<bool>(METASOUND_GET_PARAM_NAME_AND_METADATA_ADVANCED(FusionSamplerNodePinNames::EnableMTFusion), false),
-				TInputConstructorVertex<FString>(METASOUND_GET_PARAM_NAME_AND_METADATA_ADVANCED(FusionSamplerNodePinNames::VoicePoolName), FString(TEXT("")))
+				TInputConstructorVertex<FString>(METASOUND_GET_PARAM_NAME_AND_METADATA_ADVANCED(FusionSamplerNodePinNames::VoicePoolName), FString(TEXT(""))),
+				TInputConstructorVertex<bool>(METASOUND_GET_PARAM_NAME_AND_METADATA_ADVANCED(FusionSamplerNodePinNames::UsePitchShifters), false)
 			);
 
 			return InputInterface;
 		}
 
 		FFusionSamplerOperatorBase(const FConstructionArgs& Args)
-			: SyncLinkOutPin(FFusionSyncLinkWriteRef::CreateNew())
+			: FFusionSampler(Args.Params.OperatorSettings.GetSampleRate(), Args.InUsePitchShifters)
+		    , SyncLinkOutPin(FFusionSyncLinkWriteRef::CreateNew())
 			, EnableInPin(Args.InEnabled)
 			, MidiStreamInPin(Args.InMidiStream)
 			, TrackNumberInPin(Args.InTrackNumber)
@@ -136,6 +141,7 @@ namespace HarmonixMetasound
 			, ClockSpeedAffectsPitchInPin(Args.InClockSpeedToPitch)
 			, EnableMTRenderingInPin(Args.InMTRenderingEnable)
 			, VoicePoolName(Args.InVoicePoolName)
+			, UsePitchShifters(Args.InUsePitchShifters)
 			, OutputChannelLayout(Args.InOutputChannelLayout)
 		{
 			Reset(Args.Params);
@@ -170,7 +176,7 @@ namespace HarmonixMetasound
 			InVertexData.BindReadVertex(METASOUND_GET_PARAM_NAME(Inputs::ClockSpeedToPitch), ClockSpeedAffectsPitchInPin);
 			InVertexData.SetValue(METASOUND_GET_PARAM_NAME(FusionSamplerNodePinNames::EnableMTFusion), EnableMTRenderingInPin);
 			InVertexData.SetValue(METASOUND_GET_PARAM_NAME(FusionSamplerNodePinNames::VoicePoolName), VoicePoolName);
-
+			InVertexData.SetValue(METASOUND_GET_PARAM_NAME(FusionSamplerNodePinNames::UsePitchShifters), UsePitchShifters);
 			Init();
 		}
 
@@ -220,6 +226,7 @@ namespace HarmonixMetasound
 		FBoolReadRef ClockSpeedAffectsPitchInPin;
 		bool EnableMTRenderingInPin;
 		FString VoicePoolName;
+		const bool UsePitchShifters;
 
 		//** DATA
 		int32 FramesPerBlock = 0;
