Fix issue with importing MIDI that was exported from Ableton with zero length track name string.

diff --git a/Engine/Plugins/Runtime/Harmonix/Source/HarmonixMidi/Private/HarmonixMidi/MidiReader.cpp b/Engine/Plugins/Runtime/Harmonix/Source/HarmonixMidi/Private/HarmonixMidi/MidiReader.cpp
--- a/Engine/Plugins/Runtime/Harmonix/Source/HarmonixMidi/Private/HarmonixMidi/MidiReader.cpp
+++ b/Engine/Plugins/Runtime/Harmonix/Source/HarmonixMidi/Private/HarmonixMidi/MidiReader.cpp
@@ -594,6 +594,11 @@ void FStdMidiFileReader::ReadMetaEvent(int32 Tick, uint8 Type)
 			// so we ignore the name we got from the file and normalize it here to "conductor"...
 			CurrentTrackName = TEXT("Conductor");
 		}
+		else if (CurrentTrackName.IsEmpty())
+		{
+			CurrentTrackName = FString::Format(TEXT("Track-{0}"), {CurrentTrackIndex - 1});
+			UE_LOG(LogMIDI, VeryVerbose, TEXT("TICK %d :: Got track name meta type but an empty string. Using a default track name..."), Tick);
+		}
 		UE_LOG(LogMIDI, VeryVerbose, TEXT("Got track name event... '%s'"), *CurrentTrackName);
 		MIDIREADER_EXPECT_WITH_ERROR(Receiver->OnText(Tick, CurrentTrackName, Type), TEXT("%s (%s): Failed to add track name to Track"), *Filename, *CurrentTrackName);
 		break;
@@ -720,11 +725,11 @@ void FStdMidiFileReader::ReadMetaEvent(int32 Tick, uint8 Type)
 FString FStdMidiFileReader::ReadText(int32 Length)
 {
 	// make sure length is rational...
-	MIDIREADER_EXPECT_WITH_ERROR_RETURN_VAL(Length > 0 && Length <= kMaxSupportedMidiStringSize, FString(),
+	MIDIREADER_EXPECT_WITH_ERROR_RETURN_VAL(Length >= 0 && Length <= kMaxSupportedMidiStringSize, FString(),
 		TEXT("%s (%s): %d byte long string found in MIDI data. Maximum supported string length is %d."), *Filename, *CurrentTrackName, Length, kMaxSupportedMidiStringSize);
 
 	// make sure there are enough bytes in the file...
-	if (FailIfBytesNotAvailableInTrackData(Length))
+	if (Length == 0 || FailIfBytesNotAvailableInTrackData(Length))
 	{
 		return FString();
 	}
