[HarmonixMetasound/MusicEventSubsystem]
- Make MusicEventSubsystem manage a MusicTempometer for the highest priority clock

#changelist validated
#virtualized

diff --git a/Engine/Plugins/Runtime/Harmonix/Source/HarmonixMetasound/Private/HarmonixMetasound/Components/MusicTempometerComponent.cpp b/Engine/Plugins/Runtime/Harmonix/Source/HarmonixMetasound/Private/HarmonixMetasound/Components/MusicTempometerComponent.cpp
--- a/Engine/Plugins/Runtime/Harmonix/Source/HarmonixMetasound/Private/HarmonixMetasound/Components/MusicTempometerComponent.cpp
+++ b/Engine/Plugins/Runtime/Harmonix/Source/HarmonixMetasound/Private/HarmonixMetasound/Components/MusicTempometerComponent.cpp
@@ -14,6 +14,79 @@ UMusicTempometerComponent::UMusicTempometerComponent()
 	PrimaryComponentTick.SetPriorityIncludingPrerequisites(true);
 }
 
+void UMusicTempometerComponent::PostLoad()
+{
+	Super::PostLoad();
+
+	// Fix up deprecated properties
+	const FName DeprecatedErroneousName = GetFName();
+	if (SecondsIncludingCountInParameterName_DEPRECATED != DeprecatedErroneousName)
+	{
+		MPCParameters.CurrentFrameParameterNames.SecondsIncludingCountInParameterName = SecondsIncludingCountInParameterName_DEPRECATED;
+		SecondsIncludingCountInParameterName_DEPRECATED = DeprecatedErroneousName;
+	}
+
+	if (BarsIncludingCountInParameterName_DEPRECATED != DeprecatedErroneousName)
+	{
+		MPCParameters.CurrentFrameParameterNames.BarsIncludingCountInParameterName = BarsIncludingCountInParameterName_DEPRECATED;
+		BarsIncludingCountInParameterName_DEPRECATED = DeprecatedErroneousName;
+	}
+	
+	if (BeatsIncludingCountInParameterName_DEPRECATED != DeprecatedErroneousName)
+	{
+		MPCParameters.CurrentFrameParameterNames.BeatsIncludingCountInParameterName = BeatsIncludingCountInParameterName_DEPRECATED;
+		BeatsIncludingCountInParameterName_DEPRECATED = DeprecatedErroneousName;
+	}
+	
+	if (SecondsFromBarOneParameterName_DEPRECATED != DeprecatedErroneousName)
+	{
+		MPCParameters.CurrentFrameParameterNames.SecondsFromBarOneParameterName = SecondsFromBarOneParameterName_DEPRECATED;
+		SecondsFromBarOneParameterName_DEPRECATED = DeprecatedErroneousName;
+	}
+	
+	if (TimestampBarParameterName_DEPRECATED != DeprecatedErroneousName)
+	{
+		MPCParameters.CurrentFrameParameterNames.TimestampBarParameterName = TimestampBarParameterName_DEPRECATED;
+		TimestampBarParameterName_DEPRECATED = DeprecatedErroneousName;
+	}
+
+	if (TimestampBeatInBarParameterName_DEPRECATED != DeprecatedErroneousName)
+	{
+		MPCParameters.CurrentFrameParameterNames.TimestampBeatInBarParameterName = TimestampBeatInBarParameterName_DEPRECATED;
+		TimestampBeatInBarParameterName_DEPRECATED = DeprecatedErroneousName;
+	}
+
+	if (BarProgressParameterName_DEPRECATED != DeprecatedErroneousName)
+	{
+		MPCParameters.CurrentFrameParameterNames.BarProgressParameterName = BarProgressParameterName_DEPRECATED;
+		BarProgressParameterName_DEPRECATED = DeprecatedErroneousName;
+	}
+
+	if (BeatProgressParameterName_DEPRECATED != DeprecatedErroneousName)
+	{
+		MPCParameters.CurrentFrameParameterNames.BeatProgressParameterName = BeatProgressParameterName_DEPRECATED;
+		BeatProgressParameterName_DEPRECATED = DeprecatedErroneousName;
+	}
+
+	if (TimeSignatureNumeratorParameterName_DEPRECATED != DeprecatedErroneousName)
+	{
+		MPCParameters.CurrentFrameParameterNames.TimeSignatureNumeratorParameterName = TimeSignatureNumeratorParameterName_DEPRECATED;
+		TimeSignatureNumeratorParameterName_DEPRECATED = DeprecatedErroneousName;
+	}
+
+	if (TimeSignatureDenominatorParameterName_DEPRECATED != DeprecatedErroneousName)
+	{
+		MPCParameters.CurrentFrameParameterNames.TimeSignatureDenominatorParameterName = TimeSignatureDenominatorParameterName_DEPRECATED;
+		TimeSignatureDenominatorParameterName_DEPRECATED = DeprecatedErroneousName;
+	}
+
+	if (TempoParameterName_DEPRECATED != DeprecatedErroneousName)
+	{
+		MPCParameters.CurrentFrameParameterNames.TempoParameterName = TempoParameterName_DEPRECATED;
+		TempoParameterName_DEPRECATED = DeprecatedErroneousName;
+	}
+}
+
 void UMusicTempometerComponent::TickComponent(float DeltaTime, enum ELevelTick TickType, FActorComponentTickFunction* ThisTickFunction)
 {
 	Super::TickComponent(DeltaTime, TickType, ThisTickFunction);
@@ -54,37 +127,16 @@ void UMusicTempometerComponent::UpdateCachedSongPos() const
 		SongPos.Reset();
 	}
 
-	// Find a MaterialParameterCollectionInstance to update
-	if (!IsValid(MaterialParameterCollectionInstance))
-	{
-		if (MaterialParameterCollection)
-		{
-			if (UWorld* World = GetWorld())
-			{
-				MaterialParameterCollectionInstance = World->GetParameterCollectionInstance(MaterialParameterCollection);
-			}
-		}
+	MusicTempometerUtilities::UpdateMaterialParameterCollectionFromSongPos(GetOwner(), MaterialParameterCollectionInstance, MaterialParameterCollection, MPCParameters, SongPos);
+}
 
-		if (!IsValid(MaterialParameterCollectionInstance))
-		{
-			return;
-		}
+void UMusicTempometerComponent::UpdateCachedSongPosIfNeeded() const
+{
+	FScopeLock lock(&SongPosUpdateMutex);
+	if (GFrameCounter != LastFrameCounter)
+	{
+		UpdateCachedSongPos();
 	}
-
-	MaterialParameterCollectionInstance->SetScalarParameterValue(SecondsIncludingCountInParameterName, SongPos.SecondsIncludingCountIn);
-	MaterialParameterCollectionInstance->SetScalarParameterValue(BarsIncludingCountInParameterName, SongPos.BarsIncludingCountIn);
-	MaterialParameterCollectionInstance->SetScalarParameterValue(BeatsIncludingCountInParameterName, SongPos.BeatsIncludingCountIn);
-
-	MaterialParameterCollectionInstance->SetScalarParameterValue(SecondsFromBarOneParameterName, SongPos.SecondsFromBarOne);
-	MaterialParameterCollectionInstance->SetScalarParameterValue(TimestampBarParameterName, (float)SongPos.Timestamp.Bar);
-	MaterialParameterCollectionInstance->SetScalarParameterValue(TimestampBeatInBarParameterName, SongPos.Timestamp.Beat);
-
-	MaterialParameterCollectionInstance->SetScalarParameterValue(BarProgressParameterName, FMath::Fractional(SongPos.BarsIncludingCountIn));
-	MaterialParameterCollectionInstance->SetScalarParameterValue(BeatProgressParameterName, FMath::Fractional(SongPos.BeatsIncludingCountIn));
-
-	MaterialParameterCollectionInstance->SetScalarParameterValue(TimeSignatureNumeratorParameterName, SongPos.TimeSigNumerator);
-	MaterialParameterCollectionInstance->SetScalarParameterValue(TimeSignatureDenominatorParameterName, SongPos.TimeSigDenominator);
-	MaterialParameterCollectionInstance->SetScalarParameterValue(TempoParameterName, SongPos.Tempo);
 }
 
 UMusicClockComponent* UMusicTempometerComponent::FindClock(AActor* Actor) const
diff --git a/Engine/Plugins/Runtime/Harmonix/Source/HarmonixMetasound/Private/HarmonixMetasound/Utils/MusicTempometerUtilities.cpp b/Engine/Plugins/Runtime/Harmonix/Source/HarmonixMetasound/Private/HarmonixMetasound/Utils/MusicTempometerUtilities.cpp
new file mode 100644
--- /dev/null
+++ b/Engine/Plugins/Runtime/Harmonix/Source/HarmonixMetasound/Private/HarmonixMetasound/Utils/MusicTempometerUtilities.cpp
@@ -0,0 +1,56 @@
+// Copyright Epic Games, Inc. All Rights Reserved.
+
+#include "HarmonixMetasound/Components/MusicClockComponent.h"
+#include "HarmonixMetasound/Utils/MusicTempometerUtilities.h"
+#include "HarmonixMidi/MidiSongPos.h"
+#include "Materials/MaterialParameterCollectionInstance.h"
+#include "Engine/World.h"
+
+#include UE_INLINE_GENERATED_CPP_BY_NAME(MusicTempometerUtilities)
+
+void MusicTempometerUtilities::UpdateMaterialParameterCollectionFromClock(const UObject* InWorldContextObject, TObjectPtr<UMaterialParameterCollectionInstance>& InOutMaterialParameterCollectionInstance, const TObjectPtr<UMaterialParameterCollection>& InMaterialParameterCollection, const FMusicTempometerMPCParameters& InMCPParameters, const UMusicClockComponent* InClockComponent)
+{
+	if (InClockComponent)
+	{
+		const FMidiSongPos MidiSongPos = InClockComponent->GetCurrentVideoRenderSongPos();
+		UpdateMaterialParameterCollectionFromSongPos(InWorldContextObject, InOutMaterialParameterCollectionInstance, InMaterialParameterCollection, InMCPParameters, MidiSongPos);
+	}
+}
+
+void MusicTempometerUtilities::UpdateMaterialParameterCollectionFromSongPos(const UObject* InWorldContextObject, TObjectPtr<UMaterialParameterCollectionInstance>& InOutMaterialParameterCollectionInstance, const TObjectPtr<UMaterialParameterCollection>& InMaterialParameterCollection, const FMusicTempometerMPCParameters& InMCPParameters, const FMidiSongPos& InMidiSongPos)
+{
+	// Find a MaterialParameterCollectionInstance to update
+	if (InOutMaterialParameterCollectionInstance == nullptr)
+	{
+		if (InMCPParameters.IsValid())
+		{
+			if (InWorldContextObject)
+			{
+				if (UWorld* World = InWorldContextObject->GetWorld())
+				{
+					InOutMaterialParameterCollectionInstance = World->GetParameterCollectionInstance(InMaterialParameterCollection);
+				}
+			}
+		}
+
+		if (InOutMaterialParameterCollectionInstance == nullptr)
+		{
+			return;
+		}
+	}
+
+	InOutMaterialParameterCollectionInstance->SetScalarParameterValue(InMCPParameters.CurrentFrameParameterNames.SecondsIncludingCountInParameterName, InMidiSongPos.SecondsIncludingCountIn);
+	InOutMaterialParameterCollectionInstance->SetScalarParameterValue(InMCPParameters.CurrentFrameParameterNames.BarsIncludingCountInParameterName, InMidiSongPos.BarsIncludingCountIn);
+	InOutMaterialParameterCollectionInstance->SetScalarParameterValue(InMCPParameters.CurrentFrameParameterNames.BeatsIncludingCountInParameterName, InMidiSongPos.BeatsIncludingCountIn);
+
+	InOutMaterialParameterCollectionInstance->SetScalarParameterValue(InMCPParameters.CurrentFrameParameterNames.SecondsFromBarOneParameterName, InMidiSongPos.SecondsFromBarOne);
+	InOutMaterialParameterCollectionInstance->SetScalarParameterValue(InMCPParameters.CurrentFrameParameterNames.TimestampBarParameterName, static_cast<float>(InMidiSongPos.Timestamp.Bar));
+	InOutMaterialParameterCollectionInstance->SetScalarParameterValue(InMCPParameters.CurrentFrameParameterNames.TimestampBeatInBarParameterName, InMidiSongPos.Timestamp.Beat);
+
+	InOutMaterialParameterCollectionInstance->SetScalarParameterValue(InMCPParameters.CurrentFrameParameterNames.BarProgressParameterName, FMath::Fractional(InMidiSongPos.BarsIncludingCountIn));
+	InOutMaterialParameterCollectionInstance->SetScalarParameterValue(InMCPParameters.CurrentFrameParameterNames.BeatProgressParameterName, FMath::Fractional(InMidiSongPos.BeatsIncludingCountIn));
+
+	InOutMaterialParameterCollectionInstance->SetScalarParameterValue(InMCPParameters.CurrentFrameParameterNames.TimeSignatureNumeratorParameterName, InMidiSongPos.TimeSigNumerator);
+	InOutMaterialParameterCollectionInstance->SetScalarParameterValue(InMCPParameters.CurrentFrameParameterNames.TimeSignatureDenominatorParameterName, InMidiSongPos.TimeSigDenominator);
+	InOutMaterialParameterCollectionInstance->SetScalarParameterValue(InMCPParameters.CurrentFrameParameterNames.TempoParameterName, InMidiSongPos.Tempo);
+}
\ No newline at end of file
diff --git a/Engine/Plugins/Runtime/Harmonix/Source/HarmonixMetasound/Public/HarmonixMetasound/Components/MusicTempometerComponent.h b/Engine/Plugins/Runtime/Harmonix/Source/HarmonixMetasound/Public/HarmonixMetasound/Components/MusicTempometerComponent.h
--- a/Engine/Plugins/Runtime/Harmonix/Source/HarmonixMetasound/Public/HarmonixMetasound/Components/MusicTempometerComponent.h
+++ b/Engine/Plugins/Runtime/Harmonix/Source/HarmonixMetasound/Public/HarmonixMetasound/Components/MusicTempometerComponent.h
@@ -3,6 +3,7 @@
 #pragma once
 
 #include "HarmonixMetasound/Components/MusicClockComponent.h"
+#include "HarmonixMetasound/Utils/MusicTempometerUtilities.h"
 #include "HarmonixMidi/MidiSongPos.h"
 #include "Components/ActorComponent.h"
 
@@ -23,92 +24,76 @@ public:
 	UMusicTempometerComponent();
 
 	/**
-	 * Seconds from the beginning of the entire music authoring.
-	 * Includes all count-in and pickup bars (ie. won't be negative when 
-	 * the music starts, and bar 1 beat 1 may not be at 0.0 seconds!
+	 * Deprecated, has moved into structure within MPCParameters
 	 */
-	UPROPERTY(EditAnywhere, Category = "MusicClock")
-	FName SecondsIncludingCountInParameterName = "MusicSecondsIncludingCountIn";
+	UPROPERTY()
+	FName SecondsIncludingCountInParameterName_DEPRECATED = "MusicSecondsIncludingCountIn";
 
 	/**
-	 * Bars from the beginning of the music. Includes all 
-	 * count-in and pickup bars (ie. won't be negative when the music starts,
-	 * and bar 1 beat 1 of the music may not be equal to elapsed bar 0.0!
-	 * NOTE: This Bar, unlike the Bar in the MusicTimestamp structure, is
-	 * floating point and will include a fractional portion.
+	 * Deprecated, has moved into structure within MPCParameters
 	 */
-	UPROPERTY(EditAnywhere, Category = "MusicClock")
-	FName BarsIncludingCountInParameterName = "MusicBarsIncludingCountIn";
+	UPROPERTY()
+	FName BarsIncludingCountInParameterName_DEPRECATED = "MusicBarsIncludingCountIn";
 
 	/**
-	 * Total Beats from the beginning of the music. Includes all 
-	 * count-in and pickup bars/beats (ie. won't be negative when the music starts,
-	 * and elapsed beat 0.0 may not equal a timestamp of bar 1 beat 1!
+	 * Deprecated, has moved into structure within MPCParameters
 	 */
-	UPROPERTY(EditAnywhere, Category = "MusicClock")
-	FName BeatsIncludingCountInParameterName = "MusicBeatsIncludingCountIn";
+	UPROPERTY()
+	FName BeatsIncludingCountInParameterName_DEPRECATED = "MusicBeatsIncludingCountIn";
 
 	/**
-	 * Seconds from Bar 1 Beat 1 of the music. If the music has a count-in
-	 * or pickup bars this number may be negative when the music starts!
+	 * Deprecated, has moved into structure within MPCParameters
 	 */
-	UPROPERTY(EditAnywhere, Category = "MusicClock")
-	FName SecondsFromBarOneParameterName = "MusicSecondsFromBarOne";
+	UPROPERTY()
+	FName SecondsFromBarOneParameterName_DEPRECATED = "MusicSecondsFromBarOne";
 
 	/**
-	 * This is the bar member of the MusicTimestamp structure which represents a classic
-	 * "music time" where bar 1 beat 1 is the beginning of the music AFTER count-in and pickups.
-	 * If the music has a count-in or pickup bars this number may be negative (or zero) until 
-	 * bar 1 beat 1 is reached!
-	 * NOTE: Unlike "BarsFromCountIn", Bar here will be an integer! This because it is the bar 
-	 * from the MusicTimestamp structure which also provides a "beat in bar" member to denote 
-	 * a "distance into the bar".
+	 * Deprecated, has moved into structure within MPCParameters
 	 */
-	UPROPERTY(EditAnywhere, Category = "MusicClock")
-	FName TimestampBarParameterName = "MusicTimestampBar";
+	UPROPERTY()
+	FName TimestampBarParameterName_DEPRECATED = "MusicTimestampBar";
 
 	/**
-	 * This is the beat member of the MusicTimestamp structure which represents a classic
-	 * "music time" where bar 1 beat 1 is the beginning of the music AFTER count-in and pickups.
-	 * If the music has a count-in or pickup bars this number may be negative (or zero) until
-	 * bar 1 beat 1 is reached! NOTE: beat in this context is 1 based! The first beat in a bar
-	 * is beat 1!
-	 * Note: This is a floating point number.
+	 * Deprecated, has moved into structure within MPCParameters
 	 */
-	UPROPERTY(EditAnywhere, Category = "MusicClock")
-	FName TimestampBeatInBarParameterName = "MusicTimestampBeatInBar";
+	UPROPERTY()
+	FName TimestampBeatInBarParameterName_DEPRECATED = "MusicTimestampBeatInBar";
 
 	/**
-	 * Progress of the current bar [0, 1]. This is the same as the
-	 * fractional part of BarsFromCountIn.
+	 * Deprecated, has moved into structure within MPCParameters
 	 */
-	UPROPERTY(EditAnywhere, Category = "MusicClock")
-	FName BarProgressParameterName = "MusicBarProgress";
+	UPROPERTY()
+	FName BarProgressParameterName_DEPRECATED = "MusicBarProgress";
 
 	/**
-	 * Progress of the current beat [0, 1]. This is the same as the
-	 * fractional part of BeatsFromCountIn;
+	 * Deprecated, has moved into structure within MPCParameters
 	 */
-	UPROPERTY(EditAnywhere, Category = "MusicClock")
-	FName BeatProgressParameterName = "MusicBeatProgress";
+	UPROPERTY()
+	FName BeatProgressParameterName_DEPRECATED = "MusicBeatProgress";
 
 	/**
-	 * Current time signature numerator (beats per bar).
+	 * Deprecated, has moved into structure within MPCParameters
 	 */
-	UPROPERTY(EditAnywhere, Category = "MusicClock")
-	FName TimeSignatureNumeratorParameterName = "MusicTimeSignatureNumerator";
+	UPROPERTY()
+	FName TimeSignatureNumeratorParameterName_DEPRECATED = "MusicTimeSignatureNumerator";
 
 	/**
-	 * Current time signature denominator (scale from note duration to beat).
+	 * Deprecated, has moved into structure within MPCParameters
 	 */
-	UPROPERTY(EditAnywhere, Category = "MusicClock")
-	FName TimeSignatureDenominatorParameterName = "MusicTimeSignatureDenominator";
+	UPROPERTY()
+	FName TimeSignatureDenominatorParameterName_DEPRECATED = "MusicTimeSignatureDenominator";
 
 	/**
-	 * Current tempo (beats per minute).
+	 * Deprecated, has moved into structure within MPCParameters
+	 */
+	UPROPERTY()
+	FName TempoParameterName_DEPRECATED = "MusicTempo";
+
+	/**
+	 * Deprecated, has moved into structure within MPCParameters
 	 */
 	UPROPERTY(EditAnywhere, Category = "MusicClock")
-	FName TempoParameterName = "MusicTempo";
+	FMusicTempometerMPCParameters MPCParameters;
 
 	/**
 	 * The FSongPos for which the game thread is currently issuing graphics rendering commands, according to calibration data.
@@ -278,6 +263,7 @@ public:
 	}
 
 	//~ Begin UObject Interface.
+	virtual void PostLoad() override;
 	virtual void TickComponent(float DeltaTime, enum ELevelTick TickType, FActorComponentTickFunction* ThisTickFunction) override;
 #if WITH_EDITOR
 	virtual void PostEditChangeProperty(FPropertyChangedEvent& PropertyChangedEvent) override;
@@ -308,14 +294,7 @@ private:
 	}
 
 	void UpdateCachedSongPos() const;
-	void UpdateCachedSongPosIfNeeded() const
-	{
-		FScopeLock lock(&SongPosUpdateMutex);
-		if (GFrameCounter != LastFrameCounter)
-		{
-			UpdateCachedSongPos();
-		}
-	}
+	void UpdateCachedSongPosIfNeeded() const;
 
 	UPROPERTY(BlueprintGetter = GetSongPos, Transient, Category = "MusicClock")
 	mutable FMidiSongPos SongPos;
@@ -332,6 +311,9 @@ private:
 	UPROPERTY(EditAnywhere, BlueprintGetter = GetMaterialParameterCollection, BlueprintSetter = SetMaterialParameterCollection, Category = "MusicClock")
 	TObjectPtr<UMaterialParameterCollection> MaterialParameterCollection;
 
+	/**
+	 * Pinned material parameter collection instance
+	 */
 	UPROPERTY(Transient)
 	mutable TObjectPtr<UMaterialParameterCollectionInstance> MaterialParameterCollectionInstance;
 };
diff --git a/Engine/Plugins/Runtime/Harmonix/Source/HarmonixMetasound/Public/HarmonixMetasound/Utils/MusicTempometerUtilities.h b/Engine/Plugins/Runtime/Harmonix/Source/HarmonixMetasound/Public/HarmonixMetasound/Utils/MusicTempometerUtilities.h
new file mode 100644
--- /dev/null
+++ b/Engine/Plugins/Runtime/Harmonix/Source/HarmonixMetasound/Public/HarmonixMetasound/Utils/MusicTempometerUtilities.h
@@ -0,0 +1,153 @@
+// Copyright Epic Games, Inc. All Rights Reserved.
+
+#pragma once
+
+#include "MusicTempometerUtilities.generated.h"
+
+class UMaterialParameterCollection;
+class UMaterialParameterCollectionInstance;
+class UMusicClockComponent;
+
+struct FMidiSongPos;
+
+USTRUCT(BlueprintType)
+struct FMusicTempometerMPCParameterNames
+{
+	GENERATED_BODY()
+
+	/**
+	 * Seconds from the beginning of the entire music authoring.
+	 * Includes all count-in and pickup bars (ie. won't be negative when 
+	 * the music starts, and bar 1 beat 1 may not be at 0.0 seconds!
+	 */
+	UPROPERTY(EditDefaultsOnly, Category="MusicClock")
+	FName SecondsIncludingCountInParameterName = "MusicSecondsIncludingCountIn";
+
+	/**
+	 * Bars from the beginning of the music. Includes all
+	 * count-in and pickup bars (ie. won't be negative when the music starts,
+	 * and bar 1 beat 1 of the music may not be equal to elapsed bar 0.0!
+	 * NOTE: This Bar, unlike the Bar in the MusicTimestamp structure, is
+	 * floating point and will include a fractional portion.
+	 */
+	UPROPERTY(EditDefaultsOnly, Category = "MusicClock")
+	FName BarsIncludingCountInParameterName = "MusicBarsIncludingCountIn";
+
+	/**
+	 * Total Beats from the beginning of the music. Includes all
+	 * count-in and pickup bars/beats (ie. won't be negative when the music starts,
+	 * and elapsed beat 0.0 may not equal a timestamp of bar 1 beat 1!
+	 */
+	UPROPERTY(EditDefaultsOnly, Category = "MusicClock")
+	FName BeatsIncludingCountInParameterName = "MusicBeatsIncludingCountIn";
+
+	/**
+	 * Seconds from Bar 1 Beat 1 of the music. If the music has a count-in
+	 * or pickup bars this number may be negative when the music starts!
+	 */
+	UPROPERTY(EditDefaultsOnly, Category = "MusicClock")
+	FName SecondsFromBarOneParameterName = "MusicSecondsFromBarOne";
+
+	/**
+	 * This is the bar member of the MusicTimestamp structure which represents a classic
+	 * "music time" where bar 1 beat 1 is the beginning of the music AFTER count-in and pickups.
+	 * If the music has a count-in or pickup bars this number may be negative (or zero) until
+	 * bar 1 beat 1 is reached!
+	 * NOTE: Unlike "BarsFromCountIn", Bar here will be an integer! This because it is the bar
+	 * from the MusicTimestamp structure which also provides a "beat in bar" member to denote
+	 * a "distance into the bar".
+	 */
+	UPROPERTY(EditDefaultsOnly, Category = "MusicClock")
+	FName TimestampBarParameterName = "MusicTimestampBar";
+
+	/**
+	 * This is the beat member of the MusicTimestamp structure which represents a classic
+	 * "music time" where bar 1 beat 1 is the beginning of the music AFTER count-in and pickups.
+	 * If the music has a count-in or pickup bars this number may be negative (or zero) until
+	 * bar 1 beat 1 is reached! NOTE: beat in this context is 1 based! The first beat in a bar
+	 * is beat 1!
+	 * Note: This is a floating point number.
+	 */
+	UPROPERTY(EditDefaultsOnly, Category = "MusicClock")
+	FName TimestampBeatInBarParameterName = "MusicTimestampBeatInBar";
+
+	/**
+	 * Progress of the current bar [0, 1]. This is the same as the
+	 * fractional part of BarsFromCountIn.
+	 */
+	UPROPERTY(EditDefaultsOnly, Category = "MusicClock")
+	FName BarProgressParameterName = "MusicBarProgress";
+
+	/**
+	 * Progress of the current beat [0, 1]. This is the same as the
+	 * fractional part of BeatsFromCountIn;
+	 */
+	UPROPERTY(EditDefaultsOnly, Category = "MusicClock")
+	FName BeatProgressParameterName = "MusicBeatProgress";
+
+	/**
+	 * Current time signature numerator (beats per bar).
+	 */
+	UPROPERTY(EditDefaultsOnly, Category = "MusicClock")
+	FName TimeSignatureNumeratorParameterName = "MusicTimeSignatureNumerator";
+
+	/**
+	 * Current time signature denominator (scale from note duration to beat).
+	 */
+	UPROPERTY(EditDefaultsOnly, Category = "MusicClock")
+	FName TimeSignatureDenominatorParameterName = "MusicTimeSignatureDenominator";
+
+	/**
+	 * Current tempo (beats per minute).
+	 */
+	UPROPERTY(EditDefaultsOnly, Category = "MusicClock")
+	FName TempoParameterName = "MusicTempo";
+
+	bool IsValid() const
+	{
+		// Valid if there is an MPC collection assigned and at least one named parameter
+		return	!SecondsIncludingCountInParameterName.IsNone() ||
+				!BarsIncludingCountInParameterName.IsNone() ||
+				!BeatsIncludingCountInParameterName.IsNone() ||
+				!SecondsFromBarOneParameterName.IsNone() ||
+				!TimestampBarParameterName.IsNone() ||
+				!TimestampBeatInBarParameterName.IsNone() ||
+				!BarProgressParameterName.IsNone() ||
+				!BeatProgressParameterName.IsNone() ||
+				!TimeSignatureNumeratorParameterName.IsNone() ||
+				!TimeSignatureDenominatorParameterName.IsNone() ||
+				!TempoParameterName.IsNone();
+	}
+};
+/**
+ * Parameters used when creating/updating the corresponding Material Parameter Collection
+  */
+USTRUCT(BlueprintType)
+struct FMusicTempometerMPCParameters
+{
+	GENERATED_BODY();
+
+	UPROPERTY(EditDefaultsOnly, Category = "MusicClock", meta = (ShowOnlyInnerProperties))
+	FMusicTempometerMPCParameterNames CurrentFrameParameterNames;
+
+	bool IsValid() const
+	{
+		return CurrentFrameParameterNames.IsValid();
+	}
+};
+
+/**
+ * Utility methods used by the UMusicTempometerComponent, also publically exposed for other systems to use.
+ */
+namespace MusicTempometerUtilities
+{
+	/**
+	 * Update the supplied material parameter collection instance (create if null) with song position data extracted from the supplied clock
+	 */
+	void HARMONIXMETASOUND_API UpdateMaterialParameterCollectionFromClock(const UObject* InWorldContextObject, TObjectPtr<UMaterialParameterCollectionInstance>& InOutMaterialParameterCollectionInstance, const TObjectPtr<UMaterialParameterCollection>& InMaterialParameterCollection, const FMusicTempometerMPCParameters& InMCPParameters, const UMusicClockComponent* InClockComponent);
+	
+	/**
+	 * Update the supplied material parameter collection instance (create if null) with data from the supplied song position
+	 */
+	void HARMONIXMETASOUND_API UpdateMaterialParameterCollectionFromSongPos(const UObject* InWorldContextObject, TObjectPtr<UMaterialParameterCollectionInstance>& InOutMaterialParameterCollectionInstance, const TObjectPtr<UMaterialParameterCollection>& InMaterialParameterCollection, const FMusicTempometerMPCParameters& InMCPParameters, const FMidiSongPos& InMidiSongPos);
+};
