[Harmonix] Fixed issue where notes would retrigger when changing pages with a step length of one
#tests PIE

diff --git a/Engine/Plugins/Runtime/Harmonix/Source/HarmonixMetasound/Private/HarmonixMetasound/Nodes/StepSequencePlayerNode.cpp b/Engine/Plugins/Runtime/Harmonix/Source/HarmonixMetasound/Private/HarmonixMetasound/Nodes/StepSequencePlayerNode.cpp
--- a/Engine/Plugins/Runtime/Harmonix/Source/HarmonixMetasound/Private/HarmonixMetasound/Nodes/StepSequencePlayerNode.cpp
+++ b/Engine/Plugins/Runtime/Harmonix/Source/HarmonixMetasound/Private/HarmonixMetasound/Nodes/StepSequencePlayerNode.cpp
@@ -755,6 +755,7 @@ namespace HarmonixMetasound::Nodes::StepSequencePlayer
 			int32 CellInRow = TickInTable / TicksPerCell;
 			const int32 TickInCell = TickInTable - (CellInRow * TicksPerCell);
 			int32 SkippedIndex = CurrentStepSkipIndex - 1;
+			const bool bFirstTickInCell = (TickInTable % TicksPerCell) == 0;
 
 			if (CurrentStepSkipIndex >= 2)
 			{
@@ -822,7 +823,7 @@ namespace HarmonixMetasound::Nodes::StepSequencePlayer
 				}
 			}
 
-			if (CellInRow != CurrentCellIndex || CurrentMaxColumns == 1)
+			if (CellInRow != CurrentCellIndex || (CurrentMaxColumns == 1 && bFirstTickInCell))
 			{	
 				for (int32 NoteIdx = 0; NoteIdx < CurrentPage->Rows.Num(); ++NoteIdx)
 				{
