MusicClockCopmonent Refactor: Rewrite Music Clock Component logic out of the Component and into the MusicClockDriverBase.

This is so that we can have a Music Clock Sources that aren't attached to an actor. The Music Clock Component will simply use the music clock driver to report the current time.

Future changes would be to change the MidiClockUpdatedSubsystem to take tick MusicClockDrivers instead of MusicClockComponents.

#rnx
#tests HarmonixFunctionalTests

diff --git a/Engine/Plugins/Runtime/Harmonix/Source/HarmonixMetasound/Private/HarmonixMetasound/Components/MetasoundMusicClockDriver.cpp b/Engine/Plugins/Runtime/Harmonix/Source/HarmonixMetasound/Private/HarmonixMetasound/Components/MetasoundMusicClockDriver.cpp
--- a/Engine/Plugins/Runtime/Harmonix/Source/HarmonixMetasound/Private/HarmonixMetasound/Components/MetasoundMusicClockDriver.cpp
+++ b/Engine/Plugins/Runtime/Harmonix/Source/HarmonixMetasound/Private/HarmonixMetasound/Components/MetasoundMusicClockDriver.cpp
@@ -1,14 +1,16 @@
 // Copyright Epic Games, Inc. All Rights Reserved.
+
 #include "MetasoundMusicClockDriver.h"
 #include "MetasoundGeneratorHandle.h"
 #include "Components/AudioComponent.h"
 #include "HarmonixMetasound/Analysis/MidiSongPosVertexAnalyzer.h"
 #include "HarmonixMetasound/DataTypes/MidiClock.h"
 #include "MetasoundGenerator.h"
-#include "Engine/World.h"
 #include "Async/Async.h"
 #include "Harmonix.h"
 
+DEFINE_LOG_CATEGORY_STATIC(LogMetasoundMusicClockDriver, Log, All)
+
 namespace MetasoundMusicClockDriver
 {
 	float Fudge = 1.00;
@@ -125,7 +127,7 @@ FMidiSongPos FMetasoundMusicClockDriver::CalculateSongPosAtMsForLoopingOrMonoton
 			LoopDetected = (PositionTick - WrappedTick) > (float)(CurrentMapChain->LoopLengthTicks - 240);
 			if (LoopDetected)
 			{
-				UE_LOG(LogMusicClock, VeryVerbose, TEXT("Detected loop when calculating song pos (%f -> %f -> %f -> %d )"), PositionTick, WrappedTick, PositionTick - WrappedTick, CurrentMapChain->LoopLengthTicks);
+				UE_LOG(LogMetasoundMusicClockDriver, VeryVerbose, TEXT("Detected loop when calculating song pos (%f -> %f -> %f -> %d )"), PositionTick, WrappedTick, PositionTick - WrappedTick, CurrentMapChain->LoopLengthTicks);
 			}
 			OutSongPos.SetByTick(WrappedTick, *CurrentMapChain->SongMaps);
 			OutSongPos.Tempo = CurrentMapChain->SongMaps->GetTempoAtTick(FMath::FloorToInt32(DrivingTick));
@@ -135,7 +137,7 @@ FMidiSongPos FMetasoundMusicClockDriver::CalculateSongPosAtMsForLoopingOrMonoton
 				SeekDetected = CheckForSeek(PositionTick, WrappedTick, OutSongPos.Tempo, CurrentMapChain->SongMaps->GetTicksPerQuarterNote());
 				if (SeekDetected)
 				{
-					UE_LOG(LogMusicClock, VeryVerbose, TEXT("Detected seek when calculating song pos (%f -> %f)"), PositionTick, WrappedTick);
+					UE_LOG(LogMetasoundMusicClockDriver, VeryVerbose, TEXT("Detected seek when calculating song pos (%f -> %f)"), PositionTick, WrappedTick);
 				}
 			}
 
@@ -153,7 +155,7 @@ FMidiSongPos FMetasoundMusicClockDriver::CalculateSongPosAtMsForLoopingOrMonoton
 	SeekDetected = CheckForSeek(PositionTick, NewPositionTick, OutSongPos.Tempo, CurrentMapChain->SongMaps->GetTicksPerQuarterNote());
 	if (SeekDetected)
 	{
-		UE_LOG(LogMusicClock, VeryVerbose, TEXT("Detected seek when calculating song pos (%f -> %f)"), PositionTick, NewPositionTick);
+		UE_LOG(LogMetasoundMusicClockDriver, VeryVerbose, TEXT("Detected seek when calculating song pos (%f -> %f)"), PositionTick, NewPositionTick);
 	}
 	PositionTick = NewPositionTick;
 	return OutSongPos;
@@ -185,7 +187,7 @@ FMidiSongPos FMetasoundMusicClockDriver::CalculateSongPosAtMsForOffsetClock(floa
 	SeekDetected = CheckForSeek(PositionTick, NewPositionTick, OutSongPos.Tempo, CurrentMapChain->SongMaps->GetTicksPerQuarterNote());
 	if (SeekDetected)
 	{
-		UE_LOG(LogMusicClock, VeryVerbose, TEXT("Detected seek when calculating song pos (%f -> %f)"), PositionTick, NewPositionTick);
+		UE_LOG(LogMetasoundMusicClockDriver, VeryVerbose, TEXT("Detected seek when calculating song pos (%f -> %f)"), PositionTick, NewPositionTick);
 	}
 
 	PositionTick = NewPositionTick;
@@ -217,7 +219,7 @@ bool FMetasoundMusicClockDriver::RefreshCurrentSongPos()
 		if (!CurrentGeneratorHandle)
 		{
 			// cursor connection is not pending
-			AttemptToConnectToAudioComponentsMetasound();
+			AttemptToConnectToAudioComponentsMetasound(AudioComponentToWatch.Get(), OnAttachedDelegate, OnDetachedDelegate);
 		}
 	}
 
@@ -232,7 +234,7 @@ bool FMetasoundMusicClockDriver::RefreshCurrentSongPos()
 		else
 		{
 			// Cursor not attached so use wall clock
-			if (!WasEverConnected || ClockComponent->RunPastMusicEnd)
+			if (!WasEverConnected || RunPastMusicEnd)
 			{
 				RefreshCurrentSongPosFromWallClock();
 				return true;
@@ -293,17 +295,59 @@ const ISongMapEvaluator* FMetasoundMusicClockDriver::GetCurrentSongMapEvaluator(
 	{
 		return CurrentMapChain->SongMaps.Get();
 	}
-	return &ClockComponent->DefaultMaps;
+	return &DefaultMaps;
+}
+
+bool FMetasoundMusicClockDriver::LoopedThisFrame(ECalibratedMusicTimebase Timebase) const
+{
+	switch (Timebase)
+	{
+	case ECalibratedMusicTimebase::RawAudioRenderTime:
+	case ECalibratedMusicTimebase::AudioRenderTime:    return AudioRenderLoopDetected;
+	case ECalibratedMusicTimebase::ExperiencedTime:	   return PlayerExperiencedLoopDetected;
+	case ECalibratedMusicTimebase::VideoRenderTime:
+	default:                                           return VideoRenderLoopDetected;
+	}
 }
 
-bool FMetasoundMusicClockDriver::ConnectToAudioComponentsMetasound(UAudioComponent* InAudioComponent, FName MetasoundOuputPinName)
+bool FMetasoundMusicClockDriver::SeekedThisFrame(ECalibratedMusicTimebase Timebase) const
+{
+	switch (Timebase)
+	{
+	case ECalibratedMusicTimebase::RawAudioRenderTime:
+	case ECalibratedMusicTimebase::AudioRenderTime:    return AudioRenderSeekDetected;
+	case ECalibratedMusicTimebase::ExperiencedTime:	   return PlayerExperiencedSeekDetected;
+	case ECalibratedMusicTimebase::VideoRenderTime:
+	default:                                           return VideoRenderSeekDetected;
+	}
+}
+
+
+bool FMetasoundMusicClockDriver::ConnectToAudioComponentsMetasound(
+	UAudioComponent* InAudioComponent,
+	FName MetasoundOuputPinName,
+	UMetasoundGeneratorHandle::FOnAttached::FDelegate OnGeneratorAttachedCallback,
+	UMetasoundGeneratorHandle::FOnDetached::FDelegate OnGeneratorDetachedCallback)
 {
+	OnAttachedDelegate = OnGeneratorDetachedCallback;
+	OnDetachedDelegate = OnGeneratorDetachedCallback;
+	
 	AudioComponentToWatch = InAudioComponent;
 	MetasoundOutputName = MetasoundOuputPinName;
-	return AttemptToConnectToAudioComponentsMetasound();
+	bool Connected = AttemptToConnectToAudioComponentsMetasound(InAudioComponent, OnGeneratorAttachedCallback, OnGeneratorDetachedCallback);
+
+	if(GetState() == EMusicClockState::Running)
+	{
+		OnStart();
+	}
+	
+	return Connected;
 }
 
-bool FMetasoundMusicClockDriver::AttemptToConnectToAudioComponentsMetasound()
+bool FMetasoundMusicClockDriver::AttemptToConnectToAudioComponentsMetasound(
+	UAudioComponent* InAudioComponent,
+	UMetasoundGeneratorHandle::FOnAttached::FDelegate OnGeneratorAttachedCallback,
+	UMetasoundGeneratorHandle::FOnDetached::FDelegate OnGeneratorDetachedCallback)
 {
 	using namespace HarmonixMetasound::Analysis;
 
@@ -333,8 +377,16 @@ bool FMetasoundMusicClockDriver::AttemptToConnectToAudioComponentsMetasound()
 				MidiSongPosAnalyzerAddress));
 	}
 
-	GeneratorAttachedCallbackHandle = CurrentGeneratorHandle->OnGeneratorHandleAttached.AddLambda([this](){OnGeneratorAttached();});
-	GeneratorDetachedCallbackHandle = CurrentGeneratorHandle->OnGeneratorHandleDetached.AddLambda([this](){OnGeneratorDetached();});
+	GeneratorAttachedCallbackHandle = CurrentGeneratorHandle->OnGeneratorHandleAttached.AddLambda([this, OnGeneratorAttachedCallback]()
+	{
+		OnGeneratorAttached();
+		OnGeneratorAttachedCallback.ExecuteIfBound();
+	});
+	GeneratorDetachedCallbackHandle = CurrentGeneratorHandle->OnGeneratorHandleDetached.AddLambda([this, OnGeneratorDetachedCallback]()
+	{
+		OnGeneratorDetached();
+		OnGeneratorDetachedCallback.ExecuteIfBound();
+	});
 	GeneratorIOUpdatedCallbackHandle = CurrentGeneratorHandle->OnIOUpdatedWithChanges.AddLambda([this](const TArray<Metasound::FVertexInterfaceChange>& VertexInterfaceChanges){OnGeneratorIOUpdatedWithChanges(VertexInterfaceChanges);});
 	UMetasoundGeneratorHandle::FOnSetGraph::FDelegate OnSetGraph;
 	OnSetGraph.BindLambda([this](){OnGraphSet();});
@@ -368,7 +420,6 @@ void FMetasoundMusicClockDriver::OnGeneratorAttached()
 	SmoothedAudioRenderClockHistoryCursor = ClockHistory->CreateReadCursor();
 	SmoothedPlayerExperienceClockHistoryCursor = ClockHistory->CreateReadCursor();
 	SmoothedVideoRenderClockHistoryCursor = ClockHistory->CreateReadCursor();
-	ClockComponent->MusicClockConnectedEvent.Broadcast();
 }
 
 void FMetasoundMusicClockDriver::OnGraphSet()
@@ -400,42 +451,37 @@ void FMetasoundMusicClockDriver::OnGeneratorIOUpdatedWithChanges(const TArray<Me
 void FMetasoundMusicClockDriver::OnGeneratorDetached()
 {
 	using namespace HarmonixMetasound::Analysis;
-	check(ClockComponent);
-	if (ClockComponent->GetState() != EMusicClockState::Stopped)
+	if (GetState() != EMusicClockState::Stopped)
 	{
 		if (ClockHistory && CurrentMapChain && CurrentMapChain->SongMaps)
 		{
-			ClockComponent->DefaultMaps.Copy(*(CurrentMapChain->SongMaps), 0, LastTickSeen);
+			DefaultMaps.Copy(*(CurrentMapChain->SongMaps), 0, LastTickSeen);
 		}
-		SongPosOffsetMs = ClockComponent->CurrentSmoothedAudioRenderSongPos.SecondsIncludingCountIn * 1000.0f;
+		SongPosOffsetMs = CurrentSmoothedAudioRenderSongPos.SecondsIncludingCountIn * 1000.0f;
 		FreeRunStartTimeSecs = GetWallClockTime();
 	}
 	SetClockHistory(nullptr);
 	SmoothedAudioRenderClockHistoryCursor = HarmonixMetasound::Analysis::FMidiClockSongPositionHistory::FReadCursor();
 	SmoothedPlayerExperienceClockHistoryCursor = HarmonixMetasound::Analysis::FMidiClockSongPositionHistory::FReadCursor();
 	SmoothedVideoRenderClockHistoryCursor = HarmonixMetasound::Analysis::FMidiClockSongPositionHistory::FReadCursor();
-
-	ClockComponent->MusicClockDisconnectedEvent.Broadcast();
 }
 
 void FMetasoundMusicClockDriver::RefreshCurrentSongPosFromWallClock()
 {
-	check(ClockComponent);
+	bool TempoChanged = CurrentSmoothedAudioRenderSongPos.Tempo != Tempo;
 
-	bool TempoChanged = ClockComponent->CurrentSmoothedAudioRenderSongPos.Tempo != ClockComponent->Tempo;
+	double FreeRunTime = (GetWallClockTime() - FreeRunStartTimeSecs) * CurrentClockAdvanceRate;
 
-	double FreeRunTime = (GetWallClockTime() - FreeRunStartTimeSecs) * ClockComponent->CurrentClockAdvanceRate;
-
-	ClockComponent->CurrentRawAudioRenderSongPos.SetByTime(((float)FreeRunTime * 1000.0) + SongPosOffsetMs, ClockComponent->DefaultMaps);
-	ClockComponent->CurrentSmoothedAudioRenderSongPos = ClockComponent->CurrentRawAudioRenderSongPos;
-	ClockComponent->CurrentPlayerExperiencedSongPos.SetByTime(ClockComponent->CurrentSmoothedAudioRenderSongPos.SecondsIncludingCountIn * 1000.0f - FHarmonixModule::GetMeasuredUserExperienceAndReactionToAudioRenderOffsetMs(), ClockComponent->DefaultMaps);
-	ClockComponent->CurrentVideoRenderSongPos.SetByTime(ClockComponent->CurrentSmoothedAudioRenderSongPos.SecondsIncludingCountIn * 1000.0f - FHarmonixModule::GetMeasuredVideoToAudioRenderOffsetMs(), ClockComponent->DefaultMaps);
+	CurrentRawAudioRenderSongPos.SetByTime(((float)FreeRunTime * 1000.0) + SongPosOffsetMs, DefaultMaps);
+	CurrentSmoothedAudioRenderSongPos = CurrentRawAudioRenderSongPos;
+	CurrentPlayerExperiencedSongPos.SetByTime(CurrentSmoothedAudioRenderSongPos.SecondsIncludingCountIn * 1000.0f - FHarmonixModule::GetMeasuredUserExperienceAndReactionToAudioRenderOffsetMs(), DefaultMaps);
+	CurrentVideoRenderSongPos.SetByTime(CurrentSmoothedAudioRenderSongPos.SecondsIncludingCountIn * 1000.0f - FHarmonixModule::GetMeasuredVideoToAudioRenderOffsetMs(), DefaultMaps);
 
 	if (TempoChanged)
 	{
-		ClockComponent->Tempo = ClockComponent->CurrentSmoothedAudioRenderSongPos.Tempo;
-		ClockComponent->CurrentBeatDurationSec = (60.0f / ClockComponent->Tempo) / ClockComponent->CurrentClockAdvanceRate;
-		ClockComponent->CurrentBarDurationSec = ((ClockComponent->TimeSignatureNum * ClockComponent->CurrentBeatDurationSec) / (ClockComponent->TimeSignatureDenom / 4.0f)) / ClockComponent->CurrentClockAdvanceRate;
+		Tempo = CurrentSmoothedAudioRenderSongPos.Tempo;
+		CurrentBeatDurationSec = (60.0f / Tempo) / CurrentClockAdvanceRate;
+		CurrentBarDurationSec = ((TimeSignatureNum * CurrentBeatDurationSec) / (TimeSignatureDenom / 4.0f)) / CurrentClockAdvanceRate;
 	}
 }
 
@@ -452,26 +498,13 @@ FString FMetasoundMusicClockDriver::HistoryFailureTypeToString(FMetasoundMusicCl
 	return "Unrecognized";
 }
 
-double FMetasoundMusicClockDriver::GetWallClockTime()
-{
-	if (ClockComponent && ClockComponent->GetWorld())
-	{
-		return ClockComponent->GetWorld()->GetTimeSeconds();
-	}
-	if (GWorld)
-	{
-		return GWorld->GetTimeSeconds();
-	}
-	return 0.0;
-}
-
 void FMetasoundMusicClockDriver::RefreshCurrentSongPosFromHistory()
 {
 	using namespace HarmonixMetasound::Analysis;
 
 	check(IsInGameThread());
 
-	if (!bRunning || !ClockComponent || !SmoothedAudioRenderClockHistoryCursor.DataAvailable())
+	if (!bRunning || !SmoothedAudioRenderClockHistoryCursor.DataAvailable())
 	{
 		return;
 	}
@@ -491,7 +524,7 @@ void FMetasoundMusicClockDriver::RefreshCurrentSongPosFromHistory()
 	}
 
 	auto Entry = ClockHistory->Positions.GetEntry(ClockHistory->Positions.GetLastWriteIndex());
-	ClockComponent->CurrentRawAudioRenderSongPos.SetByTick(Entry->Item.UpToTick, *(CurrentMapChain->SongMaps));
+	CurrentRawAudioRenderSongPos.SetByTick(Entry->Item.UpToTick, *(CurrentMapChain->SongMaps));
 	Metasound::FSampleCount LastRenderPosSampleCount = Entry->Item.SampleCount;
 	float SpeedAtRawRenderTime = Entry->Item.CurrentSpeed;
 	LastTickSeen = Entry->Item.UpToTick;
@@ -522,7 +555,7 @@ void FMetasoundMusicClockDriver::RefreshCurrentSongPosFromHistory()
 
 		if (FMath::Abs(ErrorTracker.Min()) > MetasoundMusicClockDriver::MaxErrorSecondsBeforeJump)
 		{
-			UE_LOG(LogMusicClock, Verbose, TEXT("======== MASSIVE ERROR (%f) - SEEKING ==========="), (float)Error);
+			UE_LOG(LogMetasoundMusicClockDriver, Verbose, TEXT("======== MASSIVE ERROR (%f) - SEEKING ==========="), (float)Error);
 			RenderStartSampleCount = LastRenderPosSampleCount;
 			RenderStartWallClockTimeSeconds = GetWallClockTime() - (double)RenderStartSampleCount / (double)ClockHistory->SampleRate;
 			ExpectedRenderedSeconds = RenderedSeconds;
@@ -540,14 +573,14 @@ void FMetasoundMusicClockDriver::RefreshCurrentSongPosFromHistory()
 		SyncSpeed = FMath::Clamp(SyncSpeed, MetasoundMusicClockDriver::SlowestCorrectionSpeed, MetasoundMusicClockDriver::FastestCorrectionSpeed);
 	}
 
-	//UE_LOG(LogMusicClock, Log, TEXT("Local Minimum Error: %f Average: %f Speed: %f"), (float)ErrorTracker.Min(), (float)ErrorTracker.Average(), (float)SyncSpeed);
+	//UE_LOG(LogMetasoundMusicClockDriver, Log, TEXT("Local Minimum Error: %f Average: %f Speed: %f"), (float)ErrorTracker.Min(), (float)ErrorTracker.Average(), (float)SyncSpeed);
 
 	Metasound::FSampleCount ExpectedRenderPosSampleCount = (Metasound::FSampleCount)(ExpectedRenderedSeconds * (double)ClockHistory->SampleRate);
 
 	// first smoothed render time. This is closest to the actual render time. If we catch up to the render time
 	// it means we are rendering in such large blocks that we need to push up our "look behind" number for smoothing.
 	float SpeedAtSmoothRenderedTime = 1.0f;
-	float SmoothedTick = ClockComponent->GetTicksIncludingCountIn(ECalibratedMusicTimebase::AudioRenderTime);
+	float SmoothedTick = GetCurrentSongMapEvaluator()->MsToTick(CurrentSmoothedAudioRenderSongPos.SecondsIncludingCountIn * 1000.0f);
 	float SmoothedTempoMapTick = AudioRenderState.TempoMapTick;
 	EHistoryFailureType Result = CalculateSmoothedTick(ExpectedRenderPosSampleCount, LastRenderPosSampleCount,
 		SmoothedTick, SmoothedTempoMapTick, SpeedAtSmoothRenderedTime, SmoothedAudioRenderClockHistoryCursor, RenderSmoothingLagSeconds);
@@ -558,16 +591,16 @@ void FMetasoundMusicClockDriver::RefreshCurrentSongPosFromHistory()
 			if (RenderSmoothingLagSeconds < 0.250f)
 			{
 				RenderSmoothingLagSeconds += 0.005f;
-				UE_LOG(LogMusicClock, Verbose, TEXT("(%d) Smoothed Render Time too close to actual render time. Bumping up smoothing lag! (%f)"), LastRenderPosSampleCount, RenderSmoothingLagSeconds);
+				UE_LOG(LogMetasoundMusicClockDriver, Verbose, TEXT("(%d) Smoothed Render Time too close to actual render time. Bumping up smoothing lag! (%f)"), LastRenderPosSampleCount, RenderSmoothingLagSeconds);
 			}
 			else
 			{
-				UE_LOG(LogMusicClock, Verbose, TEXT("(%d) Smoothed Render Time too close to actual render time. ALREADY MAX SMOOTHING LAG! (%f)"), LastRenderPosSampleCount, RenderSmoothingLagSeconds);
+				UE_LOG(LogMetasoundMusicClockDriver, Verbose, TEXT("(%d) Smoothed Render Time too close to actual render time. ALREADY MAX SMOOTHING LAG! (%f)"), LastRenderPosSampleCount, RenderSmoothingLagSeconds);
 			}
 		}
 		else
 		{
-			UE_LOG(LogMusicClock, Verbose, TEXT("(%d) Smoothed Render Time too close to actual render time. WAITING!"), LastRenderPosSampleCount);
+			UE_LOG(LogMetasoundMusicClockDriver, Verbose, TEXT("(%d) Smoothed Render Time too close to actual render time. WAITING!"), LastRenderPosSampleCount);
 		}
 		return;
 	}
@@ -584,15 +617,15 @@ void FMetasoundMusicClockDriver::RefreshCurrentSongPosFromHistory()
 		UpdateCurrentTicksForLoopingOrMonotonicClock(SmoothedTick, SmoothedTempoMapTick);
 	}
 
-	ClockComponent->TimeSignatureNum = ClockComponent->CurrentSmoothedAudioRenderSongPos.TimeSigNumerator;
-	ClockComponent->TimeSignatureDenom = ClockComponent->CurrentSmoothedAudioRenderSongPos.TimeSigDenominator;
+	TimeSignatureNum = CurrentSmoothedAudioRenderSongPos.TimeSigNumerator;
+	TimeSignatureDenom = CurrentSmoothedAudioRenderSongPos.TimeSigDenominator;
 
-	if (ClockComponent->Tempo != ClockComponent->CurrentSmoothedAudioRenderSongPos.Tempo || ClockComponent->CurrentClockAdvanceRate != SpeedAtRawRenderTime)
+	if (Tempo != CurrentSmoothedAudioRenderSongPos.Tempo || CurrentClockAdvanceRate != SpeedAtRawRenderTime)
 	{
-		ClockComponent->CurrentClockAdvanceRate = SpeedAtRawRenderTime;
-		ClockComponent->Tempo = ClockComponent->CurrentSmoothedAudioRenderSongPos.Tempo;
-		ClockComponent->CurrentBeatDurationSec = (60.0f / ClockComponent->Tempo) / ClockComponent->CurrentClockAdvanceRate;
-		ClockComponent->CurrentBarDurationSec = ((ClockComponent->TimeSignatureNum * ClockComponent->CurrentBeatDurationSec) / (ClockComponent->TimeSignatureDenom / 4.0f)) / ClockComponent->CurrentClockAdvanceRate;
+		CurrentClockAdvanceRate = SpeedAtRawRenderTime;
+		Tempo = CurrentSmoothedAudioRenderSongPos.Tempo;
+		CurrentBeatDurationSec = (60.0f / Tempo) / CurrentClockAdvanceRate;
+		CurrentBarDurationSec = ((TimeSignatureNum * CurrentBeatDurationSec) / (TimeSignatureDenom / 4.0f)) / CurrentClockAdvanceRate;
 	}
 }
 
@@ -615,10 +648,10 @@ void FMetasoundMusicClockDriver::UpdateCurrentTicksForOffsetClock(float Smoothed
 	SmoothedPositionMs += RenderSmoothingLagSeconds * 1000.0f;
 
 	// Calculate the song position AND get the "local tick" for the Audio Render timebase
-	ClockComponent->CurrentSmoothedAudioRenderSongPos = CalculateSongPosAtMsForOffsetClock(SmoothedPositionMs,
+	CurrentSmoothedAudioRenderSongPos = CalculateSongPosAtMsForOffsetClock(SmoothedPositionMs,
 		SmoothedTick - SmoothedTempoMapTick,
 		AudioRenderState.LocalTick,
-		ClockComponent->AudioRenderSeekDetected);
+		AudioRenderSeekDetected);
 	float SmoothedTempoMapMs = CurrentMapChain->SongMaps->TickToMs(SmoothedTempoMapTick);
 	AudioRenderState.TempoMapMs = SmoothedTempoMapMs + (RenderSmoothingLagSeconds * 1000.0f);
 	AudioRenderState.TempoMapTick = CurrentMapChain->SongMaps->MsToTick(AudioRenderState.TempoMapMs);
@@ -629,10 +662,10 @@ void FMetasoundMusicClockDriver::UpdateCurrentTicksForOffsetClock(float Smoothed
 	float Ms = SmoothedPositionMs;
 	Ms -= LookBehindLagMs;
 	// Calculate the song position AND get the "local tick" for the Player Experience timebase
-	ClockComponent->CurrentPlayerExperiencedSongPos = CalculateSongPosAtMsForOffsetClock(Ms,
+	CurrentPlayerExperiencedSongPos = CalculateSongPosAtMsForOffsetClock(Ms,
 		SmoothedTick - SmoothedTempoMapTick,
 		PlayerExperienceState.LocalTick,
-		ClockComponent->PlayerExperiencedSeekDetected);
+		PlayerExperiencedSeekDetected);
 	PlayerExperienceState.TempoMapMs = AudioRenderState.TempoMapMs - LookBehindLagMs;
 	PlayerExperienceState.TempoMapTick = CurrentMapChain->SongMaps->MsToTick(PlayerExperienceState.TempoMapMs);
 
@@ -642,10 +675,10 @@ void FMetasoundMusicClockDriver::UpdateCurrentTicksForOffsetClock(float Smoothed
 	Ms = SmoothedPositionMs;
 	Ms -= LookBehindLagMs;
 	// Calculate the song position AND get the "local tick" for the Video Render timebase
-	ClockComponent->CurrentVideoRenderSongPos = CalculateSongPosAtMsForOffsetClock(Ms,
+	CurrentVideoRenderSongPos = CalculateSongPosAtMsForOffsetClock(Ms,
 		SmoothedTick - SmoothedTempoMapTick,
 		VideoRenderState.LocalTick,
-		ClockComponent->VideoRenderSeekDetected);
+		VideoRenderSeekDetected);
 	VideoRenderState.TempoMapMs = AudioRenderState.TempoMapMs - LookBehindLagMs;
 	VideoRenderState.TempoMapTick = CurrentMapChain->SongMaps->MsToTick(VideoRenderState.TempoMapMs);
 }
@@ -659,30 +692,30 @@ void FMetasoundMusicClockDriver::UpdateCurrentTicksForLoopingOrMonotonicClock(fl
 
 	// Calculate the song position AND get the "local tick" for the Audio Render timebase
 	AudioRenderState.TempoMapTick = CurrentMapChain->SongMaps->MsToTick(AudioRenderState.TempoMapMs);
-	ClockComponent->CurrentSmoothedAudioRenderSongPos = CalculateSongPosAtMsForLoopingOrMonotonicClock(AudioRenderState.TempoMapMs,
+	CurrentSmoothedAudioRenderSongPos = CalculateSongPosAtMsForLoopingOrMonotonicClock(AudioRenderState.TempoMapMs,
 		AudioRenderState.LocalTick,
-		ClockComponent->AudioRenderSeekDetected,
-		ClockComponent->AudioRenderLoopDetected);
+		AudioRenderSeekDetected,
+		AudioRenderLoopDetected);
 
 	// Now the time the user should actually be "experiencing" (ie "hearing") can be calculated as an offset
 	// from the smooth audio rendering time...
 	float LookBehindLagMs = FHarmonixModule::Get().GetMeasuredUserExperienceAndReactionToAudioRenderOffsetMs();
 	PlayerExperienceState.TempoMapMs = AudioRenderState.TempoMapMs - LookBehindLagMs;
 	PlayerExperienceState.TempoMapTick = CurrentMapChain->SongMaps->MsToTick(PlayerExperienceState.TempoMapMs);
-	ClockComponent->CurrentPlayerExperiencedSongPos = CalculateSongPosAtMsForLoopingOrMonotonicClock(PlayerExperienceState.TempoMapMs,
+	CurrentPlayerExperiencedSongPos = CalculateSongPosAtMsForLoopingOrMonotonicClock(PlayerExperienceState.TempoMapMs,
 		PlayerExperienceState.LocalTick,
-		ClockComponent->PlayerExperiencedSeekDetected,
-		ClockComponent->PlayerExperiencedLoopDetected);
+		PlayerExperiencedSeekDetected,
+		PlayerExperiencedLoopDetected);
 
 	// Now the time the game should be rendering graphics for can be calculated as an offset
 	// from the smooth audio rendering time...
 	LookBehindLagMs = FHarmonixModule::Get().GetMeasuredVideoToAudioRenderOffsetMs();
 	VideoRenderState.TempoMapMs = AudioRenderState.TempoMapMs - LookBehindLagMs;
 	VideoRenderState.TempoMapTick = CurrentMapChain->SongMaps->MsToTick(VideoRenderState.TempoMapMs);
-	ClockComponent->CurrentVideoRenderSongPos = CalculateSongPosAtMsForLoopingOrMonotonicClock(VideoRenderState.TempoMapMs,
+	CurrentVideoRenderSongPos = CalculateSongPosAtMsForLoopingOrMonotonicClock(VideoRenderState.TempoMapMs,
 		VideoRenderState.LocalTick,
-		ClockComponent->VideoRenderSeekDetected,
-		ClockComponent->VideoRenderLoopDetected);
+		VideoRenderSeekDetected,
+		VideoRenderLoopDetected);
 }
 
 FMetasoundMusicClockDriver::EHistoryFailureType FMetasoundMusicClockDriver::CalculateSmoothedTick(Metasound::FSampleCount ExpectedRenderPosSampleCount, Metasound::FSampleCount LastRenderPosSampleCount,
diff --git a/Engine/Plugins/Runtime/Harmonix/Source/HarmonixMetasound/Private/HarmonixMetasound/Components/MetasoundMusicClockDriver.h b/Engine/Plugins/Runtime/Harmonix/Source/HarmonixMetasound/Private/HarmonixMetasound/Components/MetasoundMusicClockDriver.h
--- a/Engine/Plugins/Runtime/Harmonix/Source/HarmonixMetasound/Private/HarmonixMetasound/Components/MetasoundMusicClockDriver.h
+++ b/Engine/Plugins/Runtime/Harmonix/Source/HarmonixMetasound/Private/HarmonixMetasound/Components/MetasoundMusicClockDriver.h
@@ -1,7 +1,7 @@
 // Copyright Epic Games, Inc. All Rights Reserved.
 #pragma once
 #include "Harmonix/LocalMinimumMagnitudeTracker.h"
-#include "HarmonixMetasound/Components/MusicClockComponent.h"
+#include "HarmonixMetasound/Components/MusicClockDriverBase.h"
 #include "HarmonixMetasound/Subsystems/MidiClockUpdateSubsystem.h"
 #include "MetasoundGeneratorHandle.h"
 #include "UObject/StrongObjectPtr.h"
@@ -14,9 +14,8 @@ namespace HarmonixMetasound::Analysis
 struct FMetasoundMusicClockDriver : public FMusicClockDriverBase
 {
 public:
-	FMetasoundMusicClockDriver(UMusicClockComponent* InClock)
-		: FMusicClockDriverBase(InClock)
-	{}
+	FMetasoundMusicClockDriver(UObject* WorldContextObj = nullptr)
+		: FMusicClockDriverBase(WorldContextObj) {}
 
 	virtual bool CalculateSongPosWithOffset(float MsOffset, ECalibratedMusicTimebase Timebase, FMidiSongPos& OutResult) const override;
 
@@ -27,8 +26,15 @@ public:
 	virtual void OnContinue() override;
 	virtual void OnStop() override;
 	virtual const ISongMapEvaluator* GetCurrentSongMapEvaluator() const override;
-
-	bool ConnectToAudioComponentsMetasound(UAudioComponent* InAudioComponent, FName MetasoundOuputPinName = "MIDI Clock");
+	
+	virtual bool LoopedThisFrame(ECalibratedMusicTimebase Timebase) const override;
+	virtual bool SeekedThisFrame(ECalibratedMusicTimebase Timebase) const override;
+	
+	bool ConnectToAudioComponentsMetasound(
+		UAudioComponent* InAudioComponent,
+		FName MetasoundOuputPinName = "MIDI Clock",
+		UMetasoundGeneratorHandle::FOnAttached::FDelegate OnGeneratorAttachedCallback = UMetasoundGeneratorHandle::FOnAttached::FDelegate::CreateLambda([](){}),
+		UMetasoundGeneratorHandle::FOnDetached::FDelegate OnGeneratorDetachedCallback = UMetasoundGeneratorHandle::FOnDetached::FDelegate::CreateLambda([](){}));
 
 protected:
 	void OnGeneratorAttached();
@@ -54,6 +60,7 @@ private:
 	TSharedPtr<const HarmonixMetasound::Analysis::FSongMapChain> CurrentMapChain;
 	
 	bool bRunning = false;
+	bool RunPastMusicEnd = false;
 	double FreeRunStartTimeSecs = 0.0;
 	bool WasEverConnected = false;
 	float SongPosOffsetMs = 0.0f;
@@ -81,10 +88,20 @@ private:
 	FPerTimebaseSmoothedClockState PlayerExperienceState;
 	FPerTimebaseSmoothedClockState VideoRenderState;
 
+	UMetasoundGeneratorHandle::FOnAttached::FDelegate OnAttachedDelegate;
+	UMetasoundGeneratorHandle::FOnDetached::FDelegate OnDetachedDelegate;
+	
 	FDelegateHandle GeneratorAttachedCallbackHandle;
 	FDelegateHandle GeneratorDetachedCallbackHandle;
 	FDelegateHandle GeneratorIOUpdatedCallbackHandle;
 	FDelegateHandle GraphChangedCallbackHandle;
+	
+	bool AudioRenderSeekDetected = false;
+	bool AudioRenderLoopDetected = false;
+	bool PlayerExperiencedSeekDetected = false;
+	bool PlayerExperiencedLoopDetected = false;
+	bool VideoRenderSeekDetected = false;
+	bool VideoRenderLoopDetected = false;
 
 	FMidiSongPos CalculateSongPosAtMsForLoopingOrMonotonicClock(float AbsoluteMs, float& PositionTick, bool& SeekDetected, bool& LoopDetected) const;
 	FMidiSongPos CalculateSongPosAtMsForOffsetClock(float PositionMs, float ClockTickOffsetFromDrivingClock, float& PositionTick, bool& SeekDetected) const;
@@ -93,7 +110,10 @@ private:
 	void UpdateCurrentTicksForOffsetClock(float SmoothedTick, float SmoothedTempoMapTick);
 	void UpdateCurrentTicksForLoopingOrMonotonicClock(float SmoothedTick, float SmoothedTempoMapTick);
 
-	bool AttemptToConnectToAudioComponentsMetasound();
+	bool AttemptToConnectToAudioComponentsMetasound(
+		UAudioComponent* InAudioComponent,
+		UMetasoundGeneratorHandle::FOnAttached::FDelegate OnGeneratorAttachedCallback,
+		UMetasoundGeneratorHandle::FOnDetached::FDelegate OnGeneratorDetachedCallback);
 	void DetachAllCallbacks();
 	void RefreshCurrentSongPosFromWallClock();
 	void RefreshCurrentSongPosFromHistory();
@@ -114,6 +134,4 @@ private:
 		float LookBehindSeconds);
 
 	static FString HistoryFailureTypeToString(FMetasoundMusicClockDriver::EHistoryFailureType Error);
-
-	double GetWallClockTime();
 };
diff --git a/Engine/Plugins/Runtime/Harmonix/Source/HarmonixMetasound/Private/HarmonixMetasound/Components/MusicClockComponent.cpp b/Engine/Plugins/Runtime/Harmonix/Source/HarmonixMetasound/Private/HarmonixMetasound/Components/MusicClockComponent.cpp
--- a/Engine/Plugins/Runtime/Harmonix/Source/HarmonixMetasound/Private/HarmonixMetasound/Components/MusicClockComponent.cpp
+++ b/Engine/Plugins/Runtime/Harmonix/Source/HarmonixMetasound/Private/HarmonixMetasound/Components/MusicClockComponent.cpp
@@ -3,7 +3,6 @@
 #include "HarmonixMetasound/Components/MetasoundMusicClockDriver.h"
 #include "HarmonixMetasound/Components/WallClockMusicClockDriver.h"
 #include "HarmonixMetasound/Subsystems/MidiClockUpdateSubsystem.h"
-#include "Harmonix.h"
 #include "HarmonixMidi/MusicTimeSpan.h"
 #include "Components/AudioComponent.h"
 #include "Templates/SharedPointer.h"
@@ -72,61 +71,9 @@ void UMusicClockComponent::ConnectToWallClockForMidi(UMidiFile* InTempoMap)
 	ConnectToWallClock();
 }
 
-// TODO: Cleanup task - UE-205069 - If we find we are able to use the new MidiClock/MusicClockComponent
-// ticking methods, this function should be deleted and all of the call sites cleaned up... 
-// as only the non-const "Ensure" function will be required.
-void UMusicClockComponent::EnsureClockIsValidForGameFrame() const
-{
-	if (MidiClockUpdateSubsystem::UpdateMethod != MidiClockUpdateSubsystem::EUpdateMethod::EngineTickableObjectAndTickComponent)
-	{
-		return;
-	}
-
-	//	Not for use outside the game thread.
-	if (ensureMsgf(
-		IsInGameThread(),
-		TEXT("%hs called from non-game thread.  This is not supported!"), __FUNCTION__) == false)
-	{
-		return;
-	}
-
-	if (GFrameCounter == LastUpdateFrame)
-	{
-		return;
-	}
-
-	//	Run the actual clock update.
-	if (State == EMusicClockState::Running && ClockDriver)
-	{
-		// NOTE: This is a little naughty here. Even though this function is const, this next call
-		// to the clock driver is non-const, AND it can reach back into this UMusicClockComponent and
-		// mutate some current state. It is currently the best approach we have. Future refactoring
-		// might eliminate this anomaly. Bottom line is... Callers to "outer functions" that call this
-		// function can/should/need-to call into UMusicClockComponent through a const reference, and
-		// that is reasonable. But sometimes we have to update our internal state before returning
-		// from those functions. All of those state changes happen as a result of this call to the
-		// current ClockDriver.
-		ClockDriver->EnsureClockIsValidForGameFrame();
-	}
-}
-
 void UMusicClockComponent::EnsureClockIsValidForGameFrameFromSubsystem()
 {
-	//	Not for use outside the game thread.
-	if (ensureMsgf(
-		IsInGameThread(),
-		TEXT("%hs called from non-game thread.  This is not supported!"), __FUNCTION__) == false)
-	{
-		return;
-	}
-
-	if (GFrameCounter == LastUpdateFrame)
-	{
-		return;
-	}
-
-	//	Run the actual clock update.
-	if (State == EMusicClockState::Running && ClockDriver)
+	if (ClockDriver)
 	{
 		ClockDriver->EnsureClockIsValidForGameFrame();
 	}
@@ -153,12 +100,12 @@ bool UMusicClockComponent::ConnectToMetasound()
 	}
 	DisconnectFromClockDriver();
 	TSharedPtr<FMetasoundMusicClockDriver> MetasoundClockDriver = MakeShared<FMetasoundMusicClockDriver>(this);
-	bool Connected = MetasoundClockDriver->ConnectToAudioComponentsMetasound(MetasoundsAudioComponent, MetasoundOutputName);
+	bool Connected = MetasoundClockDriver->ConnectToAudioComponentsMetasound(
+		MetasoundsAudioComponent,
+		MetasoundOutputName,
+		UMetasoundGeneratorHandle::FOnAttached::FDelegate::CreateLambda([this](){ MusicClockConnectedEvent.Broadcast(); }),
+		UMetasoundGeneratorHandle::FOnDetached::FDelegate::CreateLambda([this](){ MusicClockDisconnectedEvent.Broadcast(); }));
 	ClockDriver = MoveTemp(MetasoundClockDriver);
-	if (State == EMusicClockState::Running)
-	{
-		ClockDriver->OnStart();
-	}
 	return Connected;
 }
 
@@ -167,7 +114,7 @@ void UMusicClockComponent::ConnectToWallClock()
 	// we don't 'check' that the driver mode is wall clock here because if the driver mode is metasound and we can't 
 	// connect for some reason we will fall back to this clock driver!
 	DisconnectFromClockDriver();
-	ClockDriver = MakeShared<FWallClockMusicClockDriver>(this, TempoMap);
+	ClockDriver = MakeShared<FWallClockMusicClockDriver>(TempoMap, this);
 }
 
 void UMusicClockComponent::DisconnectFromClockDriver()
@@ -177,12 +124,6 @@ void UMusicClockComponent::DisconnectFromClockDriver()
 		ClockDriver->Disconnect();
 		ClockDriver = nullptr;
 	}
-	AudioRenderSeekDetected = false;
-	AudioRenderLoopDetected = false;
-	PlayerExperiencedSeekDetected = false;
-	PlayerExperiencedLoopDetected = false;
-	VideoRenderSeekDetected = false;
-	VideoRenderLoopDetected = false;
 }
 
 
@@ -191,8 +132,6 @@ FMidiSongPos UMusicClockComponent::CalculateSongPosWithOffset(float MsOffset, EC
 	FMidiSongPos Result;
 	if (ClockDriver != nullptr)
 	{
-		EnsureClockIsValidForGameFrame();
-
 		if (ClockDriver->CalculateSongPosWithOffset(MsOffset, Timebase, Result))
 		{
 			return Result;
@@ -204,6 +143,12 @@ FMidiSongPos UMusicClockComponent::CalculateSongPosWithOffset(float MsOffset, EC
 	return Result;
 }
 
+const FMidiSongPos& UMusicClockComponent::GetRawUnsmoothedAudioRenderPos() const
+{
+	static const FMidiSongPos PosZero;
+	return ClockDriver ? ClockDriver->CurrentRawAudioRenderSongPos : PosZero;
+}
+
 void UMusicClockComponent::BeginPlay()
 {
 	BeginPlayInternal();
@@ -239,9 +184,44 @@ void UMusicClockComponent::TickComponent(float DeltaTime, enum ELevelTick TickTy
 
 void UMusicClockComponent::TickComponentInternal()
 {
-	EnsureClockIsValidForGameFrame();
-	BroadcastSongPosChanges();
-	BroadcastSeekLoopDetections();
+	if (ClockDriver)
+	{
+		if (ClockDriver->GetState() == EMusicClockState::Running)
+		{
+			BroadcastSongPosChanges();
+			BroadcastSeekLoopDetections();	
+		}
+	}
+}
+
+float UMusicClockComponent::GetTempo() const
+{
+	return ClockDriver ? ClockDriver->Tempo : Tempo;
+}
+
+int UMusicClockComponent::GetTimeSignatureNumerator() const
+{
+	return ClockDriver ? ClockDriver->TimeSignatureNum : TimeSignatureNum;
+}
+
+int UMusicClockComponent::GetTimeSignatureDenominator() const
+{
+	return ClockDriver ? ClockDriver->TimeSignatureDenom : TimeSignatureDenom;
+}
+
+float UMusicClockComponent::GetCurrentBarDurationSeconds() const
+{
+	return ClockDriver ? ClockDriver->CurrentBarDurationSec : CurrentBarDurationSec;
+}
+
+float UMusicClockComponent::GetCurrentBeatDurationSeconds() const
+{
+	return ClockDriver ? ClockDriver->CurrentBeatDurationSec : CurrentBeatDurationSec;
+}
+
+float UMusicClockComponent::GetCurrentClockAdvanceRate() const
+{
+	return ClockDriver ? ClockDriver->CurrentClockAdvanceRate : CurrentClockAdvanceRate;
 }
 
 void UMusicClockComponent::Start()
@@ -251,63 +231,63 @@ void UMusicClockComponent::Start()
 	{
 		CreateClockDriver();
 	}
-	ClockDriver->OnStart();
+	ClockDriver->Start();
 	LastBroadcastBeat = -1;
 	LastBroadcastBar  = -1;
-	CurrentSmoothedAudioRenderSongPos.Reset();
-	CurrentVideoRenderSongPos.Reset();
-	CurrentPlayerExperiencedSongPos.Reset();
-	CurrentRawAudioRenderSongPos.Reset();
-	State = EMusicClockState::Running;
-	PlayStateEvent.Broadcast(State);
+	PlayStateEvent.Broadcast(EMusicClockState::Running);
 }
 
 void UMusicClockComponent::Pause()
 {
-	if (State != EMusicClockState::Running)
+	if (!ClockDriver)
 	{
 		return;
 	}
 
-	if (ClockDriver)
+	if (ClockDriver->GetState() != EMusicClockState::Running)
 	{
-		ClockDriver->OnPause();
+		return;
 	}
-
-	State = EMusicClockState::Paused;
-	PlayStateEvent.Broadcast(State);
+	ClockDriver->Pause();
+	PlayStateEvent.Broadcast(EMusicClockState::Paused);
 }
 
 void UMusicClockComponent::Continue()
 {
-	if (State != EMusicClockState::Paused)
+	if (!ClockDriver)
 	{
 		return;
 	}
-
-	if (ClockDriver)
+	
+	if (ClockDriver->GetState() != EMusicClockState::Paused)
 	{
-		ClockDriver->OnContinue();
+		return;
 	}
-
-	State = EMusicClockState::Running;
-	PlayStateEvent.Broadcast(State);
+	ClockDriver->Continue();
+	PlayStateEvent.Broadcast(EMusicClockState::Running);
 }
 
 void UMusicClockComponent::Stop()
 {
-	if (ClockDriver)
+	if (!ClockDriver)
 	{
-		ClockDriver->OnStop();
+		return;
 	}
-	State = EMusicClockState::Stopped;
-	CurrentSmoothedAudioRenderSongPos.Reset();
-	CurrentVideoRenderSongPos.Reset();
-	CurrentPlayerExperiencedSongPos.Reset();
-	CurrentRawAudioRenderSongPos.Reset();
-	PlayStateEvent.Broadcast(State);
+	ClockDriver->Stop();
+	PlayStateEvent.Broadcast(EMusicClockState::Stopped);
 }
 
+EMusicClockState UMusicClockComponent::GetState() const
+{
+	if (!ClockDriver)
+	{
+		return EMusicClockState::Stopped;
+	}
+
+	return ClockDriver->GetState();
+}
+
+
 float UMusicClockComponent::GetSecondsIncludingCountIn(ECalibratedMusicTimebase Timebase) const
 {
 	return GetSongPos(Timebase).SecondsIncludingCountIn;
@@ -401,28 +381,34 @@ float UMusicClockComponent::GetDistanceToClosestBar(ECalibratedMusicTimebase Tim
 
 float UMusicClockComponent::GetDeltaBar(ECalibratedMusicTimebase Timebase) const
 {
-	EnsureClockIsValidForGameFrame();
-	switch (Timebase)
+	if (ClockDriver)
 	{
-	case ECalibratedMusicTimebase::RawAudioRenderTime: return RawAudioRenderDeltaBarF;
-	case ECalibratedMusicTimebase::AudioRenderTime:    return AudioRenderDeltaBarF;
-	case ECalibratedMusicTimebase::ExperiencedTime:	   return PlayerExperienceDeltaBarF;
-	case ECalibratedMusicTimebase::VideoRenderTime:
-	default:                                           return VideoRenderDeltaBarF;
+		switch (Timebase)
+		{
+		case ECalibratedMusicTimebase::RawAudioRenderTime: return ClockDriver->RawAudioRenderDeltaBarF;
+		case ECalibratedMusicTimebase::AudioRenderTime:    return ClockDriver->AudioRenderDeltaBarF;
+		case ECalibratedMusicTimebase::ExperiencedTime:	   return ClockDriver->PlayerExperienceDeltaBarF;
+		case ECalibratedMusicTimebase::VideoRenderTime:
+		default:                                           return ClockDriver->VideoRenderDeltaBarF;
+		}
 	}
+	return 0.0f;
 }
 
 float UMusicClockComponent::GetDeltaBeat(ECalibratedMusicTimebase Timebase) const
 {
-	EnsureClockIsValidForGameFrame();
-	switch (Timebase)
+	if (ClockDriver)
 	{
-	case ECalibratedMusicTimebase::RawAudioRenderTime: return RawAudioRenderDeltaBeatF;
-	case ECalibratedMusicTimebase::AudioRenderTime:    return AudioRenderDeltaBeatF;
-	case ECalibratedMusicTimebase::ExperiencedTime:	   return PlayerExperienceDeltaBeatF;
-	case ECalibratedMusicTimebase::VideoRenderTime:
-	default:                                           return VideoRenderDeltaBeatF;
+		switch (Timebase)
+		{
+		case ECalibratedMusicTimebase::RawAudioRenderTime: return ClockDriver->RawAudioRenderDeltaBeatF;
+		case ECalibratedMusicTimebase::AudioRenderTime:    return ClockDriver->AudioRenderDeltaBeatF;
+		case ECalibratedMusicTimebase::ExperiencedTime:	   return ClockDriver->PlayerExperienceDeltaBeatF;
+		case ECalibratedMusicTimebase::VideoRenderTime:
+		default:                                           return ClockDriver->VideoRenderDeltaBeatF;
+		}
 	}
+	return 0.0f;
 }
 
 const TArray<FSongSection>& UMusicClockComponent::GetSongSections() const
@@ -523,28 +509,20 @@ float UMusicClockComponent::GetSongRemainingMs(ECalibratedMusicTimebase Timebase
 
 bool UMusicClockComponent::SeekedThisFrame(ECalibratedMusicTimebase Timebase) const
 {
-	EnsureClockIsValidForGameFrame();
-	switch (Timebase)
+	if (ClockDriver)
 	{
-	case ECalibratedMusicTimebase::RawAudioRenderTime:
-	case ECalibratedMusicTimebase::AudioRenderTime:    return AudioRenderSeekDetected;
-	case ECalibratedMusicTimebase::ExperiencedTime:	   return PlayerExperiencedSeekDetected;
-	case ECalibratedMusicTimebase::VideoRenderTime:
-	default:                                           return VideoRenderSeekDetected;
+		return ClockDriver->SeekedThisFrame(Timebase);
 	}
+	return false;
 }
 
 bool UMusicClockComponent::LoopedThisFrame(ECalibratedMusicTimebase Timebase) const
 {
-	EnsureClockIsValidForGameFrame();
-	switch (Timebase)
+	if (ClockDriver)
 	{
-	case ECalibratedMusicTimebase::RawAudioRenderTime:
-	case ECalibratedMusicTimebase::AudioRenderTime:    return AudioRenderLoopDetected;
-	case ECalibratedMusicTimebase::ExperiencedTime:	   return PlayerExperiencedLoopDetected;
-	case ECalibratedMusicTimebase::VideoRenderTime:
-	default:                                           return VideoRenderLoopDetected;
+		return ClockDriver->LoopedThisFrame(Timebase);
 	}
+	return false;
 }
 
 
@@ -557,13 +535,11 @@ const ISongMapEvaluator& UMusicClockComponent::GetSongMaps() const
 
 const FMidiSongPos& UMusicClockComponent::GetSongPos(ECalibratedMusicTimebase Timebase) const
 {
-	EnsureClockIsValidForGameFrame();
 	return GetCurrentSongPosInternal(Timebase);
 }
 
 const FMidiSongPos& UMusicClockComponent::GetPreviousSongPos(ECalibratedMusicTimebase Timebase) const
 {
-	EnsureClockIsValidForGameFrame();
 	return GetPreviousSongPosInternal(Timebase);
 }
 
@@ -604,12 +580,7 @@ FMidiSongPos UMusicClockComponent::GetCurrentRawAudioRenderSongPos() const
 
 float UMusicClockComponent::MeasureSpanProgress(const FMusicalTimeSpan& Span, ECalibratedMusicTimebase Timebase) const
 {
-	EnsureClockIsValidForGameFrame();
-	if (const ISongMapEvaluator* Maps = ClockDriver ? ClockDriver->GetCurrentSongMapEvaluator() : &DefaultMaps)
-	{
-		return Span.CalcPositionInSpan(GetCurrentSongPosInternal(Timebase), *Maps);
-	}
-	return 0.0f;
+	return Span.CalcPositionInSpan(GetCurrentSongPosInternal(Timebase), GetSongMaps());
 }
 
 void UMusicClockComponent::BroadcastSongPosChanges()
@@ -640,32 +611,59 @@ void UMusicClockComponent::BroadcastSongPosChanges()
 	}
 }
 
-void UMusicClockComponent::BroadcastSeekLoopDetections() const
+const FMusicTimeDiscontinuityEvent* UMusicClockComponent::GetMusicTimeDiscontinuityEventEventInternal(ECalibratedMusicTimebase Timebase) const
 {
-	if(AudioRenderSeekDetected)
-	{
-		AudioRenderMusicTimeDiscontinuityEvent.Broadcast(EMusicTimeDiscontinuityType::Seek, PrevAudioRenderSongPos, CurrentSmoothedAudioRenderSongPos);
-	}
-	if(AudioRenderLoopDetected)
-	{
-		AudioRenderMusicTimeDiscontinuityEvent.Broadcast(EMusicTimeDiscontinuityType::Loop, PrevAudioRenderSongPos, CurrentSmoothedAudioRenderSongPos);
-	}
-	if(PlayerExperiencedSeekDetected)
-	{
-		PlayerExperienceMusicTimeDiscontinuityEvent.Broadcast(EMusicTimeDiscontinuityType::Seek, PrevPlayerExperiencedSongPos, CurrentPlayerExperiencedSongPos);
-	}
-	if(PlayerExperiencedLoopDetected)
-	{
-		PlayerExperienceMusicTimeDiscontinuityEvent.Broadcast(EMusicTimeDiscontinuityType::Loop, PrevPlayerExperiencedSongPos, CurrentPlayerExperiencedSongPos);
-	}
-	if(VideoRenderSeekDetected)
+	switch (Timebase)
 	{
-		VideoRenderMusicTimeDiscontinuityEvent.Broadcast(EMusicTimeDiscontinuityType::Seek, PrevVideoRenderSongPos, CurrentVideoRenderSongPos);
+	case ECalibratedMusicTimebase::AudioRenderTime:
+		return &AudioRenderMusicTimeDiscontinuityEvent;
+	case ECalibratedMusicTimebase::ExperiencedTime:
+		return &PlayerExperienceMusicTimeDiscontinuityEvent;
+	case ECalibratedMusicTimebase::VideoRenderTime:
+		return &VideoRenderMusicTimeDiscontinuityEvent;
+	default:
+		return nullptr;
 	}
-	if(VideoRenderLoopDetected)
+}
+
+
+void UMusicClockComponent::BroadcastSeekLoopDetections() const
+{
+	auto BroadcastSeekLoopEvents = [this](ECalibratedMusicTimebase Timebase)
 	{
-		VideoRenderMusicTimeDiscontinuityEvent.Broadcast(EMusicTimeDiscontinuityType::Loop, PrevVideoRenderSongPos, CurrentVideoRenderSongPos);
-	}
+		if (!ClockDriver)
+		{
+			return;
+		}
+
+		if (!SeekedThisFrame(Timebase) && !LoopedThisFrame(Timebase))
+		{
+			return;
+		}
+		
+		const FMusicTimeDiscontinuityEvent* Event = GetMusicTimeDiscontinuityEventEventInternal(Timebase);
+		
+		if (Event == nullptr)
+		{
+			return;
+		}
+	
+		const FMidiSongPos& PrevSongPos = GetPreviousSongPosInternal(Timebase);
+		const FMidiSongPos& CurrentSongPos = GetCurrentSongPosInternal(Timebase);
+
+		if (SeekedThisFrame(Timebase))
+		{
+			Event->Broadcast(EMusicTimeDiscontinuityType::Seek, PrevSongPos, CurrentSongPos);
+		}
+		if (LoopedThisFrame(Timebase))
+		{
+			Event->Broadcast(EMusicTimeDiscontinuityType::Loop, PrevSongPos, CurrentSongPos);
+		}
+	};
+
+	BroadcastSeekLoopEvents(ECalibratedMusicTimebase::AudioRenderTime);
+	BroadcastSeekLoopEvents(ECalibratedMusicTimebase::ExperiencedTime);
+	BroadcastSeekLoopEvents(ECalibratedMusicTimebase::VideoRenderTime);
 }
 
 void UMusicClockComponent::MakeDefaultSongMap()
@@ -678,51 +676,43 @@ void UMusicClockComponent::MakeDefaultSongMap()
 
 const FMidiSongPos& UMusicClockComponent::GetCurrentSongPosInternal(ECalibratedMusicTimebase Timebase) const
 {
-	switch (Timebase)
+	if (ClockDriver)
 	{
-	case ECalibratedMusicTimebase::AudioRenderTime:
-		return CurrentSmoothedAudioRenderSongPos;
-	case ECalibratedMusicTimebase::ExperiencedTime:
-		return CurrentPlayerExperiencedSongPos;
-	case ECalibratedMusicTimebase::VideoRenderTime:
-		return CurrentVideoRenderSongPos;
-	case ECalibratedMusicTimebase::RawAudioRenderTime:
-		return CurrentRawAudioRenderSongPos;
+		return ClockDriver->GetCurrentSongPos(Timebase);
 	}
-	check(false);
-	// minor robustness, return something even if there's an illegal value. 
-	return CurrentVideoRenderSongPos;
+	static const FMidiSongPos PosZero;
+	return PosZero;
 }
 
 const FMidiSongPos& UMusicClockComponent::GetPreviousSongPosInternal(ECalibratedMusicTimebase Timebase) const
 {
-	switch (Timebase)
+	if (ClockDriver)
 	{
-	case ECalibratedMusicTimebase::AudioRenderTime:
-		return PrevAudioRenderSongPos;
-	case ECalibratedMusicTimebase::ExperiencedTime:
-		return PrevPlayerExperiencedSongPos;
-	case ECalibratedMusicTimebase::VideoRenderTime:
-		return PrevVideoRenderSongPos;
-	case ECalibratedMusicTimebase::RawAudioRenderTime:
-		return PrevRawAudioRenderSongPos;
+		return ClockDriver->GetPreviousSongPos(Timebase);
 	}
-	check(false);
-	// minor robustness, return something even if there's an illegal value. 
-	return PrevVideoRenderSongPos;
+
+	static const FMidiSongPos PosZero;
+	return PosZero;
 }
 
 FMusicalTime UMusicClockComponent::GetMusicalTime() const
 {
-	FMusicalTime Result; 
-	if (CurrentVideoRenderSongPos.TimeSigDenominator < 1)
+	FMusicalTime Result;
+	if (!ClockDriver)
+	{
+		return Result;
+	}
+	
+	const FMidiSongPos& SongPos = ClockDriver->GetCurrentSongPos(ECalibratedMusicTimebase::VideoRenderTime);
+	if (SongPos.TimeSigDenominator < 1)
 	{
 		return Result;
 	}
-	Result.Bar = CurrentVideoRenderSongPos.Timestamp.Bar - 1;
-	Result.TicksPerBeat = (MusicalTime::TicksPerQuarterNote * 4 / CurrentVideoRenderSongPos.TimeSigDenominator);
-	Result.TicksPerBar = Result.TicksPerBeat * CurrentVideoRenderSongPos.TimeSigNumerator;
-	Result.TickInBar = static_cast<int32>((CurrentVideoRenderSongPos.Timestamp.Beat - 1.0f) * static_cast<float>(Result.TicksPerBeat));
+
+	Result.Bar = SongPos.Timestamp.Bar - 1;
+	Result.TicksPerBeat = (MusicalTime::TicksPerQuarterNote * 4 / SongPos.TimeSigDenominator);
+	Result.TicksPerBar = Result.TicksPerBeat * SongPos.TimeSigNumerator;
+	Result.TickInBar = static_cast<int32>((SongPos.Timestamp.Beat - 1.0f) * static_cast<float>(Result.TicksPerBeat));
 	return Result;
 }
 
@@ -731,7 +721,7 @@ int32 UMusicClockComponent::GetAbsoluteTickPosition() const
 	// Note: Since this is an override of a function in the music environment system we need to
 	// be sure we are using its standard ticks per quarter note units.
 	const ISongMapEvaluator& MapEvaluator = GetSongMaps();
-	float CurrentTick = MapEvaluator.FractionalBarIncludingCountInToTick(CurrentVideoRenderSongPos.BarsIncludingCountIn);
+	float CurrentTick = MapEvaluator.FractionalBarIncludingCountInToTick(GetCurrentSongPosInternal(ECalibratedMusicTimebase::VideoRenderTime).BarsIncludingCountIn);
 	return FMath::FloorToInt32((CurrentTick / static_cast<float>(MapEvaluator.GetTicksPerQuarterNote())) * static_cast<float>(MusicalTime::TicksPerQuarterNote));
 }
 
@@ -740,7 +730,7 @@ FMusicalTime UMusicClockComponent::GetMusicalTimeWithSourceSpaceOffset(const FMu
 {
 	FMusicalTime Result;
 	const ISongMapEvaluator& MapEvaluator = GetSongMaps();
-	float CurrentTick = MapEvaluator.FractionalBarIncludingCountInToTick(CurrentVideoRenderSongPos.BarsIncludingCountIn);
+	float CurrentTick = MapEvaluator.FractionalBarIncludingCountInToTick(GetCurrentSongPosInternal(ECalibratedMusicTimebase::VideoRenderTime).BarsIncludingCountIn);
 	
 	float OffsetTick = MapEvaluator.FractionalBarIncludingCountInToTick(static_cast<float>(Offset.FractionalBar()));
 	float AdjustedTick = CurrentTick - OffsetTick;
@@ -759,7 +749,7 @@ int32 UMusicClockComponent::GetAbsoluteTickPositionWithSourceSpaceOffset(const F
 	// Note: Since this is an override of a function in the music environment system we need to
 	// be sure we are using its standard ticks per quarter note units.
 	const ISongMapEvaluator& MapEvaluator = GetSongMaps();
-	float CurrentTick = MapEvaluator.FractionalBarIncludingCountInToTick(CurrentVideoRenderSongPos.BarsIncludingCountIn);
+	float CurrentTick = MapEvaluator.FractionalBarIncludingCountInToTick(GetCurrentSongPosInternal(ECalibratedMusicTimebase::VideoRenderTime).BarsIncludingCountIn);
 	float OffsetTick = MapEvaluator.FractionalBarIncludingCountInToTick(static_cast<float>(Offset.FractionalBar()));
 	float AdjustedTick = CurrentTick - OffsetTick;
 	return FMath::FloorToInt32((AdjustedTick / static_cast<float>(MapEvaluator.GetTicksPerQuarterNote())) * static_cast<float>(MusicalTime::TicksPerQuarterNote));
@@ -781,35 +771,3 @@ FMusicalTime UMusicClockComponent::Quantize(const FMusicalTime& MusicalTime, int
 	return Result;
 }
 
-void FMusicClockDriverBase::EnsureClockIsValidForGameFrame()
-{
-	// Here the ClockDriver reaches back up and mutates its owning UMusicClockComponent
-	// to make sure its current state is appropriate to the current musical time. See 
-	// UMusicClockComponent::EnsureClockIsValidForGameFrame for more details as to why 
-	// this is so.
-	ClockComponent->PrevRawAudioRenderSongPos = ClockComponent->CurrentRawAudioRenderSongPos;
-	ClockComponent->PrevAudioRenderSongPos = ClockComponent->CurrentSmoothedAudioRenderSongPos;
-	ClockComponent->PrevPlayerExperiencedSongPos = ClockComponent->CurrentPlayerExperiencedSongPos;
-	ClockComponent->PrevVideoRenderSongPos = ClockComponent->CurrentVideoRenderSongPos;
-
-	ClockComponent->AudioRenderSeekDetected = false;
-	ClockComponent->AudioRenderLoopDetected = false;
-	ClockComponent->PlayerExperiencedSeekDetected = false;
-	ClockComponent->PlayerExperiencedLoopDetected = false;
-	ClockComponent->VideoRenderSeekDetected = false;
-	ClockComponent->VideoRenderLoopDetected = false;
-
-	if (RefreshCurrentSongPos())
-	{
-		ClockComponent->RawAudioRenderDeltaBarF = ClockComponent->CurrentRawAudioRenderSongPos.BarsIncludingCountIn - ClockComponent->PrevRawAudioRenderSongPos.BarsIncludingCountIn;
-		ClockComponent->RawAudioRenderDeltaBeatF = ClockComponent->CurrentRawAudioRenderSongPos.BeatsIncludingCountIn - ClockComponent->PrevRawAudioRenderSongPos.BeatsIncludingCountIn;
-		ClockComponent->AudioRenderDeltaBarF = ClockComponent->CurrentSmoothedAudioRenderSongPos.BarsIncludingCountIn - ClockComponent->PrevAudioRenderSongPos.BarsIncludingCountIn;
-		ClockComponent->AudioRenderDeltaBeatF = ClockComponent->CurrentSmoothedAudioRenderSongPos.BeatsIncludingCountIn - ClockComponent->PrevAudioRenderSongPos.BeatsIncludingCountIn;
-		ClockComponent->PlayerExperienceDeltaBarF = ClockComponent->CurrentPlayerExperiencedSongPos.BarsIncludingCountIn - ClockComponent->PrevPlayerExperiencedSongPos.BarsIncludingCountIn;
-		ClockComponent->PlayerExperienceDeltaBeatF = ClockComponent->CurrentPlayerExperiencedSongPos.BeatsIncludingCountIn - ClockComponent->PrevPlayerExperiencedSongPos.BeatsIncludingCountIn;
-		ClockComponent->VideoRenderDeltaBarF = ClockComponent->CurrentVideoRenderSongPos.BarsIncludingCountIn - ClockComponent->PrevVideoRenderSongPos.BarsIncludingCountIn;
-		ClockComponent->VideoRenderDeltaBeatF = ClockComponent->CurrentVideoRenderSongPos.BeatsIncludingCountIn - ClockComponent->PrevVideoRenderSongPos.BeatsIncludingCountIn;
-		ClockComponent->LastUpdateFrame = GFrameCounter;
-	}
-}
-
diff --git a/Engine/Plugins/Runtime/Harmonix/Source/HarmonixMetasound/Private/HarmonixMetasound/Components/MusicClockDriverBase.cpp b/Engine/Plugins/Runtime/Harmonix/Source/HarmonixMetasound/Private/HarmonixMetasound/Components/MusicClockDriverBase.cpp
new file mode 100644
--- /dev/null
+++ b/Engine/Plugins/Runtime/Harmonix/Source/HarmonixMetasound/Private/HarmonixMetasound/Components/MusicClockDriverBase.cpp
@@ -0,0 +1,163 @@
+﻿// Copyright Epic Games, Inc. All Rights Reserved.
+
+#include "HarmonixMetasound/Components/MusicClockDriverBase.h"
+#include "Engine/World.h"
+
+void FMusicClockDriverBase::ResetDefaultSongMaps()
+{
+	DefaultMaps.EmptyAllMaps();
+	DefaultMaps.Init(Harmonix::Midi::Constants::GTicksPerQuarterNoteInt);
+	DefaultMaps.GetTempoMap().AddTempoInfoPoint(Harmonix::Midi::Constants::BPMToMidiTempo(Tempo), 0);
+	DefaultMaps.GetBarMap().AddTimeSignatureAtBarIncludingCountIn(0, TimeSignatureNum, TimeSignatureDenom);
+}
+
+FMusicClockDriverBase::FMusicClockDriverBase(UObject* WorldContextObject)
+{
+	ResetDefaultSongMaps();
+	ContextObj = TWeakObjectPtr<UObject>(WorldContextObject);
+}
+
+FMusicClockDriverBase::~FMusicClockDriverBase()
+{
+	
+}
+
+double FMusicClockDriverBase::GetWallClockTime() const
+{
+	if (GWorld)
+	{
+		return GWorld->GetTimeSeconds();
+	}
+
+	if (TStrongObjectPtr<UObject> Obj = ContextObj.Pin())
+	{
+		if (ensure(Obj->GetWorld()))
+		{
+			return Obj->GetWorld()->GetTimeSeconds();
+		}
+	}
+	
+	return 0.0;
+}
+
+
+const FMidiSongPos& FMusicClockDriverBase::GetCurrentSongPos(ECalibratedMusicTimebase Timebase) const
+{
+	switch (Timebase)
+	{
+	case ECalibratedMusicTimebase::AudioRenderTime:
+		return CurrentSmoothedAudioRenderSongPos;
+	case ECalibratedMusicTimebase::ExperiencedTime:
+		return CurrentPlayerExperiencedSongPos;
+	case ECalibratedMusicTimebase::VideoRenderTime:
+		return CurrentVideoRenderSongPos;
+	case ECalibratedMusicTimebase::RawAudioRenderTime:
+		return CurrentRawAudioRenderSongPos;
+	}
+	checkNoEntry();
+	// minor robustness, return something even if there's an illegal value. 
+	return CurrentVideoRenderSongPos;
+}
+
+const FMidiSongPos& FMusicClockDriverBase::GetPreviousSongPos(ECalibratedMusicTimebase Timebase) const
+{
+	switch (Timebase)
+	{
+	case ECalibratedMusicTimebase::AudioRenderTime:
+		return PrevAudioRenderSongPos;
+	case ECalibratedMusicTimebase::ExperiencedTime:
+		return PrevPlayerExperiencedSongPos;
+	case ECalibratedMusicTimebase::VideoRenderTime:
+		return PrevVideoRenderSongPos;
+	case ECalibratedMusicTimebase::RawAudioRenderTime:
+		return PrevRawAudioRenderSongPos;
+	}
+	checkNoEntry();
+	
+	// minor robustness, return something even if there's an illegal value. 
+	return CurrentVideoRenderSongPos;
+}
+
+
+
+void FMusicClockDriverBase::Start()
+{
+	ResetDefaultSongMaps();
+	OnStart();
+	CurrentSmoothedAudioRenderSongPos.Reset();
+    CurrentVideoRenderSongPos.Reset();
+    CurrentPlayerExperiencedSongPos.Reset();
+    CurrentRawAudioRenderSongPos.Reset();
+	MusicClockState = EMusicClockState::Running;
+}
+
+void FMusicClockDriverBase::Continue()
+{
+	if (MusicClockState != EMusicClockState::Paused)
+	{
+		return;
+	}
+	
+	OnContinue();
+	MusicClockState = EMusicClockState::Running;
+}
+
+void FMusicClockDriverBase::Pause()
+{
+	if (MusicClockState != EMusicClockState::Running)
+	{
+		return;
+	}
+
+	OnPause();
+	MusicClockState = EMusicClockState::Paused;
+}
+
+void FMusicClockDriverBase::Stop()
+{
+	OnStop();
+	MusicClockState = EMusicClockState::Stopped;
+	CurrentSmoothedAudioRenderSongPos.Reset();
+	CurrentVideoRenderSongPos.Reset();
+	CurrentPlayerExperiencedSongPos.Reset();
+	CurrentRawAudioRenderSongPos.Reset();
+}
+
+
+void FMusicClockDriverBase::EnsureClockIsValidForGameFrame()
+{
+	//	Not for use outside the game thread.
+	if (ensureMsgf(
+		IsInGameThread(),
+		TEXT("%hs called from non-game thread.  This is not supported!"), __FUNCTION__) == false)
+	{
+		return;
+	}
+
+	if (GFrameCounter == LastUpdateFrame)
+	{
+		return;
+	}
+
+	//	Run the actual clock update.
+	if (GetState() == EMusicClockState::Running)
+	{
+		PrevRawAudioRenderSongPos = CurrentRawAudioRenderSongPos;
+		PrevAudioRenderSongPos = CurrentSmoothedAudioRenderSongPos;
+		PrevPlayerExperiencedSongPos = CurrentPlayerExperiencedSongPos;
+		PrevVideoRenderSongPos = CurrentVideoRenderSongPos;
+
+		if (RefreshCurrentSongPos())
+		{
+			RawAudioRenderDeltaBarF = CurrentRawAudioRenderSongPos.BarsIncludingCountIn - PrevRawAudioRenderSongPos.BarsIncludingCountIn;
+			RawAudioRenderDeltaBeatF = CurrentRawAudioRenderSongPos.BeatsIncludingCountIn - PrevRawAudioRenderSongPos.BeatsIncludingCountIn;
+			AudioRenderDeltaBarF = CurrentSmoothedAudioRenderSongPos.BarsIncludingCountIn - PrevAudioRenderSongPos.BarsIncludingCountIn;
+			AudioRenderDeltaBeatF = CurrentSmoothedAudioRenderSongPos.BeatsIncludingCountIn - PrevAudioRenderSongPos.BeatsIncludingCountIn;
+			PlayerExperienceDeltaBarF = CurrentPlayerExperiencedSongPos.BarsIncludingCountIn - PrevPlayerExperiencedSongPos.BarsIncludingCountIn;
+			PlayerExperienceDeltaBeatF = CurrentPlayerExperiencedSongPos.BeatsIncludingCountIn - PrevPlayerExperiencedSongPos.BeatsIncludingCountIn;
+			VideoRenderDeltaBarF = CurrentVideoRenderSongPos.BarsIncludingCountIn - PrevVideoRenderSongPos.BarsIncludingCountIn;
+			VideoRenderDeltaBeatF = CurrentVideoRenderSongPos.BeatsIncludingCountIn - PrevVideoRenderSongPos.BeatsIncludingCountIn;
+			LastUpdateFrame = GFrameCounter;
+		}
+	}
+}
\ No newline at end of file
diff --git a/Engine/Plugins/Runtime/Harmonix/Source/HarmonixMetasound/Private/HarmonixMetasound/Components/WallClockMusicClockDriver.cpp b/Engine/Plugins/Runtime/Harmonix/Source/HarmonixMetasound/Private/HarmonixMetasound/Components/WallClockMusicClockDriver.cpp
--- a/Engine/Plugins/Runtime/Harmonix/Source/HarmonixMetasound/Private/HarmonixMetasound/Components/WallClockMusicClockDriver.cpp
+++ b/Engine/Plugins/Runtime/Harmonix/Source/HarmonixMetasound/Private/HarmonixMetasound/Components/WallClockMusicClockDriver.cpp
@@ -1,7 +1,6 @@
 // Copyright Epic Games, Inc. All Rights Reserved.
 #include "WallClockMusicClockDriver.h"
 #include "Harmonix.h"
-#include "Engine/World.h"
 
 bool FWallClockMusicClockDriver::CalculateSongPosWithOffset(float MsOffset, ECalibratedMusicTimebase Timebase, FMidiSongPos& OutResult) const
 {
@@ -15,14 +14,14 @@ bool FWallClockMusicClockDriver::CalculateSongPosWithOffset(float MsOffset, ECal
 	switch (Timebase)
 	{
 	case ECalibratedMusicTimebase::AudioRenderTime:
-		OutResult.SetByTime((ClockComponent->CurrentSmoothedAudioRenderSongPos.SecondsIncludingCountIn * 1000.0f) + MsOffset, *Maps);
+		OutResult.SetByTime((CurrentSmoothedAudioRenderSongPos.SecondsIncludingCountIn * 1000.0f) + MsOffset, *Maps);
 		break;
 	case ECalibratedMusicTimebase::ExperiencedTime:
-		OutResult.SetByTime((ClockComponent->CurrentPlayerExperiencedSongPos.SecondsIncludingCountIn * 1000.0f) + MsOffset, *Maps);
+		OutResult.SetByTime((CurrentPlayerExperiencedSongPos.SecondsIncludingCountIn * 1000.0f) + MsOffset, *Maps);
 		break;
 	case ECalibratedMusicTimebase::VideoRenderTime:
 	default:
-		OutResult.SetByTime((ClockComponent->CurrentVideoRenderSongPos.SecondsIncludingCountIn * 1000.0f) + MsOffset, *Maps);
+		OutResult.SetByTime((CurrentVideoRenderSongPos.SecondsIncludingCountIn * 1000.0f) + MsOffset, *Maps);
 		break;
 	}
 
@@ -37,25 +36,22 @@ void FWallClockMusicClockDriver::Disconnect()
 bool FWallClockMusicClockDriver::RefreshCurrentSongPos()
 {
 	check(IsInGameThread());
-	check(ClockComponent);
-	check(ClockComponent->GetWorld());
+	bool TempoChanged = CurrentSmoothedAudioRenderSongPos.Tempo != Tempo;
 
-	bool TempoChanged = ClockComponent->CurrentSmoothedAudioRenderSongPos.Tempo != ClockComponent->Tempo;
-
-	double RunTime = ClockComponent->GetWorld()->GetTimeSeconds() - StartTimeSecs;
+	double RunTime = GetWallClockTime() - StartTimeSecs;
 
 	const ISongMapEvaluator* Maps = GetCurrentSongMapEvaluator();
 	check(Maps);
 
-	ClockComponent->CurrentSmoothedAudioRenderSongPos.SetByTime((float)(RunTime * 1000.0), *Maps);
-	ClockComponent->CurrentPlayerExperiencedSongPos.SetByTime(ClockComponent->CurrentSmoothedAudioRenderSongPos.SecondsIncludingCountIn * 1000.0f - FHarmonixModule::GetMeasuredUserExperienceAndReactionToAudioRenderOffsetMs(), *Maps);
-	ClockComponent->CurrentVideoRenderSongPos.SetByTime(ClockComponent->CurrentSmoothedAudioRenderSongPos.SecondsIncludingCountIn * 1000.0f - FHarmonixModule::GetMeasuredVideoToAudioRenderOffsetMs(), *Maps);
+	CurrentSmoothedAudioRenderSongPos.SetByTime((float)(RunTime * 1000.0), *Maps);
+	CurrentPlayerExperiencedSongPos.SetByTime(CurrentSmoothedAudioRenderSongPos.SecondsIncludingCountIn * 1000.0f - FHarmonixModule::GetMeasuredUserExperienceAndReactionToAudioRenderOffsetMs(), *Maps);
+	CurrentVideoRenderSongPos.SetByTime(CurrentSmoothedAudioRenderSongPos.SecondsIncludingCountIn * 1000.0f - FHarmonixModule::GetMeasuredVideoToAudioRenderOffsetMs(), *Maps);
 
 	if (TempoChanged)
 	{
-		ClockComponent->Tempo = ClockComponent->CurrentSmoothedAudioRenderSongPos.Tempo;
-		ClockComponent->CurrentBeatDurationSec = (60.0f / ClockComponent->Tempo) / ClockComponent->CurrentClockAdvanceRate;
-		ClockComponent->CurrentBarDurationSec = ((ClockComponent->TimeSignatureNum * ClockComponent->CurrentBeatDurationSec) / (ClockComponent->TimeSignatureDenom / 4.0f)) / ClockComponent->CurrentClockAdvanceRate;
+		Tempo = CurrentSmoothedAudioRenderSongPos.Tempo;
+		CurrentBeatDurationSec = (60.0f / Tempo) / CurrentClockAdvanceRate;
+		CurrentBarDurationSec = ((TimeSignatureNum * CurrentBeatDurationSec) / (TimeSignatureDenom / 4.0f)) / CurrentClockAdvanceRate;
 	}
 
 	return true;
@@ -65,23 +61,20 @@ bool FWallClockMusicClockDriver::RefreshCurrentSongPos()
 void FWallClockMusicClockDriver::OnStart()
 {
 	check(IsInGameThread());
-	check(ClockComponent);
-	StartTimeSecs = ClockComponent->GetWorld()->GetTimeSeconds();
+	StartTimeSecs = GetWallClockTime();
 	PauseTimeSecs = 0.0;
 }
 
 void FWallClockMusicClockDriver::OnPause()
 {
 	check(IsInGameThread());
-	check(ClockComponent);
-	PauseTimeSecs = ClockComponent->GetWorld()->GetTimeSeconds();
+	PauseTimeSecs = GetWallClockTime();
 }
 
 void FWallClockMusicClockDriver::OnContinue()
 {
 	check(IsInGameThread());
-	check(ClockComponent);
-	double CurrentTime = ClockComponent->GetWorld()->GetTimeSeconds();
+	double CurrentTime = GetWallClockTime();
 	StartTimeSecs += (CurrentTime - PauseTimeSecs);
 	PauseTimeSecs = 0.0;
 	RefreshCurrentSongPos();
@@ -94,6 +87,5 @@ const ISongMapEvaluator* FWallClockMusicClockDriver::GetCurrentSongMapEvaluator(
 	{
 		return TempoMapMidi->GetSongMaps();
 	}
-	check(ClockComponent);
-	return &ClockComponent->DefaultMaps;
+	return &DefaultMaps;
 }
diff --git a/Engine/Plugins/Runtime/Harmonix/Source/HarmonixMetasound/Private/HarmonixMetasound/Components/WallClockMusicClockDriver.h b/Engine/Plugins/Runtime/Harmonix/Source/HarmonixMetasound/Private/HarmonixMetasound/Components/WallClockMusicClockDriver.h
--- a/Engine/Plugins/Runtime/Harmonix/Source/HarmonixMetasound/Private/HarmonixMetasound/Components/WallClockMusicClockDriver.h
+++ b/Engine/Plugins/Runtime/Harmonix/Source/HarmonixMetasound/Private/HarmonixMetasound/Components/WallClockMusicClockDriver.h
@@ -1,14 +1,14 @@
 // Copyright Epic Games, Inc. All Rights Reserved.
 #pragma once
-#include "HarmonixMetasound/Components/MusicClockComponent.h"
+#include "HarmonixMetasound/Components/MusicClockDriverBase.h"
 #include "HarmonixMidi/MidiFile.h"
 #include "UObject/StrongObjectPtr.h"
 
 struct FWallClockMusicClockDriver : public FMusicClockDriverBase
 {
 public:
-	FWallClockMusicClockDriver(UMusicClockComponent* InClock, UMidiFile* InTempoMap)
-		: FMusicClockDriverBase(InClock)
+	FWallClockMusicClockDriver(UMidiFile* InTempoMap, UObject* WorldContextObj = nullptr)
+		: FMusicClockDriverBase(WorldContextObj)
 		, TempoMapMidi(InTempoMap)
 	{}
 
@@ -21,6 +21,9 @@ public:
 	virtual void OnContinue() override;
 	virtual void OnStop() override {}
 	virtual const ISongMapEvaluator* GetCurrentSongMapEvaluator() const override;
+	
+	virtual bool LoopedThisFrame(ECalibratedMusicTimebase Timebase) const override { return false; }
+	virtual bool SeekedThisFrame(ECalibratedMusicTimebase Timebase) const override { return false; }
 
 private:
 	TWeakObjectPtr<UMidiFile> TempoMapMidi;
diff --git a/Engine/Plugins/Runtime/Harmonix/Source/HarmonixMetasound/Public/HarmonixMetasound/Components/MusicClockComponent.h b/Engine/Plugins/Runtime/Harmonix/Source/HarmonixMetasound/Public/HarmonixMetasound/Components/MusicClockComponent.h
--- a/Engine/Plugins/Runtime/Harmonix/Source/HarmonixMetasound/Public/HarmonixMetasound/Components/MusicClockComponent.h
+++ b/Engine/Plugins/Runtime/Harmonix/Source/HarmonixMetasound/Public/HarmonixMetasound/Components/MusicClockComponent.h
@@ -73,30 +73,24 @@ public:
 	UPROPERTY(EditAnywhere, BlueprintReadWrite, Category = "MusicClock", meta = (EditCondition = "DriveMethod == EMusicClockDriveMethod::MetaSound"))
 	TObjectPtr<UAudioComponent> MetasoundsAudioComponent;
 
-	UPROPERTY(EditDefaultsOnly, Category = "MusicClock", meta = (EditCondition = "DriveMethod == EMusicClockDriveMethod::WallClock"))
-	TObjectPtr<UMidiFile> TempoMap;
-
-	UPROPERTY(EditAnywhere, BlueprintReadWrite, Category = "MusicClock|Defaults")
-	float Tempo = 120.0f;
-
-	UPROPERTY(EditAnywhere, BlueprintReadWrite, Category = "MusicClock|Defaults")
-	int TimeSignatureNum = 4;
-
-	UPROPERTY(EditAnywhere, BlueprintReadWrite, Category = "MusicClock|Defaults")
-	int TimeSignatureDenom = 4;
+	UFUNCTION(BlueprintPure, Category = "MusicClock")
+	float GetTempo() const;
 
-	UPROPERTY(Transient, VisibleAnywhere, BlueprintReadOnly, Category = "MusicClock")
-	float CurrentBeatDurationSec = 0.5f;
+	UFUNCTION(BlueprintPure, Category = "MusicClock")
+	int GetTimeSignatureNumerator() const;
 
-	UPROPERTY(Transient, VisibleAnywhere, BlueprintReadOnly, Category = "MusicClock")
-	float CurrentBarDurationSec = 2.0f;
+	UFUNCTION(BlueprintPure, Category = "MusicClock")
+	int GetTimeSignatureDenominator() const;
 
-	UPROPERTY(Transient, VisibleAnywhere, BlueprintReadOnly, Category = "MusicClock")
-	float CurrentClockAdvanceRate = 1.0f;
+	UFUNCTION(BlueprintPure, Category = "MusicClock")
+	float GetCurrentBeatDurationSeconds() const;
 
-	UPROPERTY(EditAnywhere, BlueprintReadWrite, Category = "MusicClock")
-	bool RunPastMusicEnd = false;
+	UFUNCTION(BlueprintPure, Category = "MusicClock")
+	float GetCurrentBarDurationSeconds() const;
 
+	UFUNCTION(BlueprintPure, Category = "MusicClock")
+	float GetCurrentClockAdvanceRate() const;
+	
 	UFUNCTION(BlueprintCallable, Category = "MusicClock")
 	void Start();
 
@@ -110,9 +104,7 @@ public:
 	void Stop();
 
 	UFUNCTION(BlueprintPure, Category = "MusicClock")
-	EMusicClockState GetState() const { return State; }
-
-	// Getters for all of the fields in mCurrentSongPos...
+	EMusicClockState GetState() const;
 	
 	// Time from the beginning of the authored music content.
 	// NOTE: INCLUDES time for count-in and pickup bars.
@@ -327,6 +319,30 @@ private:
 	UPROPERTY(BlueprintGetter = GetCurrentRawAudioRenderSongPos, Category = "MusicClock")
 	FMidiSongPos CurrentRawAudioRenderSongPos;
 
+	UPROPERTY(EditDefaultsOnly, Category = "MusicClock", meta = (EditCondition = "DriveMethod == EMusicClockDriveMethod::WallClock"))
+    TObjectPtr<UMidiFile> TempoMap;
+
+    UPROPERTY(BlueprintGetter = GetTempo,  EditAnywhere, Category = "MusicClock|Defaults")
+    float Tempo = 120.0f;
+
+    UPROPERTY(BlueprintGetter = GetTimeSignatureNumerator, EditAnywhere, Category = "MusicClock|Defaults")
+    int TimeSignatureNum = 4;
+
+    UPROPERTY(BlueprintGetter = GetTimeSignatureDenominator, EditAnywhere, Category = "MusicClock|Defaults")
+    int TimeSignatureDenom = 4;
+
+    UPROPERTY(BlueprintGetter = GetCurrentBeatDurationSeconds, Transient, VisibleAnywhere, Category = "MusicClock")
+    float CurrentBeatDurationSec = 0.5f;
+
+    UPROPERTY(BlueprintGetter = GetCurrentBarDurationSeconds, Transient, VisibleAnywhere, Category = "MusicClock")
+    float CurrentBarDurationSec = 2.0f;
+
+    UPROPERTY(BlueprintGetter = GetCurrentClockAdvanceRate, Transient, VisibleAnywhere, Category = "MusicClock")
+    float CurrentClockAdvanceRate = 1.0f;
+
+    UPROPERTY(EditAnywhere, Category = "MusicClock")
+    bool RunPastMusicEnd = false;
+
 public:
 	// Getter functions for the Blueprint properties exposed above...
 	UFUNCTION(BlueprintCallable, BlueprintPure, Category = "MusicClock")
@@ -369,7 +385,7 @@ public:
 	// Note: Not const as it might cause the clock to update from its source.
 	FMidiSongPos CalculateSongPosWithOffset(float MsOffset, ECalibratedMusicTimebase Timebase = ECalibratedMusicTimebase::VideoRenderTime) const;
 
-	const FMidiSongPos& GetRawUnsmoothedAudioRenderPos() const { return CurrentRawAudioRenderSongPos; }
+	const FMidiSongPos& GetRawUnsmoothedAudioRenderPos() const;
 
 	// vvv BEGIN  IMusicEnvironmentClockSource Implementation vvv
 	virtual FMusicalTime GetMusicalTime() const override;
@@ -385,86 +401,27 @@ public:
 	void TickComponentInternal();
 	
 private:
-	friend struct FMusicClockDriverBase;
-	friend struct FMetasoundMusicClockDriver;
-	friend struct FWallClockMusicClockDriver;
 	friend class  UMidiClockUpdateSubsystem;
 
 	TSharedPtr<FMusicClockDriverBase> ClockDriver;
 
-	EMusicClockState State = EMusicClockState::Stopped;
-
 	FSongMaps DefaultMaps;
-
-	// TODO shouldn't these deltas be rolled into the MidiSongPos?
-	float RawAudioRenderDeltaBarF = 0.0f;
-	float RawAudioRenderDeltaBeatF = 0.0f;
-	float AudioRenderDeltaBarF = 0.0f;
-	float AudioRenderDeltaBeatF = 0.0f;
-	float PlayerExperienceDeltaBarF = 0.0f;
-	float PlayerExperienceDeltaBeatF = 0.0f;
-	float VideoRenderDeltaBarF = 0.0f;
-	float VideoRenderDeltaBeatF = 0.0f;
+	
 	int32 LastBroadcastBar = -1;
 	int32 LastBroadcastBeat = -1;
 	FSongSection LastBroadcastSongSection;
-
-	bool AudioRenderSeekDetected = false;
-	bool AudioRenderLoopDetected = false;
-	bool PlayerExperiencedSeekDetected = false;
-	bool PlayerExperiencedLoopDetected = false;
-	bool VideoRenderSeekDetected = false;
-	bool VideoRenderLoopDetected = false;
-
-	FMidiSongPos PrevRawAudioRenderSongPos;
-	FMidiSongPos PrevAudioRenderSongPos;
-	FMidiSongPos PrevPlayerExperiencedSongPos;
-	FMidiSongPos PrevVideoRenderSongPos;
-
-	uint64 LastUpdateFrame = 0;
-
+	
 	void CreateClockDriver();
 	void BroadcastSongPosChanges();
 	void BroadcastSeekLoopDetections() const;
+	const FMusicTimeDiscontinuityEvent* GetMusicTimeDiscontinuityEventEventInternal(ECalibratedMusicTimebase Timebase) const;
 	void MakeDefaultSongMap();
 	bool ConnectToMetasound();
 	void ConnectToWallClock();
 	void DisconnectFromClockDriver();
-
-	// Ensures the clock will be updated once per frame.  Should only get called on the game thread.
-	void EnsureClockIsValidForGameFrame() const;// TODO: Cleanup task - UE-205069 - If we find we are able to use the new 
-												// MidiClock/MusicClockComponent ticking methods, this function should be
-												// deleted and all of the call sites cleaned up... as only this next, 
-												// non-const "Ensure" function will be required...
+	
 	void EnsureClockIsValidForGameFrameFromSubsystem();
 	const FMidiSongPos& GetCurrentSongPosInternal(ECalibratedMusicTimebase Timebase) const;
 	const FMidiSongPos& GetPreviousSongPosInternal(ECalibratedMusicTimebase Timebase) const;
 };
 
-struct FMusicClockDriverBase : public TSharedFromThis<FMusicClockDriverBase>
-{
-public:
-	FMusicClockDriverBase() = delete;
-	FMusicClockDriverBase(UMusicClockComponent* DrivenClock)
-		: ClockComponent(DrivenClock)
-	{}
-	virtual ~FMusicClockDriverBase() = default;
-
-	void EnsureClockIsValidForGameFrame();
-
-	virtual bool CalculateSongPosWithOffset(float MsOffset, ECalibratedMusicTimebase Timebase, FMidiSongPos& OutResult) const = 0;
-
-	virtual void Disconnect() = 0;
-	virtual void OnStart() = 0;
-	virtual void OnPause() = 0;
-	virtual void OnContinue() = 0;
-	virtual void OnStop() = 0;
-	virtual const ISongMapEvaluator* GetCurrentSongMapEvaluator() const = 0;
-
-protected:
-	UMusicClockComponent* ClockComponent;
-
-private:
-	virtual bool RefreshCurrentSongPos() = 0;
-};
-
diff --git a/Engine/Plugins/Runtime/Harmonix/Source/HarmonixMetasound/Public/HarmonixMetasound/Components/MusicClockDriverBase.h b/Engine/Plugins/Runtime/Harmonix/Source/HarmonixMetasound/Public/HarmonixMetasound/Components/MusicClockDriverBase.h
new file mode 100644
--- /dev/null
+++ b/Engine/Plugins/Runtime/Harmonix/Source/HarmonixMetasound/Public/HarmonixMetasound/Components/MusicClockDriverBase.h
@@ -0,0 +1,87 @@
+﻿// Copyright Epic Games, Inc. All Rights Reserved.
+#pragma once
+
+#include "HarmonixMetasound/Components/MusicClockComponent.h"
+#include "Harmonix/MusicalTimebase.h"
+#include "HarmonixMidi/MidiSongPos.h"
+#include "HarmonixMidi/SongMaps.h"
+
+struct ISongMapEvaluator;
+
+struct FMusicClockDriverBase : public TSharedFromThis<FMusicClockDriverBase>
+{
+public:
+	FMusicClockDriverBase(UObject* WorldContextObject);
+	virtual ~FMusicClockDriverBase();
+
+	void EnsureClockIsValidForGameFrame();
+
+	virtual bool CalculateSongPosWithOffset(float MsOffset, ECalibratedMusicTimebase Timebase, FMidiSongPos& OutResult) const = 0;
+
+	double GetWallClockTime() const;
+
+	EMusicClockState GetState() const { return MusicClockState; }
+	const FMidiSongPos& GetCurrentSongPos(ECalibratedMusicTimebase Timebase) const;
+	const FMidiSongPos& GetPreviousSongPos(ECalibratedMusicTimebase Timebase) const;
+
+	void Start();
+	void Pause();
+	void Continue();
+	void Stop();
+	
+	virtual const ISongMapEvaluator* GetCurrentSongMapEvaluator() const = 0;
+
+	virtual bool LoopedThisFrame(ECalibratedMusicTimebase Timebase) const = 0;
+	virtual bool SeekedThisFrame(ECalibratedMusicTimebase Timebase) const = 0;
+
+	virtual void Disconnect() = 0;
+	
+	FSongMaps DefaultMaps;
+	
+	float Tempo = 120.0f;
+	int TimeSignatureNum = 4;
+	int TimeSignatureDenom = 4;
+
+	float CurrentBeatDurationSec = 0.5f;
+	float CurrentBarDurationSec = 2.0f;
+
+	FMidiSongPos CurrentSmoothedAudioRenderSongPos;
+	FMidiSongPos CurrentVideoRenderSongPos;
+	FMidiSongPos CurrentPlayerExperiencedSongPos;
+	FMidiSongPos CurrentRawAudioRenderSongPos;
+	
+	FMidiSongPos PrevRawAudioRenderSongPos;
+	FMidiSongPos PrevAudioRenderSongPos;
+	FMidiSongPos PrevPlayerExperiencedSongPos;
+	FMidiSongPos PrevVideoRenderSongPos;
+	
+	float CurrentClockAdvanceRate = 1.0f;
+
+	// TODO shouldn't these deltas be rolled into the MidiSongPos?
+	float RawAudioRenderDeltaBarF = 0.0f;
+	float RawAudioRenderDeltaBeatF = 0.0f;
+	float AudioRenderDeltaBarF = 0.0f;
+	float AudioRenderDeltaBeatF = 0.0f;
+	float PlayerExperienceDeltaBarF = 0.0f;
+	float PlayerExperienceDeltaBeatF = 0.0f;
+	float VideoRenderDeltaBarF = 0.0f;
+	float VideoRenderDeltaBeatF = 0.0f;
+	
+protected:
+	
+	virtual void OnStart() = 0;
+	virtual void OnPause() = 0;
+	virtual void OnContinue() = 0;
+	virtual void OnStop() = 0;
+	
+private:
+
+	void ResetDefaultSongMaps();
+	virtual bool RefreshCurrentSongPos() = 0;
+
+	uint64 LastUpdateFrame = 0;
+
+	EMusicClockState MusicClockState = EMusicClockState::Stopped;
+
+	TWeakObjectPtr<UObject> ContextObj;
+};
\ No newline at end of file
