[Harmonix] Disable MusicTempometerComponent tick on dedicated server

#rnx

diff --git a/Engine/Plugins/Runtime/Harmonix/Source/HarmonixMetasound/Private/HarmonixMetasound/Components/MusicTempometerComponent.cpp b/Engine/Plugins/Runtime/Harmonix/Source/HarmonixMetasound/Private/HarmonixMetasound/Components/MusicTempometerComponent.cpp
--- a/Engine/Plugins/Runtime/Harmonix/Source/HarmonixMetasound/Private/HarmonixMetasound/Components/MusicTempometerComponent.cpp
+++ b/Engine/Plugins/Runtime/Harmonix/Source/HarmonixMetasound/Private/HarmonixMetasound/Components/MusicTempometerComponent.cpp
@@ -11,6 +11,7 @@ UMusicTempometerComponent::UMusicTempometerComponent()
 	SetTickGroup(TG_PrePhysics);
 
 	PrimaryComponentTick.bCanEverTick = true;
+	PrimaryComponentTick.bAllowTickOnDedicatedServer = false;
 	PrimaryComponentTick.SetPriorityIncludingPrerequisites(true);
 }
 
