[Harmonix] Don't adjust tempometer MPC params if they already have changes from the CDO
#tests PIE

diff --git a/Engine/Plugins/Runtime/Harmonix/Source/HarmonixMetasound/Private/HarmonixMetasound/Components/MusicTempometerComponent.cpp b/Engine/Plugins/Runtime/Harmonix/Source/HarmonixMetasound/Private/HarmonixMetasound/Components/MusicTempometerComponent.cpp
--- a/Engine/Plugins/Runtime/Harmonix/Source/HarmonixMetasound/Private/HarmonixMetasound/Components/MusicTempometerComponent.cpp
+++ b/Engine/Plugins/Runtime/Harmonix/Source/HarmonixMetasound/Private/HarmonixMetasound/Components/MusicTempometerComponent.cpp
@@ -21,6 +21,17 @@ void UMusicTempometerComponent::PostLoad()
 
 	// Fix up deprecated properties
 	const FName DeprecatedErroneousName = GetFName();
+
+	if (UMusicTempometerComponent* DefaultTempometer = GetClass()->GetDefaultObject<UMusicTempometerComponent>())
+	{
+		// If the CDO's parameter names and these parameter names don't match, they've already been edited.
+		// Don't overwrite with the deprecated names.
+		if (MPCParameters.CurrentFrameParameterNames != DefaultTempometer->MPCParameters.CurrentFrameParameterNames)
+		{
+			return;
+		}
+	}
+
 	if (SecondsIncludingCountInParameterName_DEPRECATED != DeprecatedErroneousName)
 	{
 		MPCParameters.CurrentFrameParameterNames.SecondsIncludingCountInParameterName = SecondsIncludingCountInParameterName_DEPRECATED;
diff --git a/Engine/Plugins/Runtime/Harmonix/Source/HarmonixMetasound/Public/HarmonixMetasound/Utils/MusicTempometerUtilities.h b/Engine/Plugins/Runtime/Harmonix/Source/HarmonixMetasound/Public/HarmonixMetasound/Utils/MusicTempometerUtilities.h
--- a/Engine/Plugins/Runtime/Harmonix/Source/HarmonixMetasound/Public/HarmonixMetasound/Utils/MusicTempometerUtilities.h
+++ b/Engine/Plugins/Runtime/Harmonix/Source/HarmonixMetasound/Public/HarmonixMetasound/Utils/MusicTempometerUtilities.h
@@ -118,6 +118,26 @@ struct FMusicTempometerMPCParameterNames
 				!TimeSignatureDenominatorParameterName.IsNone() ||
 				!TempoParameterName.IsNone();
 	}
+
+	bool operator==(const FMusicTempometerMPCParameterNames& Other)
+	{
+		return	SecondsIncludingCountInParameterName == Other.SecondsIncludingCountInParameterName &&
+			BarsIncludingCountInParameterName == Other.BarsIncludingCountInParameterName &&
+			BeatsIncludingCountInParameterName == Other.BeatsIncludingCountInParameterName &&
+			SecondsFromBarOneParameterName == Other.SecondsFromBarOneParameterName &&
+			TimestampBarParameterName == Other.TimestampBarParameterName &&
+			TimestampBeatInBarParameterName == Other.TimestampBeatInBarParameterName &&
+			BarProgressParameterName == Other.BarProgressParameterName &&
+			BeatProgressParameterName == Other.BeatProgressParameterName &&
+			TimeSignatureNumeratorParameterName == Other.TimeSignatureNumeratorParameterName &&
+			TimeSignatureDenominatorParameterName == Other.TimeSignatureDenominatorParameterName &&
+			TempoParameterName == Other.TempoParameterName;
+	}
+
+	bool operator!=(const FMusicTempometerMPCParameterNames& Other)
+	{
+		return	!operator==(Other);
+	}
 };
 /**
  * Parameters used when creating/updating the corresponding Material Parameter Collection
