[Harmonix] Add "MusicTimestampValid" parameter to MusicTempometer so that materials can know whether or not the rest of the MPC values are coming from a clock that is valid and not stopped.

#tests PIE - sample material in test level
#rnx

diff --git a/Engine/Plugins/Runtime/Harmonix/Source/HarmonixMetasound/Private/HarmonixMetasound/Utils/MusicTempometerUtilities.cpp b/Engine/Plugins/Runtime/Harmonix/Source/HarmonixMetasound/Private/HarmonixMetasound/Utils/MusicTempometerUtilities.cpp
--- a/Engine/Plugins/Runtime/Harmonix/Source/HarmonixMetasound/Private/HarmonixMetasound/Utils/MusicTempometerUtilities.cpp
+++ b/Engine/Plugins/Runtime/Harmonix/Source/HarmonixMetasound/Private/HarmonixMetasound/Utils/MusicTempometerUtilities.cpp
@@ -53,4 +53,5 @@ void MusicTempometerUtilities::UpdateMaterialParameterCollectionFromSongPos(cons
 	InOutMaterialParameterCollectionInstance->SetScalarParameterValue(InMCPParameters.CurrentFrameParameterNames.TimeSignatureNumeratorParameterName, InMidiSongPos.TimeSigNumerator);
 	InOutMaterialParameterCollectionInstance->SetScalarParameterValue(InMCPParameters.CurrentFrameParameterNames.TimeSignatureDenominatorParameterName, InMidiSongPos.TimeSigDenominator);
 	InOutMaterialParameterCollectionInstance->SetScalarParameterValue(InMCPParameters.CurrentFrameParameterNames.TempoParameterName, InMidiSongPos.Tempo);
+	InOutMaterialParameterCollectionInstance->SetScalarParameterValue(InMCPParameters.CurrentFrameParameterNames.TimestampValidParameterName, InMidiSongPos.IsValid() ? 1.0f : 0.0f);
 }
\ No newline at end of file
diff --git a/Engine/Plugins/Runtime/Harmonix/Source/HarmonixMetasound/Public/HarmonixMetasound/Utils/MusicTempometerUtilities.h b/Engine/Plugins/Runtime/Harmonix/Source/HarmonixMetasound/Public/HarmonixMetasound/Utils/MusicTempometerUtilities.h
--- a/Engine/Plugins/Runtime/Harmonix/Source/HarmonixMetasound/Public/HarmonixMetasound/Utils/MusicTempometerUtilities.h
+++ b/Engine/Plugins/Runtime/Harmonix/Source/HarmonixMetasound/Public/HarmonixMetasound/Utils/MusicTempometerUtilities.h
@@ -103,6 +103,12 @@ struct FMusicTempometerMPCParameterNames
 	UPROPERTY(EditDefaultsOnly, Category = "MusicClock")
 	FName TempoParameterName = "MusicTempo";
 
+	/**
+	 * 0 or 1, represents whether the rest of the parameters are defaults (0) or the current state of a clock (1).
+	 */
+	UPROPERTY(EditDefaultsOnly, Category = "MusicClock")
+	FName TimestampValidParameterName = "MusicTimestampValid";
+
 	bool IsValid() const
 	{
 		// Valid if there is an MPC collection assigned and at least one named parameter
@@ -116,7 +122,8 @@ struct FMusicTempometerMPCParameterNames
 				!BeatProgressParameterName.IsNone() ||
 				!TimeSignatureNumeratorParameterName.IsNone() ||
 				!TimeSignatureDenominatorParameterName.IsNone() ||
-				!TempoParameterName.IsNone();
+				!TempoParameterName.IsNone() ||
+				!TimestampValidParameterName.IsNone();
 	}
 
 	bool operator==(const FMusicTempometerMPCParameterNames& Other)
@@ -131,7 +138,8 @@ struct FMusicTempometerMPCParameterNames
 			BeatProgressParameterName == Other.BeatProgressParameterName &&
 			TimeSignatureNumeratorParameterName == Other.TimeSignatureNumeratorParameterName &&
 			TimeSignatureDenominatorParameterName == Other.TimeSignatureDenominatorParameterName &&
-			TempoParameterName == Other.TempoParameterName;
+			TempoParameterName == Other.TempoParameterName &&
+			TimestampValidParameterName == Other.TimestampValidParameterName;
 	}
 
 	bool operator!=(const FMusicTempometerMPCParameterNames& Other)
