Merging CL#40313599

[Harmonix] [FM] Change default Tempo to 0 instead of 120, except for the wall clock, so we we don't claim 120 BPM before music clocks start incrementing.

#tests PIE Main
#rnx

#changelist validated
#virtualized

diff --git a/Engine/Plugins/Runtime/Harmonix/Source/HarmonixMetasound/Private/HarmonixMetasound/Components/MetasoundMusicClockDriver.cpp b/Engine/Plugins/Runtime/Harmonix/Source/HarmonixMetasound/Private/HarmonixMetasound/Components/MetasoundMusicClockDriver.cpp
--- a/Engine/Plugins/Runtime/Harmonix/Source/HarmonixMetasound/Private/HarmonixMetasound/Components/MetasoundMusicClockDriver.cpp
+++ b/Engine/Plugins/Runtime/Harmonix/Source/HarmonixMetasound/Private/HarmonixMetasound/Components/MetasoundMusicClockDriver.cpp
@@ -422,7 +422,7 @@ void FMetasoundMusicClockDriver::RefreshCurrentSongPosFromWallClock()
 {
 	check(ClockComponent);
 
-	bool TempoChanged = ClockComponent->CurrentSmoothedAudioRenderSongPos.Tempo != ClockComponent->Tempo;
+	bool TempoChanged = ClockComponent->CurrentSmoothedAudioRenderSongPos.Tempo != ClockComponent->CurrentTempo;
 
 	double FreeRunTime = (GetWallClockTime() - FreeRunStartTimeSecs) * ClockComponent->CurrentClockAdvanceRate;
 
@@ -433,9 +433,7 @@ void FMetasoundMusicClockDriver::RefreshCurrentSongPosFromWallClock()
 
 	if (TempoChanged)
 	{
-		ClockComponent->Tempo = ClockComponent->CurrentSmoothedAudioRenderSongPos.Tempo;
-		ClockComponent->CurrentBeatDurationSec = (60.0f / ClockComponent->Tempo) / ClockComponent->CurrentClockAdvanceRate;
-		ClockComponent->CurrentBarDurationSec = ((ClockComponent->TimeSignatureNum * ClockComponent->CurrentBeatDurationSec) / (ClockComponent->TimeSignatureDenom / 4.0f)) / ClockComponent->CurrentClockAdvanceRate;
+		ClockComponent->UpdateCurrentTempo(ClockComponent->CurrentSmoothedAudioRenderSongPos.Tempo);
 	}
 }
 
@@ -584,15 +582,13 @@ void FMetasoundMusicClockDriver::RefreshCurrentSongPosFromHistory()
 		UpdateCurrentTicksForLoopingOrMonotonicClock(SmoothedTick, SmoothedTempoMapTick);
 	}
 
-	ClockComponent->TimeSignatureNum = ClockComponent->CurrentSmoothedAudioRenderSongPos.TimeSigNumerator;
-	ClockComponent->TimeSignatureDenom = ClockComponent->CurrentSmoothedAudioRenderSongPos.TimeSigDenominator;
+	ClockComponent->CurrentTimeSigNum = ClockComponent->CurrentSmoothedAudioRenderSongPos.TimeSigNumerator;
+	ClockComponent->CurrentTimeSigDenom = ClockComponent->CurrentSmoothedAudioRenderSongPos.TimeSigDenominator;
 
-	if (ClockComponent->Tempo != ClockComponent->CurrentSmoothedAudioRenderSongPos.Tempo || ClockComponent->CurrentClockAdvanceRate != SpeedAtRawRenderTime)
+	if (ClockComponent->CurrentTempo != ClockComponent->CurrentSmoothedAudioRenderSongPos.Tempo || ClockComponent->CurrentClockAdvanceRate != SpeedAtRawRenderTime)
 	{
 		ClockComponent->CurrentClockAdvanceRate = SpeedAtRawRenderTime;
-		ClockComponent->Tempo = ClockComponent->CurrentSmoothedAudioRenderSongPos.Tempo;
-		ClockComponent->CurrentBeatDurationSec = (60.0f / ClockComponent->Tempo) / ClockComponent->CurrentClockAdvanceRate;
-		ClockComponent->CurrentBarDurationSec = ((ClockComponent->TimeSignatureNum * ClockComponent->CurrentBeatDurationSec) / (ClockComponent->TimeSignatureDenom / 4.0f)) / ClockComponent->CurrentClockAdvanceRate;
+		ClockComponent->UpdateCurrentTempo(ClockComponent->CurrentSmoothedAudioRenderSongPos.Tempo);
 	}
 }
 
diff --git a/Engine/Plugins/Runtime/Harmonix/Source/HarmonixMetasound/Private/HarmonixMetasound/Components/MusicClockComponent.cpp b/Engine/Plugins/Runtime/Harmonix/Source/HarmonixMetasound/Private/HarmonixMetasound/Components/MusicClockComponent.cpp
--- a/Engine/Plugins/Runtime/Harmonix/Source/HarmonixMetasound/Private/HarmonixMetasound/Components/MusicClockComponent.cpp
+++ b/Engine/Plugins/Runtime/Harmonix/Source/HarmonixMetasound/Private/HarmonixMetasound/Components/MusicClockComponent.cpp
@@ -72,6 +72,43 @@ void UMusicClockComponent::ConnectToWallClockForMidi(UMidiFile* InTempoMap)
 	ConnectToWallClock();
 }
 
+void UMusicClockComponent::SetDefaultTempo(float TempoBpm)
+{
+	// Defaults only used to make the tempo map, once before ever running.
+	ensure(State == EMusicClockState::Stopped);
+	DefaultTempo = TempoBpm;
+}
+
+void UMusicClockComponent::SetDefaultTimeSignatureNum(int Num)
+{
+	ensure(State == EMusicClockState::Stopped);
+	DefaultTimeSignatureNum = Num;
+}
+
+void UMusicClockComponent::SetDefaultTimeSignatureDenom(int Denom)
+{
+	ensure(State == EMusicClockState::Stopped);
+	DefaultTimeSignatureDenom = Denom;
+}
+
+void UMusicClockComponent::CalculateBarAndBeatDurations(float Tempo, float AdvanceRate, int32 TimeSigNum, int32 TimeSigDenom, float& OutBarDuration, float& OutBeatDuration)
+{
+	OutBeatDuration = 0.0f;
+	OutBarDuration = 0.0f;
+	if (Tempo > 0.0f && AdvanceRate > 0.0f && TimeSigNum > 0 && TimeSigDenom)
+	{
+		OutBeatDuration = (60.0f / Tempo) / AdvanceRate;
+		OutBarDuration = ((TimeSigNum * OutBeatDuration) / (TimeSigDenom / 4.0f)) / AdvanceRate;
+	}
+}
+
+void UMusicClockComponent::UpdateCurrentTempo(float Tempo)
+{
+	CurrentTempo = Tempo;
+	CalculateBarAndBeatDurations(CurrentTempo, CurrentClockAdvanceRate, CurrentTimeSigNum, CurrentTimeSigDenom, CurrentBarDurationSec, CurrentBeatDurationSec);
+}
+
+
 // TODO: Cleanup task - UE-205069 - If we find we are able to use the new MidiClock/MusicClockComponent
 // ticking methods, this function should be deleted and all of the call sites cleaned up... 
 // as only the non-const "Ensure" function will be required.
@@ -244,6 +281,32 @@ void UMusicClockComponent::TickComponentInternal()
 	BroadcastSeekLoopDetections();
 }
 
+float UMusicClockComponent::GetCurrentTempo() const
+{
+	return CurrentTempo;
+}
+
+void UMusicClockComponent::GetCurrentTimeSignature(int& OutNumerator, int& OutDenominator) const
+{
+	OutNumerator = CurrentTimeSigNum;
+	OutDenominator = CurrentTimeSigDenom;
+}
+
+float UMusicClockComponent::GetCurrentBarDurationSeconds() const
+{
+	return CurrentBeatDurationSec;
+}
+
+float UMusicClockComponent::GetCurrentBeatDurationSeconds() const
+{
+	return CurrentBeatDurationSec;
+}
+
+float UMusicClockComponent::GetCurrentClockAdvanceRate() const
+{
+	return CurrentClockAdvanceRate;
+}
+
 void UMusicClockComponent::Start()
 {
 	MakeDefaultSongMap();
@@ -672,8 +735,8 @@ void UMusicClockComponent::MakeDefaultSongMap()
 {
 	DefaultMaps.EmptyAllMaps();
 	DefaultMaps.Init(Harmonix::Midi::Constants::GTicksPerQuarterNoteInt);
-	DefaultMaps.GetTempoMap().AddTempoInfoPoint(Harmonix::Midi::Constants::BPMToMidiTempo(Tempo), 0);
-	DefaultMaps.GetBarMap().AddTimeSignatureAtBarIncludingCountIn(0, TimeSignatureNum, TimeSignatureDenom);
+	DefaultMaps.GetTempoMap().AddTempoInfoPoint(Harmonix::Midi::Constants::BPMToMidiTempo(DefaultTempo), 0);
+	DefaultMaps.GetBarMap().AddTimeSignatureAtBarIncludingCountIn(0, DefaultTimeSignatureNum, DefaultTimeSignatureDenom);
 }
 
 const FMidiSongPos& UMusicClockComponent::GetCurrentSongPosInternal(ECalibratedMusicTimebase Timebase) const
diff --git a/Engine/Plugins/Runtime/Harmonix/Source/HarmonixMetasound/Private/HarmonixMetasound/Components/WallClockMusicClockDriver.cpp b/Engine/Plugins/Runtime/Harmonix/Source/HarmonixMetasound/Private/HarmonixMetasound/Components/WallClockMusicClockDriver.cpp
--- a/Engine/Plugins/Runtime/Harmonix/Source/HarmonixMetasound/Private/HarmonixMetasound/Components/WallClockMusicClockDriver.cpp
+++ b/Engine/Plugins/Runtime/Harmonix/Source/HarmonixMetasound/Private/HarmonixMetasound/Components/WallClockMusicClockDriver.cpp
@@ -40,7 +40,7 @@ bool FWallClockMusicClockDriver::RefreshCurrentSongPos()
 	check(ClockComponent);
 	check(ClockComponent->GetWorld());
 
-	bool TempoChanged = ClockComponent->CurrentSmoothedAudioRenderSongPos.Tempo != ClockComponent->Tempo;
+	bool TempoChanged = ClockComponent->CurrentSmoothedAudioRenderSongPos.Tempo != ClockComponent->CurrentTempo;
 
 	double RunTime = ClockComponent->GetWorld()->GetTimeSeconds() - StartTimeSecs;
 
@@ -51,11 +51,14 @@ bool FWallClockMusicClockDriver::RefreshCurrentSongPos()
 	ClockComponent->CurrentPlayerExperiencedSongPos.SetByTime(ClockComponent->CurrentSmoothedAudioRenderSongPos.SecondsIncludingCountIn * 1000.0f - FHarmonixModule::GetMeasuredUserExperienceAndReactionToAudioRenderOffsetMs(), *Maps);
 	ClockComponent->CurrentVideoRenderSongPos.SetByTime(ClockComponent->CurrentSmoothedAudioRenderSongPos.SecondsIncludingCountIn * 1000.0f - FHarmonixModule::GetMeasuredVideoToAudioRenderOffsetMs(), *Maps);
 
-	if (TempoChanged)
+	bool TimeSigChanged = ClockComponent->CurrentSmoothedAudioRenderSongPos.TimeSigNumerator != ClockComponent->CurrentTimeSigNum
+	                   || ClockComponent->CurrentSmoothedAudioRenderSongPos.TimeSigDenominator != ClockComponent->CurrentTimeSigDenom;
+
+	if (TempoChanged || TimeSigChanged)
 	{
-		ClockComponent->Tempo = ClockComponent->CurrentSmoothedAudioRenderSongPos.Tempo;
-		ClockComponent->CurrentBeatDurationSec = (60.0f / ClockComponent->Tempo) / ClockComponent->CurrentClockAdvanceRate;
-		ClockComponent->CurrentBarDurationSec = ((ClockComponent->TimeSignatureNum * ClockComponent->CurrentBeatDurationSec) / (ClockComponent->TimeSignatureDenom / 4.0f)) / ClockComponent->CurrentClockAdvanceRate;
+		ClockComponent->CurrentTimeSigNum = ClockComponent->CurrentSmoothedAudioRenderSongPos.TimeSigNumerator;
+		ClockComponent->CurrentTimeSigDenom = ClockComponent->CurrentSmoothedAudioRenderSongPos.TimeSigDenominator;
+		ClockComponent->UpdateCurrentTempo(ClockComponent->CurrentSmoothedAudioRenderSongPos.Tempo);
 	}
 
 	return true;
diff --git a/Engine/Plugins/Runtime/Harmonix/Source/HarmonixMetasound/Public/HarmonixMetasound/Components/MusicClockComponent.h b/Engine/Plugins/Runtime/Harmonix/Source/HarmonixMetasound/Public/HarmonixMetasound/Components/MusicClockComponent.h
--- a/Engine/Plugins/Runtime/Harmonix/Source/HarmonixMetasound/Public/HarmonixMetasound/Components/MusicClockComponent.h
+++ b/Engine/Plugins/Runtime/Harmonix/Source/HarmonixMetasound/Public/HarmonixMetasound/Components/MusicClockComponent.h
@@ -76,27 +76,6 @@ public:
 	UPROPERTY(EditDefaultsOnly, Category = "MusicClock", meta = (EditCondition = "DriveMethod == EMusicClockDriveMethod::WallClock"))
 	TObjectPtr<UMidiFile> TempoMap;
 
-	UPROPERTY(EditAnywhere, BlueprintReadWrite, Category = "MusicClock|Defaults")
-	float Tempo = 120.0f;
-
-	UPROPERTY(EditAnywhere, BlueprintReadWrite, Category = "MusicClock|Defaults")
-	int TimeSignatureNum = 4;
-
-	UPROPERTY(EditAnywhere, BlueprintReadWrite, Category = "MusicClock|Defaults")
-	int TimeSignatureDenom = 4;
-
-	UPROPERTY(Transient, VisibleAnywhere, BlueprintReadOnly, Category = "MusicClock")
-	float CurrentBeatDurationSec = 0.5f;
-
-	UPROPERTY(Transient, VisibleAnywhere, BlueprintReadOnly, Category = "MusicClock")
-	float CurrentBarDurationSec = 2.0f;
-
-	UPROPERTY(Transient, VisibleAnywhere, BlueprintReadOnly, Category = "MusicClock")
-	float CurrentClockAdvanceRate = 1.0f;
-
-	UPROPERTY(EditAnywhere, BlueprintReadWrite, Category = "MusicClock")
-	bool RunPastMusicEnd = false;
-
 	UFUNCTION(BlueprintCallable, Category = "MusicClock")
 	void Start();
 
@@ -109,6 +88,47 @@ public:
 	UFUNCTION(BlueprintCallable, Category = "MusicClock")
 	void Stop();
 
+	/**
+	 * Get The Current Tempo (BPM) of the runnig clock
+	 * 0 if the clock is stopped
+	 */
+	UFUNCTION(BlueprintPure, Category = "MusicClock")
+	float GetCurrentTempo() const;
+
+	/**
+	 * Get The Current Time Signature
+	 * 
+	 * 0 if the clock is stopped
+	 */
+	UFUNCTION(BlueprintPure, Category = "MusicClock")
+	void GetCurrentTimeSignature(int& OutNumerator, int& OutDenominator) const;
+
+	/**
+	 * Get The Current Bar Duration in Seconds
+	 * Based on the current tempo, advance rate and time signature
+	 * 
+	 * 0 if the clock is stopped
+	 */
+	UFUNCTION(BlueprintPure, Category = "MusicClock")
+	float GetCurrentBarDurationSeconds() const;
+
+	/**
+	 * Get The Current Beat Duration in Seconds
+	 * Based on the current tempo, advance rate and time signature
+	 * 
+	 * 0 if the clock is stopped
+	 */
+	UFUNCTION(BlueprintPure, Category = "MusicClock")
+	float GetCurrentBeatDurationSeconds() const;
+	
+	/**
+     * Get The Current Clock advance rate
+     * 
+     * 0 if the clock is stopped
+     */
+	UFUNCTION(BlueprintPure, Category = "MusicClock")
+	float GetCurrentClockAdvanceRate() const;
+
 	UFUNCTION(BlueprintPure, Category = "MusicClock")
 	EMusicClockState GetState() const { return State; }
 
@@ -315,9 +335,59 @@ public:
 	UPROPERTY(BlueprintAssignable, Category = "MusicClock")
 	FMusicTimeDiscontinuityEvent VideoRenderMusicTimeDiscontinuityEvent;
 
+	// Set the default tempo, only call when there's no tempo map.
+	void SetDefaultTempo(float TempoBpm);
+	// Set the default time signature numerator, only call when there's no tempo map.
+	void SetDefaultTimeSignatureNum(int Num);
+	// Set the default time signature denominator, only call when there's no tempo map.
+	void SetDefaultTimeSignatureDenom(int Denom);
+
+	static void CalculateBarAndBeatDurations(float Tempo, float AdvanceRate, int32 TimeSigNum, int32 TimeSigDenom, float& OutBarDuration, float& OutBeatDuration);
+
+protected:
+	void UpdateCurrentTempo(float Tempo);
+
+	/**
+	 *
+	 * Default Tempo (BPM)
+	 * used for Wall clock with no tempo map
+	 *
+	 * NOT THE CURRENT RUNNING TEMPO
+	 */
+	UPROPERTY(EditAnywhere, BlueprintReadWrite, Category = "MusicClock|Defaults")
+	float DefaultTempo = 120.0f;
+
+	/**
+	 * Default Time Signature (Numerator)
+	 * used for the Wall Clock with no tempo map
+	 *
+	 * NOT THE CURRENT RUNNING TIME SIG
+	 */
+	UPROPERTY(EditAnywhere, BlueprintReadWrite, Category = "MusicClock|Defaults")
+	int DefaultTimeSignatureNum = 4;
+
+	/**
+	 * Default Time Signature (Denominator) 
+	 * used for the Wall clock with no tempo map
+	 *
+	 * NOT THE CURRENT RUNNING TIME SIG
+	 */
+	UPROPERTY(EditAnywhere, BlueprintReadWrite, Category = "MusicClock|Defaults")
+	int DefaultTimeSignatureDenom = 4;
+
+	UPROPERTY(Transient, VisibleAnywhere, BlueprintGetter=GetCurrentBeatDurationSeconds, Category = "MusicClock")
+	float CurrentBeatDurationSec = 0.0f;
+
+	UPROPERTY(Transient, VisibleAnywhere, BlueprintGetter=GetCurrentBarDurationSeconds, Category = "MusicClock")
+	float CurrentBarDurationSec = 0.0f;
+
+	UPROPERTY(EditAnywhere, BlueprintReadWrite, Category = "MusicClock")
+	bool RunPastMusicEnd = false;
+
 private:
 	// Don't let C++ access these directly! They are a blueprint convenience and only work because 
 	// they specify getter functions!
+	// NOTE These only have getters because a future refactor will get rid of them. -jwf + jake burga.
 	UPROPERTY(BlueprintGetter = GetCurrentSmoothedAudioRenderSongPos, Category = "MusicClock")
 	FMidiSongPos CurrentSmoothedAudioRenderSongPos;
 	UPROPERTY(BlueprintGetter = GetCurrentVideoRenderSongPos, Category = "MusicClock")
@@ -327,6 +397,15 @@ private:
 	UPROPERTY(BlueprintGetter = GetCurrentRawAudioRenderSongPos, Category = "MusicClock")
 	FMidiSongPos CurrentRawAudioRenderSongPos;
 
+	UPROPERTY(Transient, VisibleAnywhere, BlueprintGetter=GetCurrentClockAdvanceRate, Category = "MusicClock")
+	float CurrentClockAdvanceRate = 0.0f;
+
+	float CurrentTempo = 0.0f;
+
+	float CurrentTimeSigNum = 0;
+
+	float CurrentTimeSigDenom = 0;
+
 public:
 	// Getter functions for the Blueprint properties exposed above...
 	UFUNCTION(BlueprintCallable, BlueprintPure, Category = "MusicClock")
diff --git a/Engine/Plugins/Runtime/Harmonix/Source/HarmonixMidi/Public/HarmonixMidi/MidiSongPos.h b/Engine/Plugins/Runtime/Harmonix/Source/HarmonixMidi/Public/HarmonixMidi/MidiSongPos.h
--- a/Engine/Plugins/Runtime/Harmonix/Source/HarmonixMidi/Public/HarmonixMidi/MidiSongPos.h
+++ b/Engine/Plugins/Runtime/Harmonix/Source/HarmonixMidi/Public/HarmonixMidi/MidiSongPos.h
@@ -19,22 +19,43 @@ struct HARMONIXMIDI_API FMidiSongPos
 	GENERATED_BODY()
 
 public:
+	/** Total seconds from bar 1 beat 1. (negative for count-in and/or pickup bars) */
 	UPROPERTY(BlueprintReadOnly, Category = "MidiSongPosition")
-	float SecondsFromBarOne = 0.0f; // total seconds from bar 1 beat 1. (negative for count-in and/or pickup bars)
+	float SecondsFromBarOne = 0.0f;
+
+	/** Total seconds from the beginning of the musical content (ie. includes all count-in and pickup bars) */
 	UPROPERTY(BlueprintReadOnly, Category = "MidiSongPosition")
-	float SecondsIncludingCountIn = 0.0f; // total seconds from the beginning of the musical content (ie. includes all count-in and pickup bars)
+	float SecondsIncludingCountIn = 0.0f;
+
+	/** Time Signature Numerator (Indicates Number of Beats Per Bar) */
 	UPROPERTY(BlueprintReadOnly, Category = "MidiSongPosition")
-	int32 TimeSigNumerator = 4;
+	int32 TimeSigNumerator = 0;
+
+	/** Time Signature Numerator (Indicates Subdivision of Bar that counts as 1 beat (4 means Quarter Note = One Beat)) */
+	UPROPERTY(BlueprintReadOnly, Category = "MidiSongPosition")
+	int32 TimeSigDenominator = 0;
+
+	/** Tempo in Quarter Notes Per Minute (BPM) */
 	UPROPERTY(BlueprintReadOnly, Category = "MidiSongPosition")
-	int32 TimeSigDenominator = 4;
+	float Tempo = 0.0f;
+
+	/** Total bars from the beginning of the song. */
 	UPROPERTY(BlueprintReadOnly, Category = "MidiSongPosition")
-	float Tempo = 120.0f;	// Tempo in beats per minute.
+	float BarsIncludingCountIn = 0.0f;
+
+	/** Total beats from the beginning of the song. */
 	UPROPERTY(BlueprintReadOnly, Category = "MidiSongPosition")
-	float BarsIncludingCountIn = 0.0f; // total bars from the beginning of the song.
+	float BeatsIncludingCountIn = 0.0f;
+
+	/**
+	 * Whether this beat is a Downbeat (beat 1 of the current bar) or otherwise
+	 * This value is determined by the authored content
+	 * Defaults to "Normal"
+	 */
 	UPROPERTY(BlueprintReadOnly, Category = "MidiSongPosition")
-	float BeatsIncludingCountIn = 0.0f; // total beats from the beginning of the song.
-	UPROPERTY(BlueprintReadOnly, Category	 = "MidiSongPosition")
 	EMusicalBeatType BeatType = EMusicalBeatType::Normal;
+
+	/** The Musical Timestamp: Bar Number and Beat Number within the bar (Bar 1 Beat 1.0f) */
 	UPROPERTY(BlueprintReadOnly, Category = "MidiSongPosition")
 	FMusicTimestamp Timestamp;
 
@@ -65,9 +86,9 @@ public:
 	{
 		SecondsFromBarOne = 0.0f;
 		SecondsIncludingCountIn = 0.0f;
-		TimeSigNumerator = 4;
-		TimeSigDenominator = 4;
-		Tempo = 120.0f;
+		TimeSigNumerator = 0;
+		TimeSigDenominator = 0;
+		Tempo = 0.0f;
 		Timestamp.Reset();
 		BarsIncludingCountIn = 0.0f;
 		BeatsIncludingCountIn = 0.0f;
@@ -91,7 +112,7 @@ public:
 	static FMidiSongPos Lerp(const FMidiSongPos& A, const FMidiSongPos& B, float Alpha);
 
 	void SetByTimeAndTick(float Ms, float Tick, const ISongMapEvaluator& Map);
-	
+
 private:
 	// Has anyone explicitly set this, or is this just the defaults.
 	bool IsSet = false;
