[Harmonix] Don't warn about PreferredReimportPath extension if it was never set. Fixes bug where reimport options in the Content Browser do nothing.

#tests Editor - both "Reimport with New File" and "Reimport" options on an existing MidiFile asset
#rnx

diff --git a/Engine/Plugins/Runtime/Harmonix/Source/HarmonixDspEditor/Private/FusionPatchAssetFactory.cpp b/Engine/Plugins/Runtime/Harmonix/Source/HarmonixDspEditor/Private/FusionPatchAssetFactory.cpp
--- a/Engine/Plugins/Runtime/Harmonix/Source/HarmonixDspEditor/Private/FusionPatchAssetFactory.cpp
+++ b/Engine/Plugins/Runtime/Harmonix/Source/HarmonixDspEditor/Private/FusionPatchAssetFactory.cpp
@@ -48,13 +48,13 @@ bool UFusionPatchAssetFactory::CanReimport(UObject* Obj, TArray<FString>& OutFil
 	{
 		// the PreferredReimportPath is the path of the new file, which can differ from the original SrcFilePath
 		// if the file has the same name but a different extension, don't import it!
-		if (!IsSupportedFileExtension(PreferredReimportPath))
+		if (!PreferredReimportPath.IsEmpty() && !IsSupportedFileExtension(PreferredReimportPath))
 		{
 			UE_LOG(LogFusionPatchAssetFactory, Warning, TEXT("%s: Failed to reimport with new file. Invalid extension: %s"), *FusionPatch->GetPathName(), *FPaths::GetExtension(PreferredReimportPath));
 			return false;
 		}
 
-		if  (FusionPatch->AssetImportData)
+		if (FusionPatch->AssetImportData)
 		{
 			FusionPatch->AssetImportData->ExtractFilenames(OutFilenames);
 		}
diff --git a/Engine/Plugins/Runtime/Harmonix/Source/HarmonixMidiEditor/Private/MidiFileFactory.cpp b/Engine/Plugins/Runtime/Harmonix/Source/HarmonixMidiEditor/Private/MidiFileFactory.cpp
--- a/Engine/Plugins/Runtime/Harmonix/Source/HarmonixMidiEditor/Private/MidiFileFactory.cpp
+++ b/Engine/Plugins/Runtime/Harmonix/Source/HarmonixMidiEditor/Private/MidiFileFactory.cpp
@@ -77,7 +77,7 @@ bool UMidiFileFactory::CanReimport(UObject* Obj, TArray<FString>& OutFilenames)
 
 		// the PreferredReimportPath is the path of the new file, which can differ from the original SrcFilePath
 		// if the file has the same name but a different extension, don't import it!
-		if (!IsSupportedFileExtension(PreferredReimportPath))
+		if (!PreferredReimportPath.IsEmpty() && !IsSupportedFileExtension(PreferredReimportPath))
 		{
 			UE_LOG(LogHarmonixMidiEditor, Warning, TEXT("%s: Failed to reimport with new file. Invalid extension: %s"), *AsMidiFile->GetPathName(), *FPaths::GetExtension(PreferredReimportPath));
 			return false;
