Fix failing Test in MidiFileImportExportTests. Expected warnings have changed

#tests HarmonixTests
#rnx

diff --git a/Engine/Plugins/Runtime/Harmonix/Source/HarmonixMidiTests/Private/MidiFileImportExportTests.cpp b/Engine/Plugins/Runtime/Harmonix/Source/HarmonixMidiTests/Private/MidiFileImportExportTests.cpp
--- a/Engine/Plugins/Runtime/Harmonix/Source/HarmonixMidiTests/Private/MidiFileImportExportTests.cpp
+++ b/Engine/Plugins/Runtime/Harmonix/Source/HarmonixMidiTests/Private/MidiFileImportExportTests.cpp
@@ -189,7 +189,6 @@ namespace HarmonixMidiTests::ImportExportTests
 		ADD_EXPECTED_MIDI_ERROR(TEXT("MIDI track data length exceeds maximum length!"));
 		ADD_EXPECTED_MIDI_ERROR(TEXT("MIDI track 0 data length as recorded in the header exceeds the amount of data in the file!"));
 		ADD_EXPECTED_MIDI_ERROR(TEXT("MIDI track 0 ... Track data underrun"));
-		ADD_EXPECTED_MIDI_ERROR(TEXT("0 byte long string found in MIDI data"));
 		ADD_EXPECTED_MIDI_ERROR(TEXT("Found invalid MIDI status byte"));
 		ADD_EXPECTED_MIDI_ERROR(TEXT("Unsupported MIDI tempo encountered"));
 		ADD_EXPECTED_MIDI_ERROR(TEXT("Failed to add tempo to Track"));
