MusicClockComponent Start & Stop now simply call Activate and Deactivate, which is where are the starting and stopping logic can happen.

This allows us to do the following:
- `bAutoActivate` on the component will `Auto Start` the clock
- You can "Start" the clock outside of PIE without having to worry about BeginPlay/EndPlay lifetime

#tests Harmonix Functional Tests
#rnx

diff --git a/Engine/Plugins/Runtime/Harmonix/Source/HarmonixMetasound/Private/HarmonixMetasound/Components/MetasoundMusicClockDriver.cpp b/Engine/Plugins/Runtime/Harmonix/Source/HarmonixMetasound/Private/HarmonixMetasound/Components/MetasoundMusicClockDriver.cpp
--- a/Engine/Plugins/Runtime/Harmonix/Source/HarmonixMetasound/Private/HarmonixMetasound/Components/MetasoundMusicClockDriver.cpp
+++ b/Engine/Plugins/Runtime/Harmonix/Source/HarmonixMetasound/Private/HarmonixMetasound/Components/MetasoundMusicClockDriver.cpp
@@ -8,6 +8,7 @@
 #include "Engine/World.h"
 #include "Async/Async.h"
 #include "Harmonix.h"
+#include "MetasoundSource.h"
 
 namespace MetasoundMusicClockDriver
 {
@@ -313,8 +314,13 @@ bool FMetasoundMusicClockDriver::AttemptToConnectToAudioComponentsMetasound()
 		return false;
 	}
 
-	DetachAllCallbacks();
+	// we have an audio component, but it's not set to play a Metasound (yet)
+	if (!Cast<UMetaSoundSource>(AudioComponentToWatch->GetSound()))
+	{
+		return false;
+	}
 
+	DetachAllCallbacks();
 	CurrentGeneratorHandle.Reset(UMetasoundGeneratorHandle::CreateMetaSoundGeneratorHandle(AudioComponentToWatch.Get()));
 	if (!CurrentGeneratorHandle)
 	{
diff --git a/Engine/Plugins/Runtime/Harmonix/Source/HarmonixMetasound/Private/HarmonixMetasound/Components/MusicClockComponent.cpp b/Engine/Plugins/Runtime/Harmonix/Source/HarmonixMetasound/Private/HarmonixMetasound/Components/MusicClockComponent.cpp
--- a/Engine/Plugins/Runtime/Harmonix/Source/HarmonixMetasound/Private/HarmonixMetasound/Components/MusicClockComponent.cpp
+++ b/Engine/Plugins/Runtime/Harmonix/Source/HarmonixMetasound/Private/HarmonixMetasound/Components/MusicClockComponent.cpp
@@ -3,10 +3,10 @@
 #include "HarmonixMetasound/Components/MetasoundMusicClockDriver.h"
 #include "HarmonixMetasound/Components/WallClockMusicClockDriver.h"
 #include "HarmonixMetasound/Subsystems/MidiClockUpdateSubsystem.h"
-#include "Harmonix.h"
 #include "HarmonixMidi/MusicTimeSpan.h"
 #include "Components/AudioComponent.h"
 #include "Templates/SharedPointer.h"
+#include "GameFramework/Actor.h"
 
 #include UE_INLINE_GENERATED_CPP_BY_NAME(MusicClockComponent)
 
@@ -16,10 +16,7 @@
 {
 	MakeDefaultSongMap();
 	PrimaryComponentTick.bCanEverTick = true;
-	// We want the music player and clocks to tick before other components, which default to the TG_DuringPhysics group.
-	// Though, this means any other TG_PrePhysics group actors or components that care about accurate song time will
-	// want to manually add a tick prerequisite on this component.
-	PrimaryComponentTick.TickGroup = TG_PrePhysics;
+	bAutoActivate = false;
 }
 
 UMusicClockComponent* UMusicClockComponent::CreateMetasoundDrivenMusicClock(UObject* WorldContextObject, UAudioComponent* InAudioComponent, FName MetasoundOutputPinName, bool Start)
@@ -241,29 +238,80 @@
 	return Result;
 }
 
-void UMusicClockComponent::BeginPlay()
+void UMusicClockComponent::Activate(bool bReset)
 {
-	BeginPlayInternal();
-	Super::BeginPlay();
-}
+	Super::Activate(bReset);
 
-void UMusicClockComponent::BeginPlayInternal()
-{
+	// if Activate is forcing a reset, let's take that as a cue to restart the clock
+	if (bReset && State != EMusicClockState::Stopped)
+	{
+		State = EMusicClockState::Stopped;
+		if (ClockDriver)
+		{
+			ClockDriver->OnStop();
+			DisconnectFromClockDriver();
+		}
+		PlayStateEvent.Broadcast(State);
+	}
+	
+	if (State == EMusicClockState::Running)
+	{
+		return;
+	}
+
+	UMidiClockUpdateSubsystem::TrackMusicClockComponent(this);
+	State = EMusicClockState::Running;
+	MakeDefaultSongMap();
+
+	if (!MetasoundsAudioComponent)
+	{
+		MetasoundsAudioComponent = GetOwner()->FindComponentByClass<UAudioComponent>();
+	}
+
 	if (!ClockDriver)
 	{
 		CreateClockDriver();
 	}
-	UMidiClockUpdateSubsystem::TrackMusicClockComponent(this);
+	ClockDriver->OnStart();
+	LastBroadcastBeat = -1;
+	LastBroadcastBar  = -1;
+	CurrentSmoothedAudioRenderSongPos.Reset();
+	CurrentVideoRenderSongPos.Reset();
+	CurrentPlayerExperiencedSongPos.Reset();
+	CurrentRawAudioRenderSongPos.Reset();
+	PlayStateEvent.Broadcast(State);
 }
 
+void UMusicClockComponent::Deactivate()
+{
+	Super::Deactivate();
+	if (State == EMusicClockState::Stopped)
+	{
+		return;
+	}
+	State = EMusicClockState::Stopped;
+	if (ClockDriver)
+	{
+		ClockDriver->OnStop();
+	}
+	CurrentSmoothedAudioRenderSongPos.Reset();
+	CurrentVideoRenderSongPos.Reset();
+	CurrentPlayerExperiencedSongPos.Reset();
+	CurrentRawAudioRenderSongPos.Reset();
+	PlayStateEvent.Broadcast(State);
+	
+	UMidiClockUpdateSubsystem::StopTrackingMusicClockComponent(this);
+	DisconnectFromClockDriver();
+}
+
+void UMusicClockComponent::BeginPlay()
+{
+	Super::BeginPlay();
+}
+
 void UMusicClockComponent::EndPlay(const EEndPlayReason::Type EndPlayReason)
 {
 	Super::EndPlay(EndPlayReason);
-	EndPlayInternal();
-}
-
-void UMusicClockComponent::EndPlayInternal()
-{
 	UMidiClockUpdateSubsystem::StopTrackingMusicClockComponent(this);
 	DisconnectFromClockDriver();
 }
@@ -309,20 +357,7 @@
 
 void UMusicClockComponent::Start()
 {
-	MakeDefaultSongMap();
-	if (!ClockDriver)
-	{
-		CreateClockDriver();
-	}
-	ClockDriver->OnStart();
-	LastBroadcastBeat = -1;
-	LastBroadcastBar  = -1;
-	CurrentSmoothedAudioRenderSongPos.Reset();
-	CurrentVideoRenderSongPos.Reset();
-	CurrentPlayerExperiencedSongPos.Reset();
-	CurrentRawAudioRenderSongPos.Reset();
-	State = EMusicClockState::Running;
-	PlayStateEvent.Broadcast(State);
+	Activate(false);
 }
 
 void UMusicClockComponent::Pause()
@@ -359,16 +394,7 @@
 
 void UMusicClockComponent::Stop()
 {
-	if (ClockDriver)
-	{
-		ClockDriver->OnStop();
-	}
-	State = EMusicClockState::Stopped;
-	CurrentSmoothedAudioRenderSongPos.Reset();
-	CurrentVideoRenderSongPos.Reset();
-	CurrentPlayerExperiencedSongPos.Reset();
-	CurrentRawAudioRenderSongPos.Reset();
-	PlayStateEvent.Broadcast(State);
+	Deactivate();
 }
 
 float UMusicClockComponent::GetSecondsIncludingCountIn(ECalibratedMusicTimebase Timebase) const
diff --git a/Engine/Plugins/Runtime/Harmonix/Source/HarmonixMetasound/Private/HarmonixMetasound/MusicEnvironmentSupport/HarmonixMusicEnvironmentMetronome.cpp b/Engine/Plugins/Runtime/Harmonix/Source/HarmonixMetasound/Private/HarmonixMetasound/MusicEnvironmentSupport/HarmonixMusicEnvironmentMetronome.cpp
--- a/Engine/Plugins/Runtime/Harmonix/Source/HarmonixMetasound/Private/HarmonixMetasound/MusicEnvironmentSupport/HarmonixMusicEnvironmentMetronome.cpp
+++ b/Engine/Plugins/Runtime/Harmonix/Source/HarmonixMetasound/Private/HarmonixMetasound/MusicEnvironmentSupport/HarmonixMusicEnvironmentMetronome.cpp
@@ -157,7 +157,7 @@ void UHarmonixMusicEnvironmentMetronome::BeginDestroy()
 {
 	if (MusicClockComponent)
 	{
-		MusicClockComponent->EndPlayInternal();
+		MusicClockComponent->Stop();
 	}
 	UObject::BeginDestroy();
 }
@@ -235,7 +235,6 @@ bool UHarmonixMusicEnvironmentMetronome::BuildAndStartMetasound(UWorld* InWorld)
 	MusicClockComponent->MetasoundOutputName = TEXT("MIDI Clock");
 	MusicClockComponent->ConnectToMetasoundOnAudioComponent(AudioComponent);
 	MusicClockComponent->Start();
-	MusicClockComponent->BeginPlayInternal();
 
 	if (MusicMap)
 	{
diff --git a/Engine/Plugins/Runtime/Harmonix/Source/HarmonixMetasound/Public/HarmonixMetasound/Components/MusicClockComponent.h b/Engine/Plugins/Runtime/Harmonix/Source/HarmonixMetasound/Public/HarmonixMetasound/Components/MusicClockComponent.h
--- a/Engine/Plugins/Runtime/Harmonix/Source/HarmonixMetasound/Public/HarmonixMetasound/Components/MusicClockComponent.h
+++ b/Engine/Plugins/Runtime/Harmonix/Source/HarmonixMetasound/Public/HarmonixMetasound/Components/MusicClockComponent.h
@@ -59,32 +59,58 @@ public:
 	// Sets default values for this component's properties
 	UMusicClockComponent();
 
+	/**
+	 * Start this clock ticking and start tracking music time
+	 * Same as "Start()"
+	 *
+	 * Setting bReset to true will restart (trigger "Stop" and then "Start") the clock is already Running
+	 */
+	void Activate(bool bReset) override;
+
+	/**
+	 * Stop this clock ticking and stop tracking music time
+	 * Same as "Stop()"
+	 */
+	void Deactivate() override;
+	
 	void BeginPlay() override;
-	void BeginPlayInternal();
 	void EndPlay(const EEndPlayReason::Type EndPlayReason) override;
-	void EndPlayInternal();
-
-	UPROPERTY(EditDefaultsOnly, Category = "MusicClock")
+	
+	UPROPERTY(EditAnywhere, Category = "MusicClock")
 	EMusicClockDriveMethod DriveMethod = EMusicClockDriveMethod::MetaSound;
 
-	UPROPERTY(EditDefaultsOnly, Category = "MusicClock", meta = (EditCondition = "DriveMethod == EMusicClockDriveMethod::MetaSound"))
+	UPROPERTY(EditAnywhere, Category = "MusicClock", meta = (EditCondition = "DriveMethod == EMusicClockDriveMethod::MetaSound"))
 	FName MetasoundOutputName = "MIDI Clock";
 
-	UPROPERTY(EditAnywhere, BlueprintReadWrite, Category = "MusicClock", meta = (EditCondition = "DriveMethod == EMusicClockDriveMethod::MetaSound"))
+	UPROPERTY(VisibleAnywhere, BlueprintReadWrite, Category = "MusicClock", meta = (EditCondition = "DriveMethod == EMusicClockDriveMethod::MetaSound"))
 	TObjectPtr<UAudioComponent> MetasoundsAudioComponent;
 
-	UPROPERTY(EditDefaultsOnly, Category = "MusicClock", meta = (EditCondition = "DriveMethod == EMusicClockDriveMethod::WallClock"))
+	UPROPERTY(EditAnywhere, Category = "MusicClock", meta = (EditCondition = "DriveMethod == EMusicClockDriveMethod::WallClock"))
 	TObjectPtr<UMidiFile> TempoMap;
 
+	/**
+	 * Start the clock ticking and tracking music time
+	 * Same as "Activate"
+	 */
 	UFUNCTION(BlueprintCallable, Category = "MusicClock")
 	void Start();
 
+	/**
+	 * Pause the clock at the current running time
+	 */
 	UFUNCTION(BlueprintCallable, Category = "MusicClock")
 	void Pause();
 
+	/**
+	 * Continue the clock from paused
+	 */
 	UFUNCTION(BlueprintCallable, Category = "MusicClock")
 	void Continue();
 
+	/**
+	 * Stop this clock ticking and no longer track music time
+	 * Same as "Deactivate"
+	 */
 	UFUNCTION(BlueprintCallable, Category = "MusicClock")
 	void Stop();
 
@@ -375,10 +401,10 @@ protected:
 	UPROPERTY(EditAnywhere, BlueprintReadWrite, Category = "MusicClock|Defaults")
 	int DefaultTimeSignatureDenom = 4;
 
-	UPROPERTY(Transient, VisibleAnywhere, BlueprintGetter=GetCurrentBeatDurationSeconds, Category = "MusicClock")
+	UPROPERTY(Transient, BlueprintGetter=GetCurrentBeatDurationSeconds, Category = "MusicClock")
 	float CurrentBeatDurationSec = 0.0f;
 
-	UPROPERTY(Transient, VisibleAnywhere, BlueprintGetter=GetCurrentBarDurationSeconds, Category = "MusicClock")
+	UPROPERTY(Transient, BlueprintGetter=GetCurrentBarDurationSeconds, Category = "MusicClock")
 	float CurrentBarDurationSec = 0.0f;
 
 	UPROPERTY(EditAnywhere, BlueprintReadWrite, Category = "MusicClock")
@@ -397,7 +423,7 @@ private:
 	UPROPERTY(BlueprintGetter = GetCurrentRawAudioRenderSongPos, Category = "MusicClock")
 	FMidiSongPos CurrentRawAudioRenderSongPos;
 
-	UPROPERTY(Transient, VisibleAnywhere, BlueprintGetter=GetCurrentClockAdvanceRate, Category = "MusicClock")
+	UPROPERTY(Transient, BlueprintGetter=GetCurrentClockAdvanceRate, Category = "MusicClock")
 	float CurrentClockAdvanceRate = 0.0f;
 
 	float CurrentTempo = 0.0f;
