Music clock component changes pt 2.

#rnx

diff --git a/Engine/Plugins/Runtime/Harmonix/Source/HarmonixMetasound/Private/HarmonixMetasound/Components/MusicClockComponent.cpp b/Engine/Plugins/Runtime/Harmonix/Source/HarmonixMetasound/Private/HarmonixMetasound/Components/MusicClockComponent.cpp
--- a/Engine/Plugins/Runtime/Harmonix/Source/HarmonixMetasound/Private/HarmonixMetasound/Components/MusicClockComponent.cpp
+++ b/Engine/Plugins/Runtime/Harmonix/Source/HarmonixMetasound/Private/HarmonixMetasound/Components/MusicClockComponent.cpp
@@ -240,9 +240,15 @@ FMidiSongPos UMusicClockComponent::CalculateSongPosWithOffset(float MsOffset, EC
 
 void UMusicClockComponent::Activate(bool bReset)
 {
+	// cache the value of ShouldActivate before calling the super
+	bool bShouldActivate = bReset || ShouldActivate();
 	Super::Activate(bReset);
 
-	// if Activate is forcing a reset, let's take that as a cue to restart the clock
+	if (!bShouldActivate)
+	{
+		return;
+	}
+	
 	if (bReset && State != EMusicClockState::Stopped)
 	{
 		State = EMusicClockState::Stopped;
@@ -263,7 +269,7 @@ void UMusicClockComponent::Activate(bool bReset)
 	State = EMusicClockState::Running;
 	MakeDefaultSongMap();
 
-	if (!MetasoundsAudioComponent)
+	if (!MetasoundsAudioComponent && GetOwner())
 	{
 		MetasoundsAudioComponent = GetOwner()->FindComponentByClass<UAudioComponent>();
 	}
@@ -284,7 +290,14 @@ void UMusicClockComponent::Activate(bool bReset)
 
 void UMusicClockComponent::Deactivate()
 {
+	bool bShouldDeactivate = IsActive();
 	Super::Deactivate();
+
+	if (!bShouldDeactivate)
+	{
+		return;
+	}
+	
 	if (State == EMusicClockState::Stopped)
 	{
 		return;
@@ -316,6 +329,13 @@ void UMusicClockComponent::EndPlay(const EEndPlayReason::Type EndPlayReason)
 	DisconnectFromClockDriver();
 }
 
+void UMusicClockComponent::BeginDestroy()
+{
+	Super::BeginDestroy();
+	UMidiClockUpdateSubsystem::StopTrackingMusicClockComponent(this);
+	DisconnectFromClockDriver();
+}
+
 void UMusicClockComponent::TickComponent(float DeltaTime, enum ELevelTick TickType, FActorComponentTickFunction* ThisTickFunction)
 {
 	Super::TickComponent(DeltaTime, TickType, ThisTickFunction);
diff --git a/Engine/Plugins/Runtime/Harmonix/Source/HarmonixMetasound/Public/HarmonixMetasound/Components/MusicClockComponent.h b/Engine/Plugins/Runtime/Harmonix/Source/HarmonixMetasound/Public/HarmonixMetasound/Components/MusicClockComponent.h
--- a/Engine/Plugins/Runtime/Harmonix/Source/HarmonixMetasound/Public/HarmonixMetasound/Components/MusicClockComponent.h
+++ b/Engine/Plugins/Runtime/Harmonix/Source/HarmonixMetasound/Public/HarmonixMetasound/Components/MusicClockComponent.h
@@ -75,6 +75,7 @@ public:
 	
 	void BeginPlay() override;
 	void EndPlay(const EEndPlayReason::Type EndPlayReason) override;
+	void BeginDestroy() override;
 	
 	UPROPERTY(EditAnywhere, Category = "MusicClock")
 	EMusicClockDriveMethod DriveMethod = EMusicClockDriveMethod::MetaSound;
