Don't ensure if there is no ContextObject on the clock driver. We can just use the world

#rnx
#tests

diff --git a/Engine/Plugins/Runtime/Harmonix/Source/HarmonixMetasound/Private/HarmonixMetasound/Components/MusicClockDriverBase.cpp b/Engine/Plugins/Runtime/Harmonix/Source/HarmonixMetasound/Private/HarmonixMetasound/Components/MusicClockDriverBase.cpp
--- a/Engine/Plugins/Runtime/Harmonix/Source/HarmonixMetasound/Private/HarmonixMetasound/Components/MusicClockDriverBase.cpp
+++ b/Engine/Plugins/Runtime/Harmonix/Source/HarmonixMetasound/Private/HarmonixMetasound/Components/MusicClockDriverBase.cpp
@@ -26,7 +26,7 @@ double FMusicClockDriverBase::GetWallClockTime() const
 {
 	if (TStrongObjectPtr<UObject> Obj = ContextObj.Pin())
 	{
-		if (ensure(Obj->GetWorld()))
+		if (Obj->GetWorld())
 		{
 			return Obj->GetWorld()->GetTimeSeconds();
 		}
