[Harmonix] Add BP Function Library for FMidiSongPos functions

#rnx

diff --git a/Engine/Plugins/Runtime/Harmonix/Source/HarmonixMidi/Private/HarmonixMidi/Blueprint/MidiFunctionLibrary.cpp b/Engine/Plugins/Runtime/Harmonix/Source/HarmonixMidi/Private/HarmonixMidi/Blueprint/MidiFunctionLibrary.cpp
--- a/Engine/Plugins/Runtime/Harmonix/Source/HarmonixMidi/Private/HarmonixMidi/Blueprint/MidiFunctionLibrary.cpp
+++ b/Engine/Plugins/Runtime/Harmonix/Source/HarmonixMidi/Private/HarmonixMidi/Blueprint/MidiFunctionLibrary.cpp
@@ -3,6 +3,7 @@
 #include "HarmonixMidi/Blueprint/MidiFunctionLibrary.h"
 
 #include "HarmonixMidi/MidiConstants.h"
+#include "HarmonixMidi/MidiSongPos.h"
 
 #include UE_INLINE_GENERATED_CPP_BY_NAME(MidiFunctionLibrary)
 
@@ -66,3 +67,19 @@ float UMusicalTickFunctionLibrary::QuarterNoteToTick(float InQuarterNote)
 	return InQuarterNote * Harmonix::Midi::Constants::GTicksPerQuarterNote;
 }
 
+bool UMidiSongPosFunctionLibrary::IsSongPosValid(const FMidiSongPos& SongPos)
+{
+	return SongPos.IsValid();
+}
+
+FMidiSongPos UMidiSongPosFunctionLibrary::LerpSongPos(const FMidiSongPos& A, const FMidiSongPos& B, float Alpha)
+{
+	return FMidiSongPos::Lerp(A, B, Alpha);
+}
+
+FMidiSongPos UMidiSongPosFunctionLibrary::MakeSongPosFromTime(float TimeMs, float BeatsPerMinute, int32 TimeSignatureNumerator /*= 4*/, int32 TimeSignatureDenominator /*= 4*/, int32 StartBar /*= 1*/)
+{
+	FMidiSongPos Result;
+	Result.SetByTime(TimeMs, BeatsPerMinute, TimeSignatureNumerator, TimeSignatureDenominator, StartBar);
+	return Result;
+}
diff --git a/Engine/Plugins/Runtime/Harmonix/Source/HarmonixMidi/Private/HarmonixMidi/Blueprint/MidiFunctionLibrary.h b/Engine/Plugins/Runtime/Harmonix/Source/HarmonixMidi/Private/HarmonixMidi/Blueprint/MidiFunctionLibrary.h
--- a/Engine/Plugins/Runtime/Harmonix/Source/HarmonixMidi/Private/HarmonixMidi/Blueprint/MidiFunctionLibrary.h
+++ b/Engine/Plugins/Runtime/Harmonix/Source/HarmonixMidi/Private/HarmonixMidi/Blueprint/MidiFunctionLibrary.h
@@ -8,6 +8,8 @@
 
 #include "MidiFunctionLibrary.generated.h"
 
+struct FMidiSongPos;
+
 /**
 * Function library for FMidiNote and various midi note constants
 */
@@ -80,3 +82,22 @@ public:
 	UFUNCTION(BlueprintPure, Category = "Midi")
 	static float QuarterNoteToTick(float InQuarterNote);
 };
+
+/**
+* Function library for FMidiSongPos
+*/
+UCLASS()
+class HARMONIXMIDI_API UMidiSongPosFunctionLibrary : public UBlueprintFunctionLibrary
+{
+	GENERATED_BODY()
+
+public:
+	UFUNCTION(BlueprintPure, Category = "Midi", meta = (DisplayName = "Is MIDI Song Position Valid"))
+	static bool IsSongPosValid(const FMidiSongPos& SongPos);
+
+	UFUNCTION(BlueprintPure, Category = "Midi", meta = (DisplayName = "Lerp (MIDI Song Position)"))
+	static FMidiSongPos LerpSongPos(const FMidiSongPos& A, const FMidiSongPos& B, float Alpha);
+
+	UFUNCTION(BlueprintPure, Category="Midi", meta=(NativeMakeFunc, DisplayName = "Make MIDI Song Position from Time"))
+	static FMidiSongPos MakeSongPosFromTime(float TimeMs, float BeatsPerMinute, int32 TimeSignatureNumerator = 4, int32 TimeSignatureDenominator = 4, int32 StartBar = 1);
+};
\ No newline at end of file
