Calculate Beats from BarMaps if there are no BeatMaps

#tests Harmonix Automation Tests
#rnx

diff --git a/Engine/Plugins/Runtime/Harmonix/Source/HarmonixMidi/Private/HarmonixMidi/BarMap.cpp b/Engine/Plugins/Runtime/Harmonix/Source/HarmonixMidi/Private/HarmonixMidi/BarMap.cpp
--- a/Engine/Plugins/Runtime/Harmonix/Source/HarmonixMidi/Private/HarmonixMidi/BarMap.cpp
+++ b/Engine/Plugins/Runtime/Harmonix/Source/HarmonixMidi/Private/HarmonixMidi/BarMap.cpp
@@ -334,6 +334,26 @@ FMusicTimestamp FBarMap::TickFromBarOneToMusicTimestamp(float InTickFromBarOne,
 	return TickToMusicTimestamp(InTickFromBarOne + GetTickOfBarOne(), OutBeatsPerBar);
 }
 
+float FBarMap::TickToFractionalBeatIncludingCountIn(float Tick) const
+{
+	if (Points.IsEmpty())
+	{
+		// assume 4/4
+		return Tick / (TicksPerQuarterNote);
+	}
+	
+	int32 PointIndex = FMusicMapUtl::GetPointIndexForTick(Points, int32(Tick));
+	if (!Points.IsValidIndex(PointIndex))
+	{
+		return 0.0f;
+	}
+
+	float TicksPassed = Tick - Points[PointIndex].StartTick;
+	int32 TicksPerBeat = GetTicksInBeatAfterPoint(PointIndex);
+	float BeatsPassed = TicksPassed / TicksPerBeat;
+	return  Points[PointIndex].BeatIndex + BeatsPassed;
+}
+
 float FBarMap::TickToFractionalBarIncludingCountIn(float Tick) const
 {
 	if (Points.IsEmpty())
@@ -371,6 +391,26 @@ int32 FBarMap::TickToBarIncludingCountIn(int32 Tick) const
 	return Points[TimeSigIndex].BarIndex + BarsPassed;
 }
 
+int32 FBarMap::TickToBeatIncludingCountIn(float Tick) const
+{
+	if (Points.IsEmpty())
+	{
+		// assume 4/4
+		return Tick / (TicksPerQuarterNote);
+	}
+	
+	int32 PointIndex = FMusicMapUtl::GetPointIndexForTick(Points, int32(Tick));
+	if (!Points.IsValidIndex(PointIndex))
+	{
+		return 0;
+	}
+
+	int32 TicksPassed = Tick - Points[PointIndex].StartTick;
+	int32 TicksPerBeat = GetTicksInBeatAfterPoint(PointIndex);
+	int32 BeatsPassed = TicksPassed / TicksPerBeat;
+	return Points[PointIndex].BeatIndex + BeatsPassed;
+}
+
 float FBarMap::FractionalBarIncludingCountInToTick(float Bar) const
 {
 	if (Points.IsEmpty())
@@ -397,6 +437,32 @@ float FBarMap::FractionalBarIncludingCountInToTick(float Bar) const
 	return Points[TimeSigIndex].StartTick + BarsPassed * TicksPerBar;
 }
 
+float FBarMap::FractionalBeatIncludingCountInToTick(float Beat) const
+{
+	if (Points.IsEmpty())
+	{
+		// assume 4/4
+		return Beat * TicksPerQuarterNote;
+	}
+
+	int32 BeatIndex = FMath::FloorToInt32(Beat);
+	int32 TimeSigIndex = Algo::UpperBound(Points, BeatIndex, FTimeSignaturePoint::BeatLessThan());
+	
+	if (TimeSigIndex > 0)
+	{
+		--TimeSigIndex;
+	}
+
+	if (!Points.IsValidIndex(TimeSigIndex))
+    {
+    	TimeSigIndex = 0;
+    }
+	
+	float BeatsPassed = Beat - Points[TimeSigIndex].BeatIndex;
+	int32 TicksPerBeat = GetTicksInBeatAfterPoint(TimeSigIndex);
+	return Points[TimeSigIndex].StartTick + BeatsPassed * TicksPerBeat;
+}
+
 bool FBarMap::AddTimeSignatureAtMusicTimestampBar(int32 BarNumber, int32 Numerator, int32 Denominator, bool SortNow, bool FailOnError)
 {
 	int BarIndex = BarNumber - StartBar;
diff --git a/Engine/Plugins/Runtime/Harmonix/Source/HarmonixMidi/Private/HarmonixMidi/SongMaps.cpp b/Engine/Plugins/Runtime/Harmonix/Source/HarmonixMidi/Private/HarmonixMidi/SongMaps.cpp
--- a/Engine/Plugins/Runtime/Harmonix/Source/HarmonixMidi/Private/HarmonixMidi/SongMaps.cpp
+++ b/Engine/Plugins/Runtime/Harmonix/Source/HarmonixMidi/Private/HarmonixMidi/SongMaps.cpp
@@ -358,7 +358,12 @@ float ISongMapEvaluator::GetSongLengthMs() const
 int32 ISongMapEvaluator::GetSongLengthBeats() const
 {
 	const FBeatMap& BeatMap = GetBeatMap();
-	return BeatMap.GetNumMapPoints();
+	if (!BeatMap.IsEmpty())
+	{
+		return BeatMap.GetNumMapPoints();
+	}
+
+	return FMath::Floor(GetFractionalBeatAtTick(GetSongLengthData().LastTick) - GetFractionalBeatAtTick(0));
 }
 
 float ISongMapEvaluator::GetSongLengthFractionalBars() const
@@ -748,12 +753,29 @@ float ISongMapEvaluator::GetFractionalBeatAtMs(float Ms) const
 float ISongMapEvaluator::GetFractionalBeatAtTick(float Tick) const
 {
 	const FBeatMap& BeatMap = GetBeatMap();
+	const FBarMap& BarMap = GetBarMap();
+	if (BeatMap.IsEmpty() && BarMap.IsEmpty())
+	{
+		return 1.0f;
+	}
+	
+	if (BeatMap.IsEmpty())
+	{
+		if (BarMap.GetStartBar() == 1)
+		{
+			return BarMap.TickToFractionalBeatIncludingCountIn(Tick) + 1.0f;
+		}
+		else
+		{
+			return BarMap.TickToFractionalBeatIncludingCountIn(Tick) - BarMap.TickToFractionalBeatIncludingCountIn(BarMap.GetTickOfBarOne()) + 1.0f;
+		}
+	}
+
 	int32 BeatIndex = GetBeatIndexAtTick(int32(Tick));
-	if (BeatIndex < 0 || BeatMap.IsEmpty())
+	if (BeatIndex < 0 )
 	{
 		return 1.0f; // 1 based position!
 	}
-
 	const FBeatMapPoint* BeatInfo = &BeatMap.GetBeatPointInfo(BeatIndex);
 	check(BeatInfo);
 	float TickInBeat = Tick - BeatInfo->StartTick;
@@ -770,11 +792,17 @@ int32 ISongMapEvaluator::GetBeatIndexAtMs(float Ms) const
 int32 ISongMapEvaluator::GetBeatIndexAtTick(int32 Tick) const
 {
 	const FBeatMap& BeatMap = GetBeatMap();
-	if (BeatMap.IsEmpty())
+	if (!BeatMap.IsEmpty())
 	{
-		return -1;
+		return BeatMap.GetPointIndexForTick(Tick);
+
+	}
+	const FBarMap& BarMap = GetBarMap();
+	if (!BarMap.IsEmpty())
+	{
+		return FMath::Floor(BarMap.TickToFractionalBeatIncludingCountIn(Tick));
 	}
-	return BeatMap.GetPointIndexForTick(Tick);
+	return INDEX_NONE;
 }
 
 EMusicalBeatType ISongMapEvaluator::GetBeatTypeAtMs(float Ms) const
diff --git a/Engine/Plugins/Runtime/Harmonix/Source/HarmonixMidi/Public/HarmonixMidi/BarMap.h b/Engine/Plugins/Runtime/Harmonix/Source/HarmonixMidi/Public/HarmonixMidi/BarMap.h
--- a/Engine/Plugins/Runtime/Harmonix/Source/HarmonixMidi/Public/HarmonixMidi/BarMap.h
+++ b/Engine/Plugins/Runtime/Harmonix/Source/HarmonixMidi/Public/HarmonixMidi/BarMap.h
@@ -156,6 +156,13 @@ public:
 		bool operator()(const FTimeSignaturePoint& Point, int32 InBarIndex) const { return Point.BarIndex < InBarIndex; }
 		bool operator()(const FTimeSignaturePoint& PointA, const FTimeSignaturePoint& PointB) const { return PointA.BarIndex < PointB.BarIndex; }
 	};
+
+	struct BeatLessThan
+	{
+		bool operator()(int32 InBeatIndex, const FTimeSignaturePoint& Point) const { return InBeatIndex < Point.BeatIndex; }
+		bool operator()(const FTimeSignaturePoint& Point, int32 InBeatIndex) const { return Point.BeatIndex < InBeatIndex; }
+		bool operator()(const FTimeSignaturePoint& PointA, const FTimeSignaturePoint& PointB) const { return PointA.BeatIndex < PointB.BeatIndex; }
+	};
 };
 
 /** A map of time signatures changes in a song. */
@@ -220,7 +227,6 @@ public:
 	 */
 	FMusicTimestamp TickFromBarOneToMusicTimestamp(float Tick, int32* OutBeatsPerBar = nullptr) const;
 
-
 	/**
 	 * Takes a "music timestamp" bar and computes the absolute (raw) total tick from the beginning of the
 	 * midi data. Optionally returns the number of beats in that bar and the number of ticks per beat.
@@ -299,6 +305,17 @@ public:
 	 */
 	float TickToFractionalBarIncludingCountIn(float Tick) const;
 
+	/**
+	 * Takes an absolute tick (from the beginning of the midi data) and calculates the fractional beat,
+	 * where beat 0.0 is the beginning of the raw midi data.
+	 *
+	 * Uses the collection of time signatures and the setting of kTicksPerQuarterNoteInt
+	 *
+	 * @param Tick 0-based 'raw total ticks'
+	 * @return The 0-based fractional beat index
+	 */
+	float TickToFractionalBeatIncludingCountIn(float Tick) const;
+
 	/**
 	 * Takes an absolute tick (from the beginning of the midi data) and calculates the bar index,
 	 * where bar 0 is the beginning of the raw midi data.
@@ -310,6 +327,17 @@ public:
 	 */
 	int32 TickToBarIncludingCountIn(int32 Tick) const;
 
+	/**
+	 * Takes an absolute tick (from the beginning of the midi data) and calculates the beat index,
+	 * where beat 0 is the beginning of the raw midi data.
+	 *
+	 * Uses the collection of time signatures and the setting of kTicksPerQuarterNoteInt
+	 *
+	 * @param Tick 0-based 'raw total ticks'
+	 * @return The 0-based integral beat index
+	 */
+	int32 TickToBeatIncludingCountIn(float Tick) const;
+
 	/**
 	* Takes a fractional bar and returns the absolute tick. Note: Bar 0.0 is the beginning
 	* of raw midi data. Uses the collection of time signatures and the setting of kTicksPerQuarterNoteInt
@@ -322,6 +350,18 @@ public:
 	*/
 	float FractionalBarIncludingCountInToTick(float FractionalBarIndex) const;
 
+	/**
+	* Takes a fractional beat and returns the absolute tick. Note: Beat 0.0 is the beginning
+	* of raw midi data. Uses the collection of time signatures and the setting of kTicksPerQuarterNoteInt
+	* to do its calculations.
+	*
+	* Note: If the returned tick is negative it is BEFORE any authored midi content.
+	* 
+	* @param FractionalBeatIndex 0-based fractional bar
+	* @return The 0-based raw tick
+	*/
+	float FractionalBeatIncludingCountInToTick(float FractionalBeatIndex) const;
+
 	/**
 	 * Adds a time signature at the specified position, in this case a bar number as one would find
 	 * in an FMusicTimestamp. In other words, a BarNumber before '1' would be in the count-in or pickup.
