make MPC instances weak pointers so they can be reset when the world changes

diff --git a/Engine/Plugins/Runtime/Harmonix/Source/HarmonixMetasound/Private/HarmonixMetasound/Utils/MusicTempometerUtilities.cpp b/Engine/Plugins/Runtime/Harmonix/Source/HarmonixMetasound/Private/HarmonixMetasound/Utils/MusicTempometerUtilities.cpp
--- a/Engine/Plugins/Runtime/Harmonix/Source/HarmonixMetasound/Private/HarmonixMetasound/Utils/MusicTempometerUtilities.cpp
+++ b/Engine/Plugins/Runtime/Harmonix/Source/HarmonixMetasound/Private/HarmonixMetasound/Utils/MusicTempometerUtilities.cpp
@@ -8,7 +8,7 @@
 
 #include UE_INLINE_GENERATED_CPP_BY_NAME(MusicTempometerUtilities)
 
-void MusicTempometerUtilities::UpdateMaterialParameterCollectionFromClock(const UObject* InWorldContextObject, TObjectPtr<UMaterialParameterCollectionInstance>& InOutMaterialParameterCollectionInstance, const TObjectPtr<UMaterialParameterCollection>& InMaterialParameterCollection, const FMusicTempometerMPCParameters& InMCPParameters, const UMusicClockComponent* InClockComponent)
+void MusicTempometerUtilities::UpdateMaterialParameterCollectionFromClock(const UObject* InWorldContextObject, TWeakObjectPtr<UMaterialParameterCollectionInstance>& InOutMaterialParameterCollectionInstance, const TObjectPtr<UMaterialParameterCollection>& InMaterialParameterCollection, const FMusicTempometerMPCParameters& InMCPParameters, const UMusicClockComponent* InClockComponent)
 {
 	if (InClockComponent)
 	{
@@ -17,7 +17,7 @@ void MusicTempometerUtilities::UpdateMaterialParameterCollectionFromClock(const
 	}
 }
 
-void MusicTempometerUtilities::UpdateMaterialParameterCollectionFromSongPos(const UObject* InWorldContextObject, TObjectPtr<UMaterialParameterCollectionInstance>& InOutMaterialParameterCollectionInstance, const TObjectPtr<UMaterialParameterCollection>& InMaterialParameterCollection, const FMusicTempometerMPCParameters& InMCPParameters, const FMidiSongPos& InMidiSongPos)
+void MusicTempometerUtilities::UpdateMaterialParameterCollectionFromSongPos(const UObject* InWorldContextObject, TWeakObjectPtr<UMaterialParameterCollectionInstance>& InOutMaterialParameterCollectionInstance, const TObjectPtr<UMaterialParameterCollection>& InMaterialParameterCollection, const FMusicTempometerMPCParameters& InMCPParameters, const FMidiSongPos& InMidiSongPos)
 {
 	// Find a MaterialParameterCollectionInstance to update
 	if (InOutMaterialParameterCollectionInstance == nullptr)
diff --git a/Engine/Plugins/Runtime/Harmonix/Source/HarmonixMetasound/Public/HarmonixMetasound/Components/MusicTempometerComponent.h b/Engine/Plugins/Runtime/Harmonix/Source/HarmonixMetasound/Public/HarmonixMetasound/Components/MusicTempometerComponent.h
--- a/Engine/Plugins/Runtime/Harmonix/Source/HarmonixMetasound/Public/HarmonixMetasound/Components/MusicTempometerComponent.h
+++ b/Engine/Plugins/Runtime/Harmonix/Source/HarmonixMetasound/Public/HarmonixMetasound/Components/MusicTempometerComponent.h
@@ -315,5 +315,5 @@ private:
 	 * Pinned material parameter collection instance
 	 */
 	UPROPERTY(Transient)
-	mutable TObjectPtr<UMaterialParameterCollectionInstance> MaterialParameterCollectionInstance;
+	mutable TWeakObjectPtr<UMaterialParameterCollectionInstance> MaterialParameterCollectionInstance;
 };
diff --git a/Engine/Plugins/Runtime/Harmonix/Source/HarmonixMetasound/Public/HarmonixMetasound/Utils/MusicTempometerUtilities.h b/Engine/Plugins/Runtime/Harmonix/Source/HarmonixMetasound/Public/HarmonixMetasound/Utils/MusicTempometerUtilities.h
--- a/Engine/Plugins/Runtime/Harmonix/Source/HarmonixMetasound/Public/HarmonixMetasound/Utils/MusicTempometerUtilities.h
+++ b/Engine/Plugins/Runtime/Harmonix/Source/HarmonixMetasound/Public/HarmonixMetasound/Utils/MusicTempometerUtilities.h
@@ -172,10 +172,10 @@ namespace MusicTempometerUtilities
 	/**
 	 * Update the supplied material parameter collection instance (create if null) with song position data extracted from the supplied clock
 	 */
-	void HARMONIXMETASOUND_API UpdateMaterialParameterCollectionFromClock(const UObject* InWorldContextObject, TObjectPtr<UMaterialParameterCollectionInstance>& InOutMaterialParameterCollectionInstance, const TObjectPtr<UMaterialParameterCollection>& InMaterialParameterCollection, const FMusicTempometerMPCParameters& InMCPParameters, const UMusicClockComponent* InClockComponent);
+	void HARMONIXMETASOUND_API UpdateMaterialParameterCollectionFromClock(const UObject* InWorldContextObject, TWeakObjectPtr<UMaterialParameterCollectionInstance>& InOutMaterialParameterCollectionInstance, const TObjectPtr<UMaterialParameterCollection>& InMaterialParameterCollection, const FMusicTempometerMPCParameters& InMCPParameters, const UMusicClockComponent* InClockComponent);
 	
 	/**
 	 * Update the supplied material parameter collection instance (create if null) with data from the supplied song position
 	 */
-	void HARMONIXMETASOUND_API UpdateMaterialParameterCollectionFromSongPos(const UObject* InWorldContextObject, TObjectPtr<UMaterialParameterCollectionInstance>& InOutMaterialParameterCollectionInstance, const TObjectPtr<UMaterialParameterCollection>& InMaterialParameterCollection, const FMusicTempometerMPCParameters& InMCPParameters, const FMidiSongPos& InMidiSongPos);
+	void HARMONIXMETASOUND_API UpdateMaterialParameterCollectionFromSongPos(const UObject* InWorldContextObject, TWeakObjectPtr<UMaterialParameterCollectionInstance>& InOutMaterialParameterCollectionInstance, const TObjectPtr<UMaterialParameterCollection>& InMaterialParameterCollection, const FMusicTempometerMPCParameters& InMCPParameters, const FMidiSongPos& InMidiSongPos);
 };
