Add some support for the Engine's new sense of musical time.

diff --git a/Engine/Plugins/Runtime/Harmonix/Source/HarmonixMidi/Private/HarmonixMidi/SongMaps.cpp b/Engine/Plugins/Runtime/Harmonix/Source/HarmonixMidi/Private/HarmonixMidi/SongMaps.cpp
--- a/Engine/Plugins/Runtime/Harmonix/Source/HarmonixMidi/Private/HarmonixMidi/SongMaps.cpp
+++ b/Engine/Plugins/Runtime/Harmonix/Source/HarmonixMidi/Private/HarmonixMidi/SongMaps.cpp
@@ -554,6 +554,45 @@ void ISongMapEvaluator::GetTicksForNearestSubdivision(int32 InTick, EMidiClockSu
 	UpperTick = LowerTick + TicksInSubdivision;
 }
 
+FMusicalTime ISongMapEvaluator::GetMusicalTimeAtSeconds(double Seconds) const
+{
+	float InMs = static_cast<float>(Seconds * 1000.0);
+	int32 Tick = FMath::FloorToInt32(MsToTick(InMs));
+	int32 OutBarIndex;  // 0 based
+	int32 OutBeatInBar; // 1 based
+	int32 OutTickIndex; // 0 based
+	int32 OutBeatsPerBar;
+	int32 OutTicksPerBeat;
+	TickToBarBeatTickIncludingCountIn(Tick, OutBarIndex, OutBeatInBar, OutTickIndex, &OutBeatsPerBar, &OutTicksPerBeat);
+	OutBeatInBar--; // NOW 0 based
+
+	FMusicalTime Result;
+	Result.Bar = OutBarIndex;
+	Result.TicksPerBeat = OutTicksPerBeat * MusicalTime::TicksPerQuarterNote / GetTicksPerQuarterNote();
+	Result.TicksPerBar = Result.TicksPerBeat * OutBeatsPerBar;
+	Result.TickInBar = OutBeatInBar * Result.TicksPerBeat + (OutTickIndex * MusicalTime::TicksPerQuarterNote / GetTicksPerQuarterNote());
+	return Result;
+}
+
+FMusicalTime ISongMapEvaluator::GetMusicalTimeAtFractionalBar(float Bars) const
+{
+	int32 Tick = FMath::FloorToInt32(FractionalBarIncludingCountInToTick(Bars));
+	int32 OutBarIndex;  // 0 based
+	int32 OutBeatInBar; // 1 based
+	int32 OutTickIndex; // 0 based
+	int32 OutBeatsPerBar;
+	int32 OutTicksPerBeat;
+	TickToBarBeatTickIncludingCountIn(Tick, OutBarIndex, OutBeatInBar, OutTickIndex, &OutBeatsPerBar, &OutTicksPerBeat);
+	OutBeatInBar--; // NOW 0 based
+
+	FMusicalTime Result;
+	Result.Bar = OutBarIndex;
+	Result.TicksPerBeat = OutTicksPerBeat * MusicalTime::TicksPerQuarterNote / GetTicksPerQuarterNote();
+	Result.TicksPerBar = Result.TicksPerBeat * OutBeatsPerBar;
+	Result.TickInBar = OutBeatInBar * Result.TicksPerBeat + (OutTickIndex * MusicalTime::TicksPerQuarterNote / GetTicksPerQuarterNote());
+	return Result;
+}
+
 FString ISongMapEvaluator::GetSongLengthString() const
 {
 	const FBarMap& BarMap = GetBarMap();
@@ -1031,9 +1070,9 @@ float ISongMapEvaluator::TickToFractionalBarIncludingCountIn(float Tick) const
 	return GetBarMap().TickToFractionalBarIncludingCountIn(Tick);
 }
 
-void ISongMapEvaluator::TickToBarBeatTickIncludingCountIn(int32 RawTick, int32& OutBarIndex, int32& OutBeatInBarIndex, int32& OutTickIndexInBeat, int32* OutBeatsPerBar, int32* OutTicksPerBeat) const
+void ISongMapEvaluator::TickToBarBeatTickIncludingCountIn(int32 RawTick, int32& OutBarIndex, int32& OutBeatInBar, int32& OutTickIndexInBeat, int32* OutBeatsPerBar, int32* OutTicksPerBeat) const
 {
-	GetBarMap().TickToBarBeatTickIncludingCountIn(RawTick, OutBarIndex, OutBeatInBarIndex, OutTickIndexInBeat, OutBeatsPerBar, OutTicksPerBeat);
+	GetBarMap().TickToBarBeatTickIncludingCountIn(RawTick, OutBarIndex, OutBeatInBar, OutTickIndexInBeat, OutBeatsPerBar, OutTicksPerBeat);
 }
 
 void FSongMaps::SetLengthTotalBars(int32 Bars)
diff --git a/Engine/Plugins/Runtime/Harmonix/Source/HarmonixMidi/Public/HarmonixMidi/BarMap.h b/Engine/Plugins/Runtime/Harmonix/Source/HarmonixMidi/Public/HarmonixMidi/BarMap.h
--- a/Engine/Plugins/Runtime/Harmonix/Source/HarmonixMidi/Public/HarmonixMidi/BarMap.h
+++ b/Engine/Plugins/Runtime/Harmonix/Source/HarmonixMidi/Public/HarmonixMidi/BarMap.h
@@ -292,7 +292,7 @@ public:
 	 * @param OutBeatsPerBar (output - optional)
 	 * @param OutTicksPerBeat (output - optional)
 	 */
-	void TickToBarBeatTickIncludingCountIn(int32 RawTick, int32& OutBarIndex, int32& OutBeatInBarIndex, int32& OutTickIndexInBeat, int32* OutBeatsPerBar = nullptr, int32* OutTicksPerBeat = nullptr) const;
+	void TickToBarBeatTickIncludingCountIn(int32 RawTick, int32& OutBarIndex, int32& OutBeatInBar, int32& OutTickIndexInBeat, int32* OutBeatsPerBar = nullptr, int32* OutTicksPerBeat = nullptr) const;
 
 	/**
 	 * Takes an absolute tick (from the beginning of the midi data) and calculates the fractional bar,
diff --git a/Engine/Plugins/Runtime/Harmonix/Source/HarmonixMidi/Public/HarmonixMidi/SongMaps.h b/Engine/Plugins/Runtime/Harmonix/Source/HarmonixMidi/Public/HarmonixMidi/SongMaps.h
--- a/Engine/Plugins/Runtime/Harmonix/Source/HarmonixMidi/Public/HarmonixMidi/SongMaps.h
+++ b/Engine/Plugins/Runtime/Harmonix/Source/HarmonixMidi/Public/HarmonixMidi/SongMaps.h
@@ -5,6 +5,7 @@
 #include "HarmonixMidi/BeatMap.h"
 #include "HarmonixMidi/ChordMap.h"
 #include "HarmonixMidi/SectionMap.h"
+#include "Misc/MusicalTime.h"
 #include "Sound/QuartzQuantizationUtilities.h"
 #include "FrameBasedMusicMap.h"
 #include <limits>
@@ -134,7 +135,7 @@ public:
 	float FractionalBarIncludingCountInToTick(float FractionalBarIndex) const;
 	int32 TickToBarIncludingCountIn(int32 Tick) const;
 	float TickToFractionalBarIncludingCountIn(float Tick) const;
-	void TickToBarBeatTickIncludingCountIn(int32 RawTick, int32& OutBarIndex, int32& OutBeatInBarIndex, int32& OutTickIndexInBeat, int32* OutBeatsPerBar = nullptr, int32* OutTicksPerBeat = nullptr) const;
+	void TickToBarBeatTickIncludingCountIn(int32 RawTick, int32& OutBarIndex, int32& OutBeatInBar, int32& OutTickIndexInBeat, int32* OutBeatsPerBar = nullptr, int32* OutTicksPerBeat = nullptr) const;
 	int32 CalculateMidiTick(const FMusicTimestamp& Timestamp, const EMidiClockSubdivisionQuantization Quantize) const;
 	int32 SubdivisionToMidiTicks(const EMidiClockSubdivisionQuantization Division, const int32 AtTick) const;
 	float MusicTimestampToTick(const FMusicTimestamp& Timestamp) const;
@@ -175,8 +176,11 @@ public:
 	int32 QuantizeTickToNearestSubdivision(int32 InTick, EMidiFileQuantizeDirection Direction, EMidiClockSubdivisionQuantization Division) const;
 	void GetTicksForNearestSubdivision(int32 InTick, EMidiClockSubdivisionQuantization Division, int32& LowerTick, int32& UpperTick) const;
 
-	//virtual FSongLengthData& GetSongLengthData() = 0;
 	virtual const FSongLengthData& GetSongLengthData() const = 0;
+	
+	// engine level music types
+	FMusicalTime GetMusicalTimeAtSeconds(double Seconds) const;
+	FMusicalTime GetMusicalTimeAtFractionalBar(float Bars) const;
 
 protected:
 	friend struct FSongMaps;
