- Fix issue where timeline skips back to 0 and then snaps back to cursor position when starting a sequence audition.
- Make IMusicEnvironmentClockSource getter function non-const. Implementors must be able to do real work, cache results, etc.
- Other work to prepare for Fabric Song Sync Device to be made 'music mode sequence' compatible.

diff --git a/Engine/Plugins/Runtime/Harmonix/Source/HarmonixMetasound/Private/HarmonixMetasound/Components/MusicClockComponent.cpp b/Engine/Plugins/Runtime/Harmonix/Source/HarmonixMetasound/Private/HarmonixMetasound/Components/MusicClockComponent.cpp
--- a/Engine/Plugins/Runtime/Harmonix/Source/HarmonixMetasound/Private/HarmonixMetasound/Components/MusicClockComponent.cpp
+++ b/Engine/Plugins/Runtime/Harmonix/Source/HarmonixMetasound/Private/HarmonixMetasound/Components/MusicClockComponent.cpp
@@ -806,7 +806,7 @@ const FMidiSongPos& UMusicClockComponent::GetPreviousSongPosInternal(ECalibrated
 	return PosZero;
 }
 
-FMusicalTime UMusicClockComponent::GetMusicalTime() const
+FMusicalTime UMusicClockComponent::GetMusicalTime()
 {
 	FMusicalTime Result;
 	if (!ClockDriver)
@@ -827,7 +827,7 @@ FMusicalTime UMusicClockComponent::GetMusicalTime() const
 	return Result;
 }
 
-int32 UMusicClockComponent::GetAbsoluteTickPosition() const
+int32 UMusicClockComponent::GetAbsoluteTickPosition()
 {
 	// Note: Since this is an override of a function in the music environment system we need to
 	// be sure we are using its standard ticks per quarter note units.
@@ -837,7 +837,7 @@ int32 UMusicClockComponent::GetAbsoluteTickPosition() const
 }
 
 
-FMusicalTime UMusicClockComponent::GetMusicalTimeWithSourceSpaceOffset(const FMusicalTime& Offset) const
+FMusicalTime UMusicClockComponent::GetMusicalTimeWithSourceSpaceOffset(const FMusicalTime& Offset)
 {
 	FMusicalTime Result;
 	const ISongMapEvaluator& MapEvaluator = GetSongMaps();
@@ -855,7 +855,7 @@ FMusicalTime UMusicClockComponent::GetMusicalTimeWithSourceSpaceOffset(const FMu
 	return Result;
 }
 
-int32 UMusicClockComponent::GetAbsoluteTickPositionWithSourceSpaceOffset(const FMusicalTime& Offset) const
+int32 UMusicClockComponent::GetAbsoluteTickPositionWithSourceSpaceOffset(const FMusicalTime& Offset)
 {
 	// Note: Since this is an override of a function in the music environment system we need to
 	// be sure we are using its standard ticks per quarter note units.
diff --git a/Engine/Plugins/Runtime/Harmonix/Source/HarmonixMetasound/Private/HarmonixMetasound/MusicEnvironmentSupport/HarmonixMusicEnvironmentMetronome.cpp b/Engine/Plugins/Runtime/Harmonix/Source/HarmonixMetasound/Private/HarmonixMetasound/MusicEnvironmentSupport/HarmonixMusicEnvironmentMetronome.cpp
--- a/Engine/Plugins/Runtime/Harmonix/Source/HarmonixMetasound/Private/HarmonixMetasound/MusicEnvironmentSupport/HarmonixMusicEnvironmentMetronome.cpp
+++ b/Engine/Plugins/Runtime/Harmonix/Source/HarmonixMetasound/Private/HarmonixMetasound/MusicEnvironmentSupport/HarmonixMusicEnvironmentMetronome.cpp
@@ -48,7 +48,7 @@ void UHarmonixMusicEnvironmentMetronome::Start(double FromSeconds)
 	}
 
 	RebuildMidiFile();
-	
+	MusicalTimeWhenStarted = MidiFile->GetSongMaps()->GetMusicalTimeAtSeconds(FromSeconds);
 	FAudioParameter MidiFileParam(TEXT("MIDI File"), MidiFile);
 	FAudioParameter SeekTargetMs(TEXT("SeekTarget"), static_cast<float>(FromSeconds * 1000.0f));
 	AudioComponent->Activate();
@@ -103,17 +103,18 @@ double UHarmonixMusicEnvironmentMetronome::GetCurrentPositionSeconds()
 	return 0.0;
 }
 
-FMusicalTime UHarmonixMusicEnvironmentMetronome::GetMusicalTime() const
+FMusicalTime UHarmonixMusicEnvironmentMetronome::GetMusicalTime()
 {
-	FMusicalTime MusicalTime;
+	FMusicalTime CurrentMusicalTime;
 	if (MusicClockComponent)
 	{
-		return MusicClockComponent->GetMusicalTime();
+		CurrentMusicalTime = MusicClockComponent->GetMusicalTime();
+		return MusicalTimeWhenStarted > CurrentMusicalTime ? MusicalTimeWhenStarted : CurrentMusicalTime;
 	}
-	return MusicalTime;
+	return CurrentMusicalTime;
 }
 
-int32 UHarmonixMusicEnvironmentMetronome::GetAbsoluteTickPosition() const
+int32 UHarmonixMusicEnvironmentMetronome::GetAbsoluteTickPosition()
 {
 	if (MusicClockComponent)
 	{
@@ -122,7 +123,7 @@ int32 UHarmonixMusicEnvironmentMetronome::GetAbsoluteTickPosition() const
 	return 0;
 }
 
-FMusicalTime UHarmonixMusicEnvironmentMetronome::GetMusicalTimeWithSourceSpaceOffset(const FMusicalTime& Offset) const
+FMusicalTime UHarmonixMusicEnvironmentMetronome::GetMusicalTimeWithSourceSpaceOffset(const FMusicalTime& Offset)
 {
 	FMusicalTime MusicalTime;
 	if (MusicClockComponent)
@@ -132,7 +133,7 @@ FMusicalTime UHarmonixMusicEnvironmentMetronome::GetMusicalTimeWithSourceSpaceOf
 	return MusicalTime;
 }
 
-int32 UHarmonixMusicEnvironmentMetronome::GetAbsoluteTickPositionWithSourceSpaceOffset(const FMusicalTime& Offset) const
+int32 UHarmonixMusicEnvironmentMetronome::GetAbsoluteTickPositionWithSourceSpaceOffset(const FMusicalTime& Offset)
 {
 	if (MusicClockComponent)
 	{
diff --git a/Engine/Plugins/Runtime/Harmonix/Source/HarmonixMetasound/Private/HarmonixMetasound/MusicEnvironmentSupport/HarmonixMusicEnvironmentMetronome.h b/Engine/Plugins/Runtime/Harmonix/Source/HarmonixMetasound/Private/HarmonixMetasound/MusicEnvironmentSupport/HarmonixMusicEnvironmentMetronome.h
--- a/Engine/Plugins/Runtime/Harmonix/Source/HarmonixMetasound/Private/HarmonixMetasound/MusicEnvironmentSupport/HarmonixMusicEnvironmentMetronome.h
+++ b/Engine/Plugins/Runtime/Harmonix/Source/HarmonixMetasound/Private/HarmonixMetasound/MusicEnvironmentSupport/HarmonixMusicEnvironmentMetronome.h
@@ -39,10 +39,10 @@ public:
 
 
 	// vvv BEGIN  IMusicEnvironmentClockSource Implementation vvv
-	virtual FMusicalTime GetMusicalTime() const override;
-	virtual int32 GetAbsoluteTickPosition() const override;
-	virtual FMusicalTime GetMusicalTimeWithSourceSpaceOffset(const FMusicalTime& Offset) const override;
-	virtual int32 GetAbsoluteTickPositionWithSourceSpaceOffset(const FMusicalTime& Offset) const override;
+	virtual FMusicalTime GetMusicalTime() override;
+	virtual int32 GetAbsoluteTickPosition() override;
+	virtual FMusicalTime GetMusicalTimeWithSourceSpaceOffset(const FMusicalTime& Offset) override;
+	virtual int32 GetAbsoluteTickPositionWithSourceSpaceOffset(const FMusicalTime& Offset) override;
 	virtual FMusicalTime Quantize(const FMusicalTime& MusicalTime, int32 QuantizationInterval, UFrameBasedMusicMap::EQuantizeDirection Direction = UFrameBasedMusicMap::EQuantizeDirection::Nearest) const override;
 	virtual bool CanAuditionInEditor() const override { return true; }
 	// ^^^ END IMusicEnvironmentClockSource Implementation ^^^
@@ -73,6 +73,7 @@ private:
 
 	bool bMetasoundIsPlaying = false;
 	float CurrentTempo = 120.0f;
+	FMusicalTime MusicalTimeWhenStarted;
 	
 	bool BuildAndStartMetasound(UWorld* InWorld);
 	bool RebuildMidiFile();
diff --git a/Engine/Plugins/Runtime/Harmonix/Source/HarmonixMetasound/Public/HarmonixMetasound/Components/MusicClockComponent.h b/Engine/Plugins/Runtime/Harmonix/Source/HarmonixMetasound/Public/HarmonixMetasound/Components/MusicClockComponent.h
--- a/Engine/Plugins/Runtime/Harmonix/Source/HarmonixMetasound/Public/HarmonixMetasound/Components/MusicClockComponent.h
+++ b/Engine/Plugins/Runtime/Harmonix/Source/HarmonixMetasound/Public/HarmonixMetasound/Components/MusicClockComponent.h
@@ -487,10 +487,10 @@ public:
 	const FMidiSongPos& GetRawUnsmoothedAudioRenderPos() const;
 
 	// vvv BEGIN  IMusicEnvironmentClockSource Implementation vvv
-	virtual FMusicalTime GetMusicalTime() const override;
-	virtual int32 GetAbsoluteTickPosition() const override;
-	virtual FMusicalTime GetMusicalTimeWithSourceSpaceOffset(const FMusicalTime& Offset) const;
-	virtual int32 GetAbsoluteTickPositionWithSourceSpaceOffset(const FMusicalTime& Offset) const override;
+	virtual FMusicalTime GetMusicalTime() override;
+	virtual int32 GetAbsoluteTickPosition() override;
+	virtual FMusicalTime GetMusicalTimeWithSourceSpaceOffset(const FMusicalTime& Offset) override;
+	virtual int32 GetAbsoluteTickPositionWithSourceSpaceOffset(const FMusicalTime& Offset) override;
 	virtual FMusicalTime Quantize(const FMusicalTime& MusicalTime, int32 QuantizationInterval, UFrameBasedMusicMap::EQuantizeDirection Direction = UFrameBasedMusicMap::EQuantizeDirection::Nearest) const;
 	virtual bool CanAuditionInEditor() const override { return false; }
 	// ^^^ END IMusicEnvironmentClockSource Implementation ^^^
diff --git a/Engine/Plugins/Runtime/MusicEnvironment/Source/MusicEnvironment/Public/MusicEnvironmentClockSource.h b/Engine/Plugins/Runtime/MusicEnvironment/Source/MusicEnvironment/Public/MusicEnvironmentClockSource.h
--- a/Engine/Plugins/Runtime/MusicEnvironment/Source/MusicEnvironment/Public/MusicEnvironmentClockSource.h
+++ b/Engine/Plugins/Runtime/MusicEnvironment/Source/MusicEnvironment/Public/MusicEnvironmentClockSource.h
@@ -24,10 +24,10 @@ class MUSICENVIRONMENT_API IMusicEnvironmentClockSource
 	GENERATED_BODY()
 
 public:
-	virtual FMusicalTime GetMusicalTime() const = 0;
-	virtual int32 GetAbsoluteTickPosition() const = 0;
-	virtual FMusicalTime GetMusicalTimeWithSourceSpaceOffset(const FMusicalTime& Offset) const = 0;
-	virtual int32 GetAbsoluteTickPositionWithSourceSpaceOffset(const FMusicalTime& SourceOffset) const = 0;
+	virtual FMusicalTime GetMusicalTime() = 0;
+	virtual int32 GetAbsoluteTickPosition() = 0;
+	virtual FMusicalTime GetMusicalTimeWithSourceSpaceOffset(const FMusicalTime& Offset) = 0;
+	virtual int32 GetAbsoluteTickPositionWithSourceSpaceOffset(const FMusicalTime& SourceOffset) = 0;
 	virtual FMusicalTime Quantize(const FMusicalTime& MusicalTime, int32 QuantizationInterval, UFrameBasedMusicMap::EQuantizeDirection Direction = UFrameBasedMusicMap::EQuantizeDirection::Nearest) const = 0;
 	virtual bool CanAuditionInEditor() const = 0;
 };
