[MusicEventSubsystem] Extra safety checks
#rnx

diff --git a/Engine/Plugins/Runtime/Harmonix/Source/HarmonixMetasound/Private/HarmonixMetasound/Utils/MusicTempometerUtilities.cpp b/Engine/Plugins/Runtime/Harmonix/Source/HarmonixMetasound/Private/HarmonixMetasound/Utils/MusicTempometerUtilities.cpp
--- a/Engine/Plugins/Runtime/Harmonix/Source/HarmonixMetasound/Private/HarmonixMetasound/Utils/MusicTempometerUtilities.cpp
+++ b/Engine/Plugins/Runtime/Harmonix/Source/HarmonixMetasound/Private/HarmonixMetasound/Utils/MusicTempometerUtilities.cpp
@@ -20,7 +20,7 @@ void MusicTempometerUtilities::UpdateMaterialParameterCollectionFromClock(const
 void MusicTempometerUtilities::UpdateMaterialParameterCollectionFromSongPos(const UObject* InWorldContextObject, TWeakObjectPtr<UMaterialParameterCollectionInstance>& InOutMaterialParameterCollectionInstance, const TObjectPtr<UMaterialParameterCollection>& InMaterialParameterCollection, const FMusicTempometerMPCParameters& InMCPParameters, const FMidiSongPos& InMidiSongPos)
 {
 	// Find a MaterialParameterCollectionInstance to update
-	if (InOutMaterialParameterCollectionInstance == nullptr)
+	if (!InOutMaterialParameterCollectionInstance.IsValid())
 	{
 		if (InMCPParameters.IsValid())
 		{
@@ -33,7 +33,7 @@ void MusicTempometerUtilities::UpdateMaterialParameterCollectionFromSongPos(cons
 			}
 		}
 
-		if (InOutMaterialParameterCollectionInstance == nullptr)
+		if (!InOutMaterialParameterCollectionInstance.IsValid())
 		{
 			return;
 		}
