Add ability to set volume and mute metronome

Add option in editor to toggle metronome on and off

#rnx

diff --git a/Engine/Plugins/Runtime/Harmonix/Source/HarmonixMetasound/Private/HarmonixMetasound/Components/MusicClockComponent.cpp b/Engine/Plugins/Runtime/Harmonix/Source/HarmonixMetasound/Private/HarmonixMetasound/Components/MusicClockComponent.cpp
--- a/Engine/Plugins/Runtime/Harmonix/Source/HarmonixMetasound/Private/HarmonixMetasound/Components/MusicClockComponent.cpp
+++ b/Engine/Plugins/Runtime/Harmonix/Source/HarmonixMetasound/Private/HarmonixMetasound/Components/MusicClockComponent.cpp
@@ -806,7 +806,7 @@ const FMidiSongPos& UMusicClockComponent::GetPreviousSongPosInternal(ECalibrated
 	return PosZero;
 }
 
-FMusicalTime UMusicClockComponent::GetMusicalTime()
+FMusicalTime UMusicClockComponent::GetMusicalTime() const
 {
 	FMusicalTime Result;
 	if (!ClockDriver)
@@ -827,7 +827,7 @@ FMusicalTime UMusicClockComponent::GetMusicalTime()
 	return Result;
 }
 
-int32 UMusicClockComponent::GetAbsoluteTickPosition()
+int32 UMusicClockComponent::GetAbsoluteTickPosition() const
 {
 	// Note: Since this is an override of a function in the music environment system we need to
 	// be sure we are using its standard ticks per quarter note units.
@@ -837,7 +837,7 @@ int32 UMusicClockComponent::GetAbsoluteTickPosition()
 }
 
 
-FMusicalTime UMusicClockComponent::GetMusicalTimeWithSourceSpaceOffset(const FMusicalTime& Offset)
+FMusicalTime UMusicClockComponent::GetMusicalTimeWithSourceSpaceOffset(const FMusicalTime& Offset) const
 {
 	FMusicalTime Result;
 	const ISongMapEvaluator& MapEvaluator = GetSongMaps();
@@ -855,7 +855,7 @@ FMusicalTime UMusicClockComponent::GetMusicalTimeWithSourceSpaceOffset(const FMu
 	return Result;
 }
 
-int32 UMusicClockComponent::GetAbsoluteTickPositionWithSourceSpaceOffset(const FMusicalTime& Offset)
+int32 UMusicClockComponent::GetAbsoluteTickPositionWithSourceSpaceOffset(const FMusicalTime& Offset) const
 {
 	// Note: Since this is an override of a function in the music environment system we need to
 	// be sure we are using its standard ticks per quarter note units.
diff --git a/Engine/Plugins/Runtime/Harmonix/Source/HarmonixMetasound/Private/HarmonixMetasound/MusicEnvironmentSupport/HarmonixMusicEnvironmentMetronome.cpp b/Engine/Plugins/Runtime/Harmonix/Source/HarmonixMetasound/Private/HarmonixMetasound/MusicEnvironmentSupport/HarmonixMusicEnvironmentMetronome.cpp
--- a/Engine/Plugins/Runtime/Harmonix/Source/HarmonixMetasound/Private/HarmonixMetasound/MusicEnvironmentSupport/HarmonixMusicEnvironmentMetronome.cpp
+++ b/Engine/Plugins/Runtime/Harmonix/Source/HarmonixMetasound/Private/HarmonixMetasound/MusicEnvironmentSupport/HarmonixMusicEnvironmentMetronome.cpp
@@ -93,17 +93,33 @@ void UHarmonixMusicEnvironmentMetronome::Resume()
 	}
 }
 
-float UHarmonixMusicEnvironmentMetronome::GetCurrentTempo()
+float UHarmonixMusicEnvironmentMetronome::GetCurrentTempo() const
 {
 	return CurrentTempo;
 }
 
-double UHarmonixMusicEnvironmentMetronome::GetCurrentPositionSeconds()
+float UHarmonixMusicEnvironmentMetronome::GetCurrentSpeed() const
+{
+	return CurrentSpeed;
+}
+
+bool UHarmonixMusicEnvironmentMetronome::IsMuted() const
+{
+	return bIsMuted;
+}
+
+
+float UHarmonixMusicEnvironmentMetronome::GetCurrentVolume() const
+{
+	return CurrentVolume;
+}
+
+double UHarmonixMusicEnvironmentMetronome::GetCurrentPositionSeconds() const
 {
 	return 0.0;
 }
 
-FMusicalTime UHarmonixMusicEnvironmentMetronome::GetMusicalTime()
+FMusicalTime UHarmonixMusicEnvironmentMetronome::GetMusicalTime() const
 {
 	FMusicalTime CurrentMusicalTime;
 	if (MusicClockComponent)
@@ -114,7 +130,7 @@ FMusicalTime UHarmonixMusicEnvironmentMetronome::GetMusicalTime()
 	return CurrentMusicalTime;
 }
 
-int32 UHarmonixMusicEnvironmentMetronome::GetAbsoluteTickPosition()
+int32 UHarmonixMusicEnvironmentMetronome::GetAbsoluteTickPosition() const
 {
 	if (MusicClockComponent)
 	{
@@ -123,7 +139,7 @@ int32 UHarmonixMusicEnvironmentMetronome::GetAbsoluteTickPosition()
 	return 0;
 }
 
-FMusicalTime UHarmonixMusicEnvironmentMetronome::GetMusicalTimeWithSourceSpaceOffset(const FMusicalTime& Offset)
+FMusicalTime UHarmonixMusicEnvironmentMetronome::GetMusicalTimeWithSourceSpaceOffset(const FMusicalTime& Offset) const
 {
 	FMusicalTime MusicalTime;
 	if (MusicClockComponent)
@@ -133,7 +149,7 @@ FMusicalTime UHarmonixMusicEnvironmentMetronome::GetMusicalTimeWithSourceSpaceOf
 	return MusicalTime;
 }
 
-int32 UHarmonixMusicEnvironmentMetronome::GetAbsoluteTickPositionWithSourceSpaceOffset(const FMusicalTime& Offset)
+int32 UHarmonixMusicEnvironmentMetronome::GetAbsoluteTickPositionWithSourceSpaceOffset(const FMusicalTime& Offset) const
 {
 	if (MusicClockComponent)
 	{
@@ -192,6 +208,43 @@ bool UHarmonixMusicEnvironmentMetronome::OnSetTempo(float Bpm)
 	return true;
 }
 
+void UHarmonixMusicEnvironmentMetronome::OnSetSpeed(const float InSpeed)
+{
+	CurrentSpeed = InSpeed;
+	if (!AudioComponent)
+	{
+		return;
+	}
+
+	FAudioParameter VolumeParam(TEXT("Speed"), CurrentSpeed);
+	AudioComponent->SetParameter(MoveTemp(VolumeParam));
+}
+
+void UHarmonixMusicEnvironmentMetronome::OnSetVolume(float InVolume)
+{
+	CurrentVolume = InVolume;
+	if (!AudioComponent)
+	{
+		return;
+	}
+
+	FAudioParameter VolumeParam(TEXT("BeepGain"), FMath::Clamp(CurrentVolume, 0.0f, 1.0f));
+	AudioComponent->SetParameter(MoveTemp(VolumeParam));
+}
+
+void UHarmonixMusicEnvironmentMetronome::OnSetMuted(bool bInMuted)
+{
+	bIsMuted = bInMuted;
+	if (!AudioComponent)
+	{
+		return;
+	}
+
+	FAudioParameter VolumeParam(TEXT("BeepEnabled"), !bIsMuted);
+	AudioComponent->SetParameter(MoveTemp(VolumeParam));
+}
+
+
 bool UHarmonixMusicEnvironmentMetronome::RebuildMidiFile()
 {
 	MidiFile = NewObject<UMidiFile>(this, "MIDI File Music Map", RF_Transient);
@@ -264,7 +317,8 @@ bool UHarmonixMusicEnvironmentMetronome::BuildMetasound()
 		}
 
 		// Make the nodes and set their un-wired inputs to appropriate values...
-		if (!AddGraphInput(TEXT("BeepGain"), 0.25f).IsSet()) { return false; }
+		if (!AddGraphInput(TEXT("BeepGain"), GetCurrentVolume()).IsSet()) { return false; }
+		if (!AddGraphInput(TEXT("BeepEnabled"), !IsMuted()).IsSet()) { return false; }
 		
 		// the midi file is used for the midi player node if there's a music map
 		if (!AddGraphInput(TEXT("MIDI File"), nullptr).IsSet()) { return false; }
@@ -274,6 +328,7 @@ bool UHarmonixMusicEnvironmentMetronome::BuildMetasound()
 		if (!AddGraphInputTrigger(TEXT("StopMetronome")).IsSet()) { return false; }
 		if (!AddGraphInputTrigger(TEXT("PauseMetronome")).IsSet()) { return false; }
 		if (!AddGraphInputTrigger(TEXT("ContinueMetronome")).IsSet()) { return false; }
+		if (!AddGraphInput(TEXT("Speed"), GetCurrentSpeed()).IsSet()) { return false; }
 
 		Metasound::FNodeClassName ClassName = {Metasound::StandardNodes::Namespace, TEXT("Trigger Delay"), TEXT("")};
 		FMetaSoundNodeHandle PlayTriggerDelay = SourceBuilder->AddNodeByClassName(ClassName, Result);
@@ -334,8 +389,10 @@ bool UHarmonixMusicEnvironmentMetronome::BuildMetasound()
 		if (!ConnectGraphInputToNodeInput(TEXT("ContinueMetronome"), TransportNode, TEXT("Continue"))) { return false; }
 		if (!ConnectGraphInputToNodeInput(TEXT("StopMetronome"), TransportNode, TEXT("Stop"))) { return false; }
 		if (!ConnectGraphInputToNodeInput(TEXT("BeepGain"), GainMultiplier, TEXT("AdditionalOperands"))) { return false; }
+		if (!ConnectGraphInputToNodeInput(TEXT("BeepEnabled"), SineOsc, TEXT("Enabled"))) { return false; }
 
 		if (!ConnectGraphInputToNodeInput(TEXT("MIDI File"), MidiPlayerNode,TEXT("MIDI File"))) { return false; }
+		if (!ConnectGraphInputToNodeInput(TEXT("Speed"), MidiPlayerNode, TEXT("Speed"))) { return false; }
 		
 		if (!ConnectGraphInputToNodeInput(TEXT("SeekTarget"), SeekTargetConverter, TEXT("Time (Ms)"))) { return false; }
 
@@ -393,6 +450,16 @@ FMetaSoundBuilderNodeOutputHandle UHarmonixMusicEnvironmentMetronome::AddGraphIn
 	return SourceBuilder->AddGraphInputNode(InputName, Metasound::GetMetasoundDataTypeName<float>(), DefaultValue, Result);
 }
 
+FMetaSoundBuilderNodeOutputHandle UHarmonixMusicEnvironmentMetronome::AddGraphInput(const FName& InputName, const bool Value) const
+{
+	check(SourceBuilder);
+	EMetaSoundBuilderResult Result;
+
+	FMetasoundFrontendLiteral DefaultValue;
+	DefaultValue.Set(Value);
+	return SourceBuilder->AddGraphInputNode(InputName, Metasound::GetMetasoundDataTypeName<bool>(), DefaultValue, Result);
+}
+
 FMetaSoundBuilderNodeOutputHandle UHarmonixMusicEnvironmentMetronome::AddGraphInputTrigger(const FName& InputName) const
 {
 	check(SourceBuilder);
diff --git a/Engine/Plugins/Runtime/Harmonix/Source/HarmonixMetasound/Private/HarmonixMetasound/MusicEnvironmentSupport/HarmonixMusicEnvironmentMetronome.h b/Engine/Plugins/Runtime/Harmonix/Source/HarmonixMetasound/Private/HarmonixMetasound/MusicEnvironmentSupport/HarmonixMusicEnvironmentMetronome.h
--- a/Engine/Plugins/Runtime/Harmonix/Source/HarmonixMetasound/Private/HarmonixMetasound/MusicEnvironmentSupport/HarmonixMusicEnvironmentMetronome.h
+++ b/Engine/Plugins/Runtime/Harmonix/Source/HarmonixMetasound/Private/HarmonixMetasound/MusicEnvironmentSupport/HarmonixMusicEnvironmentMetronome.h
@@ -34,15 +34,17 @@ public:
 	virtual void Pause() override;
 	virtual void Resume() override;
 
-	virtual float GetCurrentTempo() override;
-	virtual double GetCurrentPositionSeconds() override;
-
+	virtual float GetCurrentTempo() const override;
+	virtual float GetCurrentVolume() const override;
+	virtual float GetCurrentSpeed() const override;
+	virtual double GetCurrentPositionSeconds() const override;
+	virtual bool IsMuted() const override;
 
 	// vvv BEGIN  IMusicEnvironmentClockSource Implementation vvv
-	virtual FMusicalTime GetMusicalTime() override;
-	virtual int32 GetAbsoluteTickPosition() override;
-	virtual FMusicalTime GetMusicalTimeWithSourceSpaceOffset(const FMusicalTime& Offset) override;
-	virtual int32 GetAbsoluteTickPositionWithSourceSpaceOffset(const FMusicalTime& Offset) override;
+	virtual FMusicalTime GetMusicalTime() const override;
+	virtual int32 GetAbsoluteTickPosition() const override;
+	virtual FMusicalTime GetMusicalTimeWithSourceSpaceOffset(const FMusicalTime& Offset) const override;
+	virtual int32 GetAbsoluteTickPositionWithSourceSpaceOffset(const FMusicalTime& Offset) const override;
 	virtual FMusicalTime Quantize(const FMusicalTime& MusicalTime, int32 QuantizationInterval, UFrameBasedMusicMap::EQuantizeDirection Direction = UFrameBasedMusicMap::EQuantizeDirection::Nearest) const override;
 	virtual bool CanAuditionInEditor() const override { return true; }
 	// ^^^ END IMusicEnvironmentClockSource Implementation ^^^
@@ -52,6 +54,9 @@ public:
 protected:
 	virtual void OnMusicMapSet() override;
 	virtual bool OnSetTempo(float Bpm) override;
+	virtual void OnSetSpeed(const float Speed) override;
+	virtual void OnSetVolume(float InVolume) override;
+	virtual void OnSetMuted(bool bNewMuted) override;
 	
 private:
 	UPROPERTY(Transient)
@@ -73,6 +78,9 @@ private:
 
 	bool bMetasoundIsPlaying = false;
 	float CurrentTempo = 120.0f;
+	float CurrentSpeed = 1.0f;
+	bool bIsMuted = false;
+	float CurrentVolume = 1.0f;
 	FMusicalTime MusicalTimeWhenStarted;
 	
 	bool BuildAndStartMetasound(UWorld* InWorld);
@@ -80,6 +88,7 @@ private:
 	bool BuildMetasound();
 	FMetaSoundBuilderNodeOutputHandle AddGraphInput(const FName& InputName, UMidiFile* MidiFile) const;
 	FMetaSoundBuilderNodeOutputHandle AddGraphInput(const FName& InputName, const float Value) const;
+	FMetaSoundBuilderNodeOutputHandle AddGraphInput(const FName& InputName, const bool Value) const;
 	FMetaSoundBuilderNodeOutputHandle AddGraphInputTrigger(const FName& InputName) const;
 	bool SetNodeInputDefault(const FMetaSoundNodeHandle& Node, const FName& InputName, const int32 Value) const;
 	bool SetNodeInputDefault(const FMetaSoundNodeHandle& Node, const FName& InputName, const float Value) const;
diff --git a/Engine/Plugins/Runtime/Harmonix/Source/HarmonixMetasound/Public/HarmonixMetasound/Components/MusicClockComponent.h b/Engine/Plugins/Runtime/Harmonix/Source/HarmonixMetasound/Public/HarmonixMetasound/Components/MusicClockComponent.h
--- a/Engine/Plugins/Runtime/Harmonix/Source/HarmonixMetasound/Public/HarmonixMetasound/Components/MusicClockComponent.h
+++ b/Engine/Plugins/Runtime/Harmonix/Source/HarmonixMetasound/Public/HarmonixMetasound/Components/MusicClockComponent.h
@@ -173,25 +173,21 @@ public:
 	
 	// Time from the beginning of the authored music content.
 	// NOTE: INCLUDES time for count-in and pickup bars.
-	// Note: Not const as it might cause the clock to update from its source.
 	UFUNCTION(BlueprintPure, Category = "MusicClock")
 	float GetSecondsIncludingCountIn(ECalibratedMusicTimebase Timebase = ECalibratedMusicTimebase::VideoRenderTime) const;
 
 	// Time from Bar 1 Beat 1. The classic "start of the song".
 	// NOTE: DOES NOT INCLUDE time for count-in and pickup bars.
-	// Note: Not const as it might cause the clock to update from its source.
 	UFUNCTION(BlueprintPure, Category = "MusicClock")
 	float GetSecondsFromBarOne(ECalibratedMusicTimebase Timebase = ECalibratedMusicTimebase::VideoRenderTime) const;
 
 	// Returns the fractional total bars from the beginning of the authored music content.
 	// NOTE: INCLUDES time for count-in and pickup bars.
-	// Note: Not const as it might cause the clock to update from its source.
 	UFUNCTION(BlueprintPure, Category = "MusicClock")
 	float GetBarsIncludingCountIn(ECalibratedMusicTimebase Timebase = ECalibratedMusicTimebase::VideoRenderTime) const;
 
 	// Returns the fractional total beats from the beginning of the authored music content.
 	// NOTE: INCLUDES time for count-in and pickup bars.
-	// Note: Not const as it might cause the clock to update from its source.
 	UFUNCTION(BlueprintPure, Category = "MusicClock")
 	float GetBeatsIncludingCountIn(ECalibratedMusicTimebase Timebase = ECalibratedMusicTimebase::VideoRenderTime) const;
 
@@ -205,17 +201,14 @@ public:
 	//    - Bar 1, Beat 1.0 is the "beginning of the song" AFTER count-in/pickups
 	//    - Bar 0, Beat 1.0 would be one bar BEFORE the "beginning of the song"... eg. a bar of count-in or pickup.
 	//    - While Bar can be positive or negative, Beat is always >= 1.0 and is read as "beat in the bar". Again... '1' based!
-	// Note: Not const as it might cause the clock to update from its source.
 	UFUNCTION(BlueprintPure, Category = "MusicClock")
 	FMusicTimestamp GetCurrentTimestamp(ECalibratedMusicTimebase Timebase = ECalibratedMusicTimebase::VideoRenderTime) const;
 
 	/** Returns the name of the section that we're currently in (intro, chorus, outro) */
-	// Note: Not const as it might cause the clock to update from its source.
 	UFUNCTION(BlueprintPure, Category = "MusicClock")
 	FString GetCurrentSectionName(ECalibratedMusicTimebase Timebase = ECalibratedMusicTimebase::VideoRenderTime) const;
 
 	/** Returns the index of the current section for the provided time base. [0, Num-1] */
-	// Note: Not const as it might cause the clock to update from its source.
 	UFUNCTION(BlueprintPure, Category = "MusicClock")
 	int32 GetCurrentSectionIndex(ECalibratedMusicTimebase Timebase = ECalibratedMusicTimebase::VideoRenderTime) const;
 
@@ -223,54 +216,43 @@ public:
 	const TArray<FSongSection>& GetSongSections() const;
 
 	/** Returns the start time of the current section in milliseconds */
-	// Note: Not const as it might cause the clock to update from its source.
 	UFUNCTION(BlueprintPure, Category = "MusicClock")
 	float GetCurrentSectionStartMs(ECalibratedMusicTimebase Timebase = ECalibratedMusicTimebase::VideoRenderTime) const;
 
 	/** Returns the length of the current section in milliseconds */
-	// Note: Not const as it might cause the clock to update from its source.
 	UFUNCTION(BlueprintPure, Category = "MusicClock")
 	float GetCurrentSectionLengthMs(ECalibratedMusicTimebase Timebase = ECalibratedMusicTimebase::VideoRenderTime) const;
 
 	/** Gets a value expressed in beats between 0-1 that indicates how much progress we made in the current beat */
-	// Note: Not const as it might cause the clock to update from its source.
 	UFUNCTION(BlueprintPure, Category = "MusicClock")
 	float GetDistanceFromCurrentBeat(ECalibratedMusicTimebase Timebase = ECalibratedMusicTimebase::VideoRenderTime) const;
 
 	/** Gets a value expressed in beats between 0-1 that indicates how close we are to the next beat. */
-	// Note: Not const as it might cause the clock to update from its source.
 	UFUNCTION(BlueprintPure, Category = "MusicClock")
 	float GetDistanceToNextBeat(ECalibratedMusicTimebase Timebase = ECalibratedMusicTimebase::VideoRenderTime) const;
 
 	/** Gets a value expressed in beats between 0-1 that indicates how close we are to the closest beat (current beat or next beat). */
-	// Note: Not const as it might cause the clock to update from its source.
 	UFUNCTION(BlueprintPure, Category = "MusicClock")
 	float GetDistanceToClosestBeat(ECalibratedMusicTimebase Timebase = ECalibratedMusicTimebase::VideoRenderTime) const;
 
 	/** Gets a value expressed in bars between 0-1 that indicates how much progress we made towards the current bar to the next one. */
-	// Note: Not const as it might cause the clock to update from its source.
 	UFUNCTION(BlueprintPure, Category = "MusicClock")
 	float GetDistanceFromCurrentBar(ECalibratedMusicTimebase Timebase = ECalibratedMusicTimebase::VideoRenderTime) const;
 
 	/** Gets a value expressed in bars between 0-1 that indicates how close we are to the next bar. */
-	// Note: Not const as it might cause the clock to update from its source.
 	UFUNCTION(BlueprintPure, Category = "MusicClock")
 	float GetDistanceToNextBar(ECalibratedMusicTimebase Timebase = ECalibratedMusicTimebase::VideoRenderTime) const;
 
 	/** Gets value expressed in bars between 0-1 that indicates how close we are to the closest bar (current bar or next bar). */
-	// Note: Not const as it might cause the clock to update from its source.
 	UFUNCTION(BlueprintPure, Category = "MusicClock")
 	float GetDistanceToClosestBar(ECalibratedMusicTimebase Timebase = ECalibratedMusicTimebase::VideoRenderTime) const;
 
-	// Note: Not const as it might cause the clock to update from its source.
 	UFUNCTION(BlueprintPure, Category = "MusicClock")
 	float GetDeltaBar(ECalibratedMusicTimebase Timebase = ECalibratedMusicTimebase::VideoRenderTime) const;
 
-	// Note: Not const as it might cause the clock to update from its source.
 	UFUNCTION(BlueprintPure, Category = "MusicClock")
 	float GetDeltaBeat(ECalibratedMusicTimebase Timebase = ECalibratedMusicTimebase::VideoRenderTime) const;
 
-	// Note: Not const as it might cause the clock to update from its source.
 	UFUNCTION(BlueprintPure, Category = "MusicClock")
 	const FMidiSongPos& GetSongPos(ECalibratedMusicTimebase Timebase = ECalibratedMusicTimebase::VideoRenderTime) const;
 
@@ -278,17 +260,14 @@ public:
 	const FMidiSongPos& GetPreviousSongPos(ECalibratedMusicTimebase Timebase = ECalibratedMusicTimebase::VideoRenderTime) const;
 
 	/** Returns the remaining time until the end of the MIDI in milliseconds based on the timestamp corresponding to the passed Timebase */
-	// Note: Not const as it might cause the clock to update from its source.
 	UFUNCTION(BlueprintPure, Category = "MusicClock")
 	float GetSongRemainingMs(ECalibratedMusicTimebase Timebase = ECalibratedMusicTimebase::VideoRenderTime) const;
 
 	/** Returns true if there was a seek in the specified timebase */
-	// Note: Not const as it might cause the clock to update from its source.
 	UFUNCTION(BlueprintPure, Category = "MusicClock")
 	bool SeekedThisFrame(ECalibratedMusicTimebase Timebase = ECalibratedMusicTimebase::VideoRenderTime) const;
 
 	/** Returns true if there was a seek in the specified timebase */
-	// Note: Not const as it might cause the clock to update from its source.
 	UFUNCTION(BlueprintPure, Category = "MusicClock")
 	bool LoopedThisFrame(ECalibratedMusicTimebase Timebase = ECalibratedMusicTimebase::VideoRenderTime) const;
 
@@ -465,7 +444,6 @@ public:
 	UFUNCTION(BlueprintCallable, BlueprintPure, Category = "MusicClock")
 	FMidiSongPos GetCurrentRawAudioRenderSongPos() const;
 
-	// Note: Not const as it might cause the clock to update from its source.
 	UFUNCTION(BlueprintPure, Category = "MusicClock")
 	float MeasureSpanProgress(const FMusicalTimeSpan& Span, ECalibratedMusicTimebase Timebase = ECalibratedMusicTimebase::VideoRenderTime) const;
 
@@ -481,16 +459,15 @@ public:
 	UFUNCTION(BlueprintCallable, Category = "MusicClock")
 	void ConnectToWallClockForMidi(UMidiFile* InTempoMap);
 
-	// Note: Not const as it might cause the clock to update from its source.
 	FMidiSongPos CalculateSongPosWithOffset(float MsOffset, ECalibratedMusicTimebase Timebase = ECalibratedMusicTimebase::VideoRenderTime) const;
 
 	const FMidiSongPos& GetRawUnsmoothedAudioRenderPos() const;
 
 	// vvv BEGIN  IMusicEnvironmentClockSource Implementation vvv
-	virtual FMusicalTime GetMusicalTime() override;
-	virtual int32 GetAbsoluteTickPosition() override;
-	virtual FMusicalTime GetMusicalTimeWithSourceSpaceOffset(const FMusicalTime& Offset) override;
-	virtual int32 GetAbsoluteTickPositionWithSourceSpaceOffset(const FMusicalTime& Offset) override;
+	virtual FMusicalTime GetMusicalTime() const override;
+	virtual int32 GetAbsoluteTickPosition() const override;
+	virtual FMusicalTime GetMusicalTimeWithSourceSpaceOffset(const FMusicalTime& Offset) const override;
+	virtual int32 GetAbsoluteTickPositionWithSourceSpaceOffset(const FMusicalTime& Offset) const override;
 	virtual FMusicalTime Quantize(const FMusicalTime& MusicalTime, int32 QuantizationInterval, UFrameBasedMusicMap::EQuantizeDirection Direction = UFrameBasedMusicMap::EQuantizeDirection::Nearest) const;
 	virtual bool CanAuditionInEditor() const override { return false; }
 	// ^^^ END IMusicEnvironmentClockSource Implementation ^^^
diff --git a/Engine/Plugins/Runtime/MusicEnvironment/Source/MusicEnvironment/Private/MusicEnvironmentMetronome.cpp b/Engine/Plugins/Runtime/MusicEnvironment/Source/MusicEnvironment/Private/MusicEnvironmentMetronome.cpp
--- a/Engine/Plugins/Runtime/MusicEnvironment/Source/MusicEnvironment/Private/MusicEnvironmentMetronome.cpp
+++ b/Engine/Plugins/Runtime/MusicEnvironment/Source/MusicEnvironment/Private/MusicEnvironmentMetronome.cpp
@@ -17,3 +17,21 @@ bool IMusicEnvironmentMetronome::SetTempo(const float Bpm)
 	}
 	return OnSetTempo(Bpm);
 }
+
+void IMusicEnvironmentMetronome::SetSpeed(float Speed)
+{
+	OnSetSpeed(Speed);
+}
+
+void IMusicEnvironmentMetronome::SetVolume(float InVolume)
+{
+	OnSetVolume(InVolume);
+}
+
+void IMusicEnvironmentMetronome::SetMuted(bool bInMuted)
+{
+	OnSetMuted(bInMuted);
+}
+
+
+
diff --git a/Engine/Plugins/Runtime/MusicEnvironment/Source/MusicEnvironment/Public/MusicEnvironmentClockSource.h b/Engine/Plugins/Runtime/MusicEnvironment/Source/MusicEnvironment/Public/MusicEnvironmentClockSource.h
--- a/Engine/Plugins/Runtime/MusicEnvironment/Source/MusicEnvironment/Public/MusicEnvironmentClockSource.h
+++ b/Engine/Plugins/Runtime/MusicEnvironment/Source/MusicEnvironment/Public/MusicEnvironmentClockSource.h
@@ -24,10 +24,10 @@ class MUSICENVIRONMENT_API IMusicEnvironmentClockSource
 	GENERATED_BODY()
 
 public:
-	virtual FMusicalTime GetMusicalTime() = 0;
-	virtual int32 GetAbsoluteTickPosition() = 0;
-	virtual FMusicalTime GetMusicalTimeWithSourceSpaceOffset(const FMusicalTime& Offset) = 0;
-	virtual int32 GetAbsoluteTickPositionWithSourceSpaceOffset(const FMusicalTime& SourceOffset) = 0;
+	virtual FMusicalTime GetMusicalTime() const = 0;
+	virtual int32 GetAbsoluteTickPosition() const = 0;
+	virtual FMusicalTime GetMusicalTimeWithSourceSpaceOffset(const FMusicalTime& Offset) const = 0;
+	virtual int32 GetAbsoluteTickPositionWithSourceSpaceOffset(const FMusicalTime& SourceOffset) const = 0;
 	virtual FMusicalTime Quantize(const FMusicalTime& MusicalTime, int32 QuantizationInterval, UFrameBasedMusicMap::EQuantizeDirection Direction = UFrameBasedMusicMap::EQuantizeDirection::Nearest) const = 0;
 	virtual bool CanAuditionInEditor() const = 0;
 };
diff --git a/Engine/Plugins/Runtime/MusicEnvironment/Source/MusicEnvironment/Public/MusicEnvironmentMetronome.h b/Engine/Plugins/Runtime/MusicEnvironment/Source/MusicEnvironment/Public/MusicEnvironmentMetronome.h
--- a/Engine/Plugins/Runtime/MusicEnvironment/Source/MusicEnvironment/Public/MusicEnvironmentMetronome.h
+++ b/Engine/Plugins/Runtime/MusicEnvironment/Source/MusicEnvironment/Public/MusicEnvironmentMetronome.h
@@ -11,7 +11,7 @@
 
 // This class does not need to be modified.
 UINTERFACE()
-class UMusicEnvironmentMetronome : public UMusicEnvironmentClockSource
+class MUSICENVIRONMENT_API UMusicEnvironmentMetronome : public UMusicEnvironmentClockSource
 {
 	GENERATED_BODY()
 };
@@ -37,11 +37,24 @@ public:
 	virtual void Resume() = 0;
 
 	bool SetTempo(float Bpm);
-	virtual float GetCurrentTempo() = 0;
-	virtual double GetCurrentPositionSeconds() = 0;
+	virtual float GetCurrentTempo() const = 0;
+	virtual double GetCurrentPositionSeconds() const = 0;
+	
+	void SetSpeed(float Speed);
+    virtual float GetCurrentSpeed() const = 0;
+    	
+	void SetVolume(float InVolumeLinear);
+	virtual float GetCurrentVolume() const = 0;
+
+	void SetMuted(bool bInMuted);
+	virtual bool IsMuted() const = 0;
 	
 protected:
 	virtual void OnMusicMapSet() = 0;
+	virtual void OnSetSpeed(const float Speed) = 0;
 	virtual bool OnSetTempo(const float Bpm) = 0;
+	virtual void OnSetVolume(float NewVolumeLinear) = 0;
+	virtual void OnSetMuted(bool bNewMuted) = 0;
+	
 	TStrongObjectPtr<UFrameBasedMusicMap> MusicMap;
 };
