Fix Clock Component returning incorrect BeatDurationSeconds

#tests HarmonixDev
#rnx

diff --git a/Engine/Plugins/Runtime/Harmonix/Source/HarmonixMetasound/Private/HarmonixMetasound/Components/MusicClockComponent.cpp b/Engine/Plugins/Runtime/Harmonix/Source/HarmonixMetasound/Private/HarmonixMetasound/Components/MusicClockComponent.cpp
--- a/Engine/Plugins/Runtime/Harmonix/Source/HarmonixMetasound/Private/HarmonixMetasound/Components/MusicClockComponent.cpp
+++ b/Engine/Plugins/Runtime/Harmonix/Source/HarmonixMetasound/Private/HarmonixMetasound/Components/MusicClockComponent.cpp
@@ -352,7 +352,7 @@ float UMusicClockComponent::GetCurrentBeatDurationSeconds() const
 {
 	if (ClockDriver)
 	{
-		return ClockDriver->CurrentBarDurationSec;
+		return ClockDriver->CurrentBeatDurationSec;
 	}
 	return 0;
 }
