Fix issue where Formant Preservation was not being set up properly on some elastique pitch shifters because of an initialization order bug.

diff --git a/Engine/Plugins/Runtime/Harmonix/Source/HarmonixDsp/Private/HarmonixDsp/AudioData/StreamingAudioRendererV2.cpp b/Engine/Plugins/Runtime/Harmonix/Source/HarmonixDsp/Private/HarmonixDsp/AudioData/StreamingAudioRendererV2.cpp
--- a/Engine/Plugins/Runtime/Harmonix/Source/HarmonixDsp/Private/HarmonixDsp/AudioData/StreamingAudioRendererV2.cpp
+++ b/Engine/Plugins/Runtime/Harmonix/Source/HarmonixDsp/Private/HarmonixDsp/AudioData/StreamingAudioRendererV2.cpp
@@ -60,7 +60,7 @@ void FStreamingAudioRendererV2::SetAudioData(TSharedRef<FSoundWaveProxy> InSound
 
 	if (Shifter)
 	{
-		Shifter->SetSampleSourceReset(SoundWaveProxy, AsShared());
+		Shifter->SetSampleSourceReset(SoundWaveProxy, AsShared(), InSettings.ShifterSettings);
 	}
 }
 
diff --git a/Engine/Plugins/Runtime/Harmonix/Source/HarmonixDsp/Private/HarmonixDsp/FusionSampler/FusionSampler.cpp b/Engine/Plugins/Runtime/Harmonix/Source/HarmonixDsp/Private/HarmonixDsp/FusionSampler/FusionSampler.cpp
--- a/Engine/Plugins/Runtime/Harmonix/Source/HarmonixDsp/Private/HarmonixDsp/FusionSampler/FusionSampler.cpp
+++ b/Engine/Plugins/Runtime/Harmonix/Source/HarmonixDsp/Private/HarmonixDsp/FusionSampler/FusionSampler.cpp
@@ -1958,7 +1958,7 @@ void FFusionSampler::UpdateVoicesForEnvelopeOrderChange()
 		{
 			TMap<FName, FTypedParameter> Options;
 			Options.Add("EnvelopeOrder", TimeStretchEnvelopeOverride);
-			Shifter->ApplyOptions(Options);
+			Shifter->UpdateOptions(Options);
 		}
 	}
 }
diff --git a/Engine/Plugins/Runtime/Harmonix/Source/HarmonixDsp/Private/HarmonixDsp/FusionSampler/FusionVoice.cpp b/Engine/Plugins/Runtime/Harmonix/Source/HarmonixDsp/Private/HarmonixDsp/FusionSampler/FusionVoice.cpp
--- a/Engine/Plugins/Runtime/Harmonix/Source/HarmonixDsp/Private/HarmonixDsp/FusionSampler/FusionVoice.cpp
+++ b/Engine/Plugins/Runtime/Harmonix/Source/HarmonixDsp/Private/HarmonixDsp/FusionSampler/FusionVoice.cpp
@@ -131,6 +131,7 @@ bool FFusionVoice::AssignIDs(FFusionSampler* InSampler, const FKeyzoneSettings*
 
 	IAudioDataRenderer::FSettings Settings;
 	Settings.Shifter = PitchShifter;
+	Settings.ShifterSettings = &InKeyZone->TimeStretchConfig;
 	Settings.TrackChannelInfo = &InKeyZone->TrackMap;
 	Settings.Sampler = MySampler;
 	//TRACE_BOOKMARK(TEXT("Starting Waveform: %s"), *KeyZone->AudioSample->GetName().ToString());
diff --git a/Engine/Plugins/Runtime/Harmonix/Source/HarmonixDsp/Private/HarmonixDsp/FusionSampler/FusionVoicePool.cpp b/Engine/Plugins/Runtime/Harmonix/Source/HarmonixDsp/Private/HarmonixDsp/FusionSampler/FusionVoicePool.cpp
--- a/Engine/Plugins/Runtime/Harmonix/Source/HarmonixDsp/Private/HarmonixDsp/FusionSampler/FusionVoicePool.cpp
+++ b/Engine/Plugins/Runtime/Harmonix/Source/HarmonixDsp/Private/HarmonixDsp/FusionSampler/FusionVoicePool.cpp
@@ -451,7 +451,6 @@ FFusionVoice* FFusionVoicePool::GetFreeVoice(
 	{
 		ShifterToUse = PitchShifterFactory->GetFreePitchShifter(InKeyzone->TimeStretchConfig);
 		check(ShifterToUse);
-		ShifterToUse->ApplyOptions(InKeyzone->TimeStretchConfig.PitchShifterOptions);
 	}
 
 	if (BestVoice->AssignIDs(InSampler, InKeyzone, InVoiceID, Handler, ShifterToUse))
diff --git a/Engine/Plugins/Runtime/Harmonix/Source/HarmonixDsp/Private/HarmonixDsp/Stretchers/SmbPitchShifter.cpp b/Engine/Plugins/Runtime/Harmonix/Source/HarmonixDsp/Private/HarmonixDsp/Stretchers/SmbPitchShifter.cpp
--- a/Engine/Plugins/Runtime/Harmonix/Source/HarmonixDsp/Private/HarmonixDsp/Stretchers/SmbPitchShifter.cpp
+++ b/Engine/Plugins/Runtime/Harmonix/Source/HarmonixDsp/Private/HarmonixDsp/Stretchers/SmbPitchShifter.cpp
@@ -187,7 +187,7 @@ void FSmbPitchShifter::SetSampleRateAndReset(float InSampleRate)
 	Reset();
 }
 
-void FSmbPitchShifter::SetupAndResetImpl()
+void FSmbPitchShifter::SetupAndResetImpl(const FTimeStretchConfig* ShifterSettings)
 {
 	check(MyAudioData);
 	SampleRate = MyAudioData->GetSampleRate();
diff --git a/Engine/Plugins/Runtime/Harmonix/Source/HarmonixDsp/Private/HarmonixDsp/Stretchers/SmbPitchShifter.h b/Engine/Plugins/Runtime/Harmonix/Source/HarmonixDsp/Private/HarmonixDsp/Stretchers/SmbPitchShifter.h
--- a/Engine/Plugins/Runtime/Harmonix/Source/HarmonixDsp/Private/HarmonixDsp/Stretchers/SmbPitchShifter.h
+++ b/Engine/Plugins/Runtime/Harmonix/Source/HarmonixDsp/Private/HarmonixDsp/Stretchers/SmbPitchShifter.h
@@ -122,7 +122,7 @@ public:
 	virtual void Cleanup() override {}
 
 protected:
-	virtual void SetupAndResetImpl() override;
+	virtual void SetupAndResetImpl(const FTimeStretchConfig* ShifterSettings) override;
 
 private:
 	// A bunch of stereo buffers...
diff --git a/Engine/Plugins/Runtime/Harmonix/Source/HarmonixDsp/Public/HarmonixDsp/AudioDataRenderer.h b/Engine/Plugins/Runtime/Harmonix/Source/HarmonixDsp/Public/HarmonixDsp/AudioDataRenderer.h
--- a/Engine/Plugins/Runtime/Harmonix/Source/HarmonixDsp/Public/HarmonixDsp/AudioDataRenderer.h
+++ b/Engine/Plugins/Runtime/Harmonix/Source/HarmonixDsp/Public/HarmonixDsp/AudioDataRenderer.h
@@ -12,6 +12,7 @@ class TAudioBuffer;
 class FGainMatrix;
 class FFusionSampler;
 class FSoundWaveProxy;
+struct FTimeStretchConfig;
 
 class IAudioDataRenderer : public TSharedFromThis<IAudioDataRenderer>
 {
@@ -21,6 +22,7 @@ public:
 	struct FSettings
 	{
 		TSharedPtr<IStretcherAndPitchShifter, ESPMode::ThreadSafe> Shifter = nullptr;
+		const FTimeStretchConfig* ShifterSettings = nullptr;
 		const TArray<FTrackChannelInfo>* TrackChannelInfo = nullptr;
 		const FFusionSampler* Sampler = nullptr;
 	};
diff --git a/Engine/Plugins/Runtime/Harmonix/Source/HarmonixDsp/Public/HarmonixDsp/StretcherAndPitchShifter.h b/Engine/Plugins/Runtime/Harmonix/Source/HarmonixDsp/Public/HarmonixDsp/StretcherAndPitchShifter.h
--- a/Engine/Plugins/Runtime/Harmonix/Source/HarmonixDsp/Public/HarmonixDsp/StretcherAndPitchShifter.h
+++ b/Engine/Plugins/Runtime/Harmonix/Source/HarmonixDsp/Public/HarmonixDsp/StretcherAndPitchShifter.h
@@ -44,7 +44,7 @@ public:
 
 	FName GetFactoryName() const { return MyFactoryName; };
 
-	virtual void ApplyOptions(const TMap<FName, FTypedParameter>& Options) {}
+	virtual void UpdateOptions(const TMap<FName, FTypedParameter>& Options) {}
 
 	virtual int32  GetInputFramesNeeded(int32 NumOutFramesNeeded, float PitchShift, float SpeedShift) = 0;
 	virtual bool HandlesSpeedCorrection() const { return false; }
@@ -53,12 +53,12 @@ public:
 	virtual bool HasCurrentSampleFrame() const { return false; }
 	virtual int32  GetCurrentSampleFrame() const { return 0; }
 
-	void SetSampleSourceReset(TSharedPtr<FSoundWaveProxy> InSoundWave, TSharedPtr<IAudioDataRenderer, ESPMode::ThreadSafe> InAudioRenderer) 
+	void SetSampleSourceReset(TSharedPtr<FSoundWaveProxy> InSoundWave, TSharedPtr<IAudioDataRenderer, ESPMode::ThreadSafe> InAudioRenderer, const FTimeStretchConfig* ShifterSettings)
 	{ 
 		MyAudioData = InSoundWave;
 		MyAudioRenderer = InAudioRenderer;
 
-		SetupAndResetImpl(); 
+		SetupAndResetImpl(ShifterSettings);
 	}
 
 	virtual double Render(TAudioBuffer<float>& OutputData, double Pos, int32 MaxFrame, double ResampleInc, double PitchShift, double Speed, bool MaintainPitchWhenSpeedChanges, bool ShouldHonorLoopPoints, const FGainMatrix& Gain) = 0;
@@ -77,7 +77,7 @@ protected:
 	TSharedPtr<FSoundWaveProxy> MyAudioData;
 	TSharedPtr<IAudioDataRenderer, ESPMode::ThreadSafe> MyAudioRenderer;
 
-	virtual void SetupAndResetImpl() { }
+	virtual void SetupAndResetImpl(const FTimeStretchConfig* ShifterSettings) { }
 
 private:
 
