Dllstorage fixes to fix some target compile errors

diff --git a/Engine/Plugins/Animation/MLDeformer/MLDeformerFramework/Source/MLDeformerFramework/Public/MLDeformerModel.h b/Engine/Plugins/Animation/MLDeformer/MLDeformerFramework/Source/MLDeformerFramework/Public/MLDeformerModel.h
--- a/Engine/Plugins/Animation/MLDeformer/MLDeformerFramework/Source/MLDeformerFramework/Public/MLDeformerModel.h
+++ b/Engine/Plugins/Animation/MLDeformer/MLDeformerFramework/Source/MLDeformerFramework/Public/MLDeformerModel.h
@@ -92,7 +92,7 @@ namespace UE::MLDeformer
 		 * This does the actual render resource init, which means this creates and fills the buffer on the GPU.
 		 * After it successfully initializes, it will empty our VertexMap member array to not store the data in both GPU memory and main memory.
 		 */
-		void InitRHI(FRHICommandListBase& RHICmdList) override;
+		UE_API virtual void InitRHI(FRHICommandListBase& RHICmdList) override;
 
 		/** The array of integers we want to store on the GPU. This buffer will be emptied after successfully calling InitRHI. */
 		TArray<int32> VertexMap;
diff --git a/Engine/Plugins/Experimental/DefaultInstallBundleManager/Source/Public/InstallBundleManagerUtil.h b/Engine/Plugins/Experimental/DefaultInstallBundleManager/Source/Public/InstallBundleManagerUtil.h
--- a/Engine/Plugins/Experimental/DefaultInstallBundleManager/Source/Public/InstallBundleManagerUtil.h
+++ b/Engine/Plugins/Experimental/DefaultInstallBundleManager/Source/Public/InstallBundleManagerUtil.h
@@ -47,7 +47,7 @@ namespace InstallBundleManagerUtil
 	class FContentBuildMetaData : public FJsonSerializable
 	{
 	public:
-		UE_API BEGIN_JSON_SERIALIZER
+		BEGIN_JSON_SERIALIZER
 			JSON_MCI_VALUE(AppName);
 			JSON_MCI_VALUE(BuildVersion);
 			JSON_MCI_VALUE(Platform);
diff --git a/Engine/Plugins/Runtime/Harmonix/Source/HarmonixDsp/Public/HarmonixDsp/GainMatrix.h b/Engine/Plugins/Runtime/Harmonix/Source/HarmonixDsp/Public/HarmonixDsp/GainMatrix.h
--- a/Engine/Plugins/Runtime/Harmonix/Source/HarmonixDsp/Public/HarmonixDsp/GainMatrix.h
+++ b/Engine/Plugins/Runtime/Harmonix/Source/HarmonixDsp/Public/HarmonixDsp/GainMatrix.h
@@ -14,11 +14,11 @@ DECLARE_LOG_CATEGORY_EXTERN(LogGainMatrix, Log, All);
 
 struct FPannerDetails;
 
-class HARMONIXDSP_API alignas(16) FGainMatrix
+class alignas(16) FGainMatrix
 {
 public:
 
-	static void Init();
+	static HARMONIXDSP_API void Init();
 
 	FGainMatrix()
 	{
@@ -68,15 +68,15 @@ public:
 	EAudioBufferChannelLayout GetChannelLayout() const { return ChannelLayout; }
 
 	// pan -1 to 1
-	void SetFromMinusOneToOneSurroundPan(float InGain, float InPan, const FGainTable* InGainTable = nullptr);
-	void SetFromPolarDegrees(float InGain, float InPan, const FGainTable* InGainTable = nullptr);
+	HARMONIXDSP_API void SetFromMinusOneToOneSurroundPan(float InGain, float InPan, const FGainTable* InGainTable = nullptr);
+	HARMONIXDSP_API void SetFromPolarDegrees(float InGain, float InPan, const FGainTable* InGainTable = nullptr);
 	// pan in radians
-	void Set(float InGain, float InPanValue, const FGainTable* InGainTable = nullptr);;
+	HARMONIXDSP_API void Set(float InGain, float InPanValue, const FGainTable* InGainTable = nullptr);;
 	// pan -1 to 1
-	void SetFromLegacyStereo(float InGain, float InPan);
-	void SetFromNewStereo(float InGain, float InPan, float InMax, const FGainTable* InGainTable = nullptr);
-	void Set(float InGain, ESpeakerChannelAssignment InChannelAssignment, const FGainTable* InGainTable = nullptr);
-	void Set(float InGain, const FPannerDetails& InPannerDetails, const FGainTable* InGainTable = nullptr);
+	HARMONIXDSP_API void SetFromLegacyStereo(float InGain, float InPan);
+	HARMONIXDSP_API void SetFromNewStereo(float InGain, float InPan, float InMax, const FGainTable* InGainTable = nullptr);
+	HARMONIXDSP_API void Set(float InGain, ESpeakerChannelAssignment InChannelAssignment, const FGainTable* InGainTable = nullptr);
+	HARMONIXDSP_API void Set(float InGain, const FPannerDetails& InPannerDetails, const FGainTable* InGainTable = nullptr);
 
 	static const FGainMatrix& Unity() { return *sUnity; }
 
@@ -302,14 +302,14 @@ public:
 		return GainsArray[index];
 	}
 
-	static bool  SpeakerChannelAssignmentMapsToOneSpeaker(ESpeakerChannelAssignment InChannelrAssignment, EAudioBufferChannelLayout InChannelLayout);
-	static int32 SpeakerChannelAssignmentToChannelIndex(ESpeakerChannelAssignment InChannelAssignment, EAudioBufferChannelLayout InChannelLayout);
+	static HARMONIXDSP_API bool  SpeakerChannelAssignmentMapsToOneSpeaker(ESpeakerChannelAssignment InChannelrAssignment, EAudioBufferChannelLayout InChannelLayout);
+	static HARMONIXDSP_API int32 SpeakerChannelAssignmentToChannelIndex(ESpeakerChannelAssignment InChannelAssignment, EAudioBufferChannelLayout InChannelLayout);
 
 private:
 
 	alignas(16) FChannelGains GainsArray[HarmonixDsp::FAudioBuffer::kMaxChannelsInAudioBuffer] = {};
 
-	static FGainMatrix* sUnity;
+	static HARMONIXDSP_API FGainMatrix* sUnity;
 
 	uint8 NumInChannels = HarmonixDsp::FAudioBuffer::kMaxChannelsInAudioBuffer;
 	uint8 NumOutChannels = HarmonixDsp::FAudioBuffer::kMaxChannelsInAudioBuffer;
