Fix a trivial typo

diff --git a/Engine/Plugins/Runtime/Harmonix/Source/HarmonixMidi/Private/HarmonixMidi/MidiFile.cpp b/Engine/Plugins/Runtime/Harmonix/Source/HarmonixMidi/Private/HarmonixMidi/MidiFile.cpp
--- a/Engine/Plugins/Runtime/Harmonix/Source/HarmonixMidi/Private/HarmonixMidi/MidiFile.cpp
+++ b/Engine/Plugins/Runtime/Harmonix/Source/HarmonixMidi/Private/HarmonixMidi/MidiFile.cpp
@@ -239,7 +239,7 @@ void UMidiFile::SaveStdMidiFile(TSharedPtr<FArchive> Archive, const FString& Fil
 
 void UMidiFile::CreateFrameBasedMusicMap(UFrameBasedMusicMap* Map) const
 {
-	return TheMidiData.SongMaps.FillInFrameBasedMusicMap(Map);
+	TheMidiData.SongMaps.FillInFrameBasedMusicMap(Map);
 }
 
 FMarkerProviderResults UMidiFile::GatherMarkers(const UFrameBasedMusicMap* Map) const
