Add Enable pin to Midi Pulse Generator node

Properly implement Enable setting in MIDI Clock Subdivision Trigger node

#rnx

diff --git a/Engine/Plugins/Runtime/Harmonix/Source/HarmonixMetasound/Private/HarmonixMetasound/Nodes/MidiClockSubdivisionTriggerNode.cpp b/Engine/Plugins/Runtime/Harmonix/Source/HarmonixMetasound/Private/HarmonixMetasound/Nodes/MidiClockSubdivisionTriggerNode.cpp
--- a/Engine/Plugins/Runtime/Harmonix/Source/HarmonixMetasound/Private/HarmonixMetasound/Nodes/MidiClockSubdivisionTriggerNode.cpp
+++ b/Engine/Plugins/Runtime/Harmonix/Source/HarmonixMetasound/Private/HarmonixMetasound/Nodes/MidiClockSubdivisionTriggerNode.cpp
@@ -216,7 +216,10 @@ namespace HarmonixMetasound::Nodes::MidiClockSubdivisionTriggerNode
 		TriggerOutPin->AdvanceBlock();
 		PulseGenerator.Process(*MidiClockInPin, [this](const Harmonix::Midi::Ops::FPulseGenerator::FPulseInfo& Pulse)
 		{
-			TriggerOutPin->TriggerFrame(Pulse.BlockFrameIndex);
+			if (*EnableInPin)
+			{
+				TriggerOutPin->TriggerFrame(Pulse.BlockFrameIndex);
+			}
 		});
 	}
 }
diff --git a/Engine/Plugins/Runtime/Harmonix/Source/HarmonixMetasound/Private/HarmonixMetasound/Nodes/MidiPulseGeneratorNode.cpp b/Engine/Plugins/Runtime/Harmonix/Source/HarmonixMetasound/Private/HarmonixMetasound/Nodes/MidiPulseGeneratorNode.cpp
--- a/Engine/Plugins/Runtime/Harmonix/Source/HarmonixMetasound/Private/HarmonixMetasound/Nodes/MidiPulseGeneratorNode.cpp
+++ b/Engine/Plugins/Runtime/Harmonix/Source/HarmonixMetasound/Private/HarmonixMetasound/Nodes/MidiPulseGeneratorNode.cpp
@@ -28,6 +28,7 @@ namespace HarmonixMetasound::Nodes::MidiPulseGeneratorNode
 	namespace Inputs
 	{
 		DEFINE_METASOUND_PARAM_ALIAS(MidiClock, CommonPinNames::Inputs::MidiClock);
+		DEFINE_METASOUND_PARAM_ALIAS(Enable, CommonPinNames::Inputs::Enable);
 		DEFINE_INPUT_METASOUND_PARAM(Interval, "Interval", "The musical time interval at which to send out the pulse");
 		DEFINE_INPUT_METASOUND_PARAM(IntervalMultiplier, "Interval Multiplier", "Multiplies the interval, 1 to use just the value of Interval");
 		DEFINE_INPUT_METASOUND_PARAM(Offset, "Offset", "Offsets the pulse by a musical time");
@@ -81,6 +82,7 @@ namespace HarmonixMetasound::Nodes::MidiPulseGeneratorNode
 			{
 				FInputVertexInterface {
 					TInputDataVertex<FMidiClock>(METASOUND_GET_PARAM_NAME_AND_METADATA(Inputs::MidiClock)),
+					TInputDataVertex<bool>(METASOUND_GET_PARAM_NAME_AND_METADATA(Inputs::Enable), true),
 					TInputDataVertex<FEnumMidiClockSubdivisionQuantizationType>(
 						METASOUND_GET_PARAM_NAME_AND_METADATA(Inputs::Interval), static_cast<int32>(DefaultInterval.Interval)),
 					TInputDataVertex<int32>(
@@ -110,6 +112,7 @@ namespace HarmonixMetasound::Nodes::MidiPulseGeneratorNode
 		struct FInputs
 		{
 			FMidiClockReadRef Clock;
+			Metasound::FBoolReadRef Enable;
 			Metasound::FEnumMidiClockSubdivisionQuantizationReadRef Interval;
 			Metasound::FInt32ReadRef IntervalMultiplier;
 			Metasound::FEnumMidiClockSubdivisionQuantizationReadRef Offset;
@@ -130,6 +133,7 @@ namespace HarmonixMetasound::Nodes::MidiPulseGeneratorNode
 			FInputs Inputs
 			{
 				InParams.InputData.GetOrCreateDefaultDataReadReference<FMidiClock>(Inputs::MidiClockName, InParams.OperatorSettings),
+				InParams.InputData.GetOrCreateDefaultDataReadReference<bool>(Inputs::EnableName, InParams.OperatorSettings),
 				InParams.InputData.GetOrCreateDefaultDataReadReference<Metasound::FEnumMidiClockSubdivisionQuantizationType>(
 					Inputs::IntervalName, InParams.OperatorSettings),
 				InParams.InputData.GetOrCreateDefaultDataReadReference<int32>(Inputs::IntervalMultiplierName, InParams.OperatorSettings),
@@ -193,6 +197,7 @@ namespace HarmonixMetasound::Nodes::MidiPulseGeneratorNode
 	private:
 		void ApplyParameters()
 		{
+			PulseGenerator.Enable(*Inputs.Enable);
 			PulseGenerator.Track = *Inputs.Track;
 			PulseGenerator.Channel = *Inputs.Channel;
 			PulseGenerator.NoteNumber = *Inputs.NoteNumber;
