Fix failing automation tests Metasound

MidiPulseGeneratorNode

Bind "Enable" pin

#tests Metasound Automated tests
#rnx

diff --git a/Engine/Plugins/Runtime/Harmonix/Source/HarmonixMetasound/Private/HarmonixMetasound/Nodes/MidiPulseGeneratorNode.cpp b/Engine/Plugins/Runtime/Harmonix/Source/HarmonixMetasound/Private/HarmonixMetasound/Nodes/MidiPulseGeneratorNode.cpp
--- a/Engine/Plugins/Runtime/Harmonix/Source/HarmonixMetasound/Private/HarmonixMetasound/Nodes/MidiPulseGeneratorNode.cpp
+++ b/Engine/Plugins/Runtime/Harmonix/Source/HarmonixMetasound/Private/HarmonixMetasound/Nodes/MidiPulseGeneratorNode.cpp
@@ -170,6 +170,7 @@ namespace HarmonixMetasound::Nodes::MidiPulseGeneratorNode
 		virtual void BindInputs(Metasound::FInputVertexInterfaceData& InVertexData) override
 		{
 			InVertexData.BindReadVertex(Inputs::MidiClockName, Inputs.Clock);
+			InVertexData.BindReadVertex(Inputs::EnableName, Inputs.Enable);
 			InVertexData.BindReadVertex(Inputs::IntervalName, Inputs.Interval);
 			InVertexData.BindReadVertex(Inputs::IntervalMultiplierName, Inputs.IntervalMultiplier);
 			InVertexData.BindReadVertex(Inputs::OffsetName, Inputs.Offset);
