Fix a bug where we can receive a generator detached callback for a different generator than the one we are tracking.

diff --git a/Engine/Plugins/Runtime/Harmonix/Source/HarmonixMetasound/Private/HarmonixMetasound/Components/MetasoundMusicClockDriver.cpp b/Engine/Plugins/Runtime/Harmonix/Source/HarmonixMetasound/Private/HarmonixMetasound/Components/MetasoundMusicClockDriver.cpp
--- a/Engine/Plugins/Runtime/Harmonix/Source/HarmonixMetasound/Private/HarmonixMetasound/Components/MetasoundMusicClockDriver.cpp
+++ b/Engine/Plugins/Runtime/Harmonix/Source/HarmonixMetasound/Private/HarmonixMetasound/Components/MetasoundMusicClockDriver.cpp
@@ -422,6 +422,7 @@ void FMetasoundMusicClockDriver::DetachAllCallbacks()
 void FMetasoundMusicClockDriver::OnGeneratorAttached()
 {
 	WasEverConnected = true;
+	UnderlyingGenerator = CurrentGeneratorHandle->GetGenerator();
 	SetClockHistory(UMidiClockUpdateSubsystem::GetOrCreateClockHistory(MidiSongPosAnalyzerAddress));
 	SmoothedAudioRenderClockHistoryCursor = ClockHistory->CreateReadCursor();
 	SmoothedPlayerExperienceClockHistoryCursor = ClockHistory->CreateReadCursor();
@@ -455,6 +456,13 @@ void FMetasoundMusicClockDriver::OnGeneratorIOUpdatedWithChanges(const TArray<Me
 
 void FMetasoundMusicClockDriver::OnGeneratorDetached()
 {
+	TWeakPtr<Metasound::FMetasoundGenerator> DetachingGenerator = CurrentGeneratorHandle->GetGenerator();
+	if (DetachingGenerator != UnderlyingGenerator)
+	{
+		// yes... it is possible we will get a detach callback for a generator that is not ours! yikes.
+		return;
+	}
+
 	using namespace HarmonixMetasound::Analysis;
 	if (GetState() != EMusicClockState::Stopped)
 	{
diff --git a/Engine/Plugins/Runtime/Harmonix/Source/HarmonixMetasound/Private/HarmonixMetasound/Components/MetasoundMusicClockDriver.h b/Engine/Plugins/Runtime/Harmonix/Source/HarmonixMetasound/Private/HarmonixMetasound/Components/MetasoundMusicClockDriver.h
--- a/Engine/Plugins/Runtime/Harmonix/Source/HarmonixMetasound/Private/HarmonixMetasound/Components/MetasoundMusicClockDriver.h
+++ b/Engine/Plugins/Runtime/Harmonix/Source/HarmonixMetasound/Private/HarmonixMetasound/Components/MetasoundMusicClockDriver.h
@@ -51,6 +51,7 @@ private:
 	TWeakObjectPtr<UAudioComponent> AudioComponentToWatch;
 	// We need a strong object ptr to this next thing since we will be the only one holding a reference to it...
 	TStrongObjectPtr<UMetasoundGeneratorHandle> CurrentGeneratorHandle;
+	TWeakPtr<Metasound::FMetasoundGenerator> UnderlyingGenerator;
 
 	Metasound::Frontend::FAnalyzerAddress MidiSongPosAnalyzerAddress;
 
