MusicalAsset as Musical Clock Source for Sequencer

diff --git a/Engine/Plugins/Runtime/Harmonix/Source/HarmonixMetasound/Private/HarmonixMetasound/DataTypes/HarmonixMetasoundMusicAsset.cpp b/Engine/Plugins/Runtime/Harmonix/Source/HarmonixMetasound/Private/HarmonixMetasound/DataTypes/HarmonixMetasoundMusicAsset.cpp
new file mode 100644
--- /dev/null
+++ b/Engine/Plugins/Runtime/Harmonix/Source/HarmonixMetasound/Private/HarmonixMetasound/DataTypes/HarmonixMetasoundMusicAsset.cpp
@@ -0,0 +1,16 @@
+// Copyright Epic Games, Inc. All Rights Reserved.
+
+#include "HarmonixMetasound/DataTypes/HarmonixMetasoundMusicAsset.h"
+#include "HarmonixMetasound/Interfaces/HarmonixMusicInterfaces.h"
+#include "MetasoundSource.h"
+
+bool UHarmonixMetasoundMusicAsset::AssetIsMissingHarmoixMusicInterface(const FAssetData& AssetData)
+{
+#if WITH_EDITORONLY_DATA
+	bool GotDocInfo = false;
+	Metasound::Frontend::FMetaSoundClassInfo ClassInfo(AssetData);
+	return !ClassInfo.InheritsInterface(HarmonixMetasound::MusicAssetInterface::FrontendVersion.Name);
+#else
+	return false;
+#endif
+}
diff --git a/Engine/Plugins/Runtime/Harmonix/Source/HarmonixMetasound/Private/HarmonixMetasound/DataTypes/HarmonixMusicAsset.cpp b/Engine/Plugins/Runtime/Harmonix/Source/HarmonixMetasound/Private/HarmonixMetasound/DataTypes/HarmonixMusicAsset.cpp
new file mode 100644
--- /dev/null
+++ b/Engine/Plugins/Runtime/Harmonix/Source/HarmonixMetasound/Private/HarmonixMetasound/DataTypes/HarmonixMusicAsset.cpp
@@ -0,0 +1,43 @@
+// Copyright Epic Games, Inc. All Rights Reserved.
+
+#include "HarmonixMetasound/DataTypes/HarmonixMusicAsset.h"
+#include "HarmonixMetasound/DataTypes/HarmonixMusicHandle.h"
+#include "HarmonixMidi/MidiFile.h"
+#include "MetasoundSource.h"
+
+void UHarmonixMusicAsset::CreateFrameBasedMusicMap(UFrameBasedMusicMap* Map) const
+{
+	if (!MidiSongMap)
+	{
+		return;
+	}
+
+	check(MidiSongMap->GetSongMaps());
+	MidiSongMap->GetSongMaps()->FillInFrameBasedMusicMap(Map);
+}
+
+FMarkerProviderResults UHarmonixMusicAsset::GatherMarkers(const UFrameBasedMusicMap* Map) const
+{
+	if (!MidiSongMap)
+	{
+		return FMarkerProviderResults();
+	}
+	return MidiSongMap->GatherMarkers(Map);
+}
+
+FMusicHandle UHarmonixMusicAsset::PrepareToPlay_Internal(UObject* PlaybackContext, UAudioComponent* OnComponent, float FromSeconds, bool IsAudition)
+{
+	UHarmonixMusicHandle* MusicHandle = UHarmonixMusicHandle::Instantiate(this, PlaybackContext, OnComponent, FromSeconds, IsAudition);
+	return MusicHandle;
+}
+
+FMusicHandle UHarmonixMusicAsset::Play_Internal(UObject* PlaybackContext, UAudioComponent* OnComponent, float FromSeconds, bool IsAudition)
+{
+	FMusicHandle MusicHandle = PrepareToPlay(PlaybackContext, OnComponent, FromSeconds, IsAudition);
+	if (MusicHandle)
+	{
+		MusicHandle->Play();
+	}
+	return MusicHandle;
+}
+
diff --git a/Engine/Plugins/Runtime/Harmonix/Source/HarmonixMetasound/Private/HarmonixMetasound/DataTypes/HarmonixMusicHandle.cpp b/Engine/Plugins/Runtime/Harmonix/Source/HarmonixMetasound/Private/HarmonixMetasound/DataTypes/HarmonixMusicHandle.cpp
new file mode 100644
--- /dev/null
+++ b/Engine/Plugins/Runtime/Harmonix/Source/HarmonixMetasound/Private/HarmonixMetasound/DataTypes/HarmonixMusicHandle.cpp
@@ -0,0 +1,166 @@
+// Copyright Epic Games, Inc. All Rights Reserved.
+
+#include "HarmonixMetasound/DataTypes/HarmonixMusicHandle.h"
+
+#include "Components/AudioComponent.h"
+#include "HarmonixMetasound/Components/MusicClockComponent.h"
+#include "HarmonixMetasound/DataTypes/HarmonixMetasoundMusicAsset.h"
+#include "HarmonixMetasound/Interfaces/HarmonixMusicInterfaces.h"
+#include "MetasoundSource.h"
+
+DEFINE_LOG_CATEGORY_STATIC(LogHarmonixMusicHandle, Log, All)
+
+UHarmonixMusicHandle* UHarmonixMusicHandle::Instantiate(UHarmonixMusicAsset* Asset, UObject* PlaybackContext, UAudioComponent* OnComponent, float FromSeconds, bool bIsAudition)
+{
+	check(PlaybackContext);
+	check(!OnComponent || PlaybackContext->GetWorld() == OnComponent->GetWorld());
+
+	// Don't crash on a poorly authored asset...
+	if (!ensureAlwaysMsgf(Asset->GetMetaSoundSource(), TEXT("Attempt to instantiate a HarmonixMusicHandle for a HarmonixMusicAsset that has no MetaSound failed.")))
+	{
+		return nullptr;
+	}
+
+	UWorld* World = PlaybackContext->GetWorld();
+
+	UHarmonixMusicHandle* TheHandle = NewObject<UHarmonixMusicHandle>(World, NAME_None, RF_Transient);
+	TheHandle->PlayingAsset = Asset;
+
+	// We need an AudioComponent to play on. If we were not provided one, spawn one...
+	if (!OnComponent)
+	{
+		TheHandle->AudioComponent = NewObject<UAudioComponent>(World, NAME_None, EObjectFlags::RF_Transient);
+		if (!TheHandle->AudioComponent)
+		{
+			return nullptr;
+		}
+		TheHandle->AudioComponent->bAllowSpatialization = false;
+	}
+	else
+	{
+		TheHandle->AudioComponent = OnComponent;
+		TheHandle->AudioComponent->Stop();
+	}
+	if (bIsAudition)
+	{
+		TheHandle->AudioComponent->SetUISound(true);
+	}
+
+	TheHandle->AudioComponent->Sound = Asset->GetMetaSoundSource();
+	TheHandle->AudioComponent->Play();
+
+	// Now, because we are a music handle, we need a music clock. Spawn one...	
+	TheHandle->MusicClockComponent = NewObject<UMusicClockComponent>(World, NAME_None, RF_Transient);
+	TheHandle->MusicClockComponent->MetasoundOutputName = HarmonixMetasound::MusicAssetInterface::MidiClockOut.Resolve();
+	TheHandle->MusicClockComponent->ConnectToMetasoundOnAudioComponent(TheHandle->AudioComponent);
+	TheHandle->MusicClockComponent->Start();
+
+	// NOTE: For now we are going to ignore "FromSeconds" as...
+	// a) We aren't supporting 'async preparation' of music assets, and therefore...
+	// b) We can just wait for the Play call and use the 'FromSeconds' from there.
+
+	// Since we aren't yet supporting async preparation of music assets for playback,
+	// just mark ourselves as ready to play.
+	TheHandle->CurrentTransportState = EMusicHanldeTransportState::ReadyToPlay;
+
+	return TheHandle;
+}
+
+bool UHarmonixMusicHandle::IsValid() const
+{
+	if (PlayingAsset && AudioComponent)
+	{
+		if (!bHasStarted || (bHasStarted && AudioComponent->GetPlayState() != EAudioComponentPlayState::Stopped))
+		{
+			return true;
+		}
+	}
+	return false;
+}
+
+bool UHarmonixMusicHandle::Play(float FromSeconds)
+{
+	using namespace HarmonixMetasound::MusicAssetInterface;
+
+	if (!AudioComponent || !MusicClockComponent)
+	{
+		return false;
+	}
+
+	FAudioParameter SeekTargetSeconds(SeekTargetSecondsIn.Resolve(), FromSeconds);
+	AudioComponent->Activate();
+	AudioComponent->SetParameter(MoveTemp(SeekTargetSeconds));
+	AudioComponent->SetTriggerParameter(SeekIn.Resolve());
+	AudioComponent->SetTriggerParameter(PlayIn.Resolve());
+	MusicClockComponent->Start();
+	CurrentTransportState = EMusicHanldeTransportState::Playing;
+	return true;
+}
+
+void UHarmonixMusicHandle::Tick(float DeltaSeconds)
+{
+	if (MusicClockComponent)
+	{
+		MusicClockComponent->TickComponentInternal();
+	}
+}
+
+void UHarmonixMusicHandle::Pause()
+{
+	using namespace HarmonixMetasound::MusicAssetInterface;
+	AudioComponent->SetTriggerParameter(PauseIn.Resolve());
+	CurrentTransportState = EMusicHanldeTransportState::Paused;
+}
+
+void UHarmonixMusicHandle::Continue()
+{
+	using namespace HarmonixMetasound::MusicAssetInterface;
+	AudioComponent->SetTriggerParameter(ContinueIn.Resolve());
+	CurrentTransportState = EMusicHanldeTransportState::Playing;
+}
+
+void UHarmonixMusicHandle::Stop_Internal()
+{
+	using namespace HarmonixMetasound::MusicAssetInterface;
+	if (MusicClockComponent)
+	{
+		MusicClockComponent->Stop();
+		MusicClockComponent = nullptr;
+	}
+	if (AudioComponent)
+	{
+		AudioComponent->SetTriggerParameter(StopIn.Resolve());
+		AudioComponent->Deactivate();
+		AudioComponent = nullptr;
+	}
+	PlayingAsset = nullptr;
+	CurrentTransportState = EMusicHanldeTransportState::Stale;
+}
+
+void UHarmonixMusicHandle::Kill_Internal()
+{
+	using namespace HarmonixMetasound::MusicAssetInterface;
+	if (MusicClockComponent)
+	{
+		MusicClockComponent->Stop();
+		MusicClockComponent = nullptr;
+	}
+	if (AudioComponent)
+	{
+		AudioComponent->SetTriggerParameter(KillIn.Resolve());
+		AudioComponent->Deactivate();
+		AudioComponent = nullptr;
+	}
+	PlayingAsset = nullptr;
+	CurrentTransportState = EMusicHanldeTransportState::Stale;
+}
+
+bool UHarmonixMusicHandle::IsUsingAsset(const UObject* Asset) const
+{
+	return IsValid() && PlayingAsset == Asset;
+}
+
+TScriptInterface<IMusicEnvironmentClockSource> UHarmonixMusicHandle::GetMusicClockSource()
+{
+	return MusicClockComponent;
+}
diff --git a/Engine/Plugins/Runtime/Harmonix/Source/HarmonixMetasound/Private/HarmonixMetasound/Interfaces/HarmonixMusicInterfaces.cpp b/Engine/Plugins/Runtime/Harmonix/Source/HarmonixMetasound/Private/HarmonixMetasound/Interfaces/HarmonixMusicInterfaces.cpp
--- a/Engine/Plugins/Runtime/Harmonix/Source/HarmonixMetasound/Private/HarmonixMetasound/Interfaces/HarmonixMusicInterfaces.cpp
+++ b/Engine/Plugins/Runtime/Harmonix/Source/HarmonixMetasound/Private/HarmonixMetasound/Interfaces/HarmonixMusicInterfaces.cpp
@@ -22,15 +22,15 @@ namespace HarmonixMetasound
 			return Version;
 		}
 
-		static const FMetasoundFrontendVersion FrontendVersion{ AUDIO_PARAMETER_INTERFACE_NAMESPACE, { 1, 0 } };
-		static const FLazyName PlayIn(AUDIO_PARAMETER_INTERFACE_MEMBER_DEFINE("Play"));
-		static const FLazyName PauseIn(AUDIO_PARAMETER_INTERFACE_MEMBER_DEFINE("Pause"));
-		static const FLazyName ContinueIn(AUDIO_PARAMETER_INTERFACE_MEMBER_DEFINE("Continue"));
-		static const FLazyName StopIn(AUDIO_PARAMETER_INTERFACE_MEMBER_DEFINE("Stop"));
-		static const FLazyName KillIn(AUDIO_PARAMETER_INTERFACE_MEMBER_DEFINE("Kill"));
-		static const FLazyName SeekIn(AUDIO_PARAMETER_INTERFACE_MEMBER_DEFINE("Seek"));
-		static const FLazyName SeekTargetIn(AUDIO_PARAMETER_INTERFACE_MEMBER_DEFINE("SeekTargetSeconds"));
-		static const FLazyName MidiClockOut(AUDIO_PARAMETER_INTERFACE_MEMBER_DEFINE("MIDI Clock"));
+		const FMetasoundFrontendVersion FrontendVersion{ AUDIO_PARAMETER_INTERFACE_NAMESPACE, { 1, 0 } };
+		const FLazyName PlayIn(AUDIO_PARAMETER_INTERFACE_MEMBER_DEFINE("Play"));
+		const FLazyName PauseIn(AUDIO_PARAMETER_INTERFACE_MEMBER_DEFINE("Pause"));
+		const FLazyName ContinueIn(AUDIO_PARAMETER_INTERFACE_MEMBER_DEFINE("Continue"));
+		const FLazyName StopIn(AUDIO_PARAMETER_INTERFACE_MEMBER_DEFINE("Stop"));
+		const FLazyName KillIn(AUDIO_PARAMETER_INTERFACE_MEMBER_DEFINE("Kill"));
+		const FLazyName SeekIn(AUDIO_PARAMETER_INTERFACE_MEMBER_DEFINE("Seek"));
+		const FLazyName SeekTargetSecondsIn(AUDIO_PARAMETER_INTERFACE_MEMBER_DEFINE("SeekTargetSeconds"));
+		const FLazyName MidiClockOut(AUDIO_PARAMETER_INTERFACE_MEMBER_DEFINE("MIDI Clock"));
 
 		Audio::FParameterInterfacePtr CreateInterface()
 		{
@@ -100,7 +100,7 @@ namespace HarmonixMetasound
 							LOCTEXT("HarmonixMusicAssetInterfaceSeekTarget", "SeekTargetSeconds"),
 							LOCTEXT("HarmonixMusicAssetInterfaceSeekTarget_Description", "The position the music should seek to."),
 							GetMetasoundDataTypeName<float>(),
-							SeekTargetIn.Resolve(),
+							SeekTargetSecondsIn.Resolve(),
 							FText(),
 							6
 						});
diff --git a/Engine/Plugins/Runtime/Harmonix/Source/HarmonixMetasound/Private/HarmonixMetasound/Interfaces/HarmonixMusicInterfaces.h b/Engine/Plugins/Runtime/Harmonix/Source/HarmonixMetasound/Private/HarmonixMetasound/Interfaces/HarmonixMusicInterfaces.h
--- a/Engine/Plugins/Runtime/Harmonix/Source/HarmonixMetasound/Private/HarmonixMetasound/Interfaces/HarmonixMusicInterfaces.h
+++ b/Engine/Plugins/Runtime/Harmonix/Source/HarmonixMetasound/Private/HarmonixMetasound/Interfaces/HarmonixMusicInterfaces.h
@@ -10,6 +10,16 @@ namespace HarmonixMetasound
 	{
 		const FMetasoundFrontendVersion& GetVersion();
 		Audio::FParameterInterfacePtr CreateInterface();
+
+		const extern FMetasoundFrontendVersion FrontendVersion;
+		const extern FLazyName PlayIn;
+		const extern FLazyName PauseIn;
+		const extern FLazyName ContinueIn;
+		const extern FLazyName StopIn;
+		const extern FLazyName KillIn;
+		const extern FLazyName SeekIn;
+		const extern FLazyName SeekTargetSecondsIn;
+		const extern FLazyName MidiClockOut;
 	}
 
 	void RegisterHarmonixMetasoundMusicInterfaces();
diff --git a/Engine/Plugins/Runtime/Harmonix/Source/HarmonixMetasound/Public/HarmonixMetasound/DataTypes/HarmonixMetasoundMusicAsset.h b/Engine/Plugins/Runtime/Harmonix/Source/HarmonixMetasound/Public/HarmonixMetasound/DataTypes/HarmonixMetasoundMusicAsset.h
new file mode 100644
--- /dev/null
+++ b/Engine/Plugins/Runtime/Harmonix/Source/HarmonixMetasound/Public/HarmonixMetasound/DataTypes/HarmonixMetasoundMusicAsset.h
@@ -0,0 +1,41 @@
+// Copyright Epic Games, Inc. All Rights Reserved.
+
+#pragma once
+
+#include "AssetRegistry/AssetData.h"
+#include "HarmonixMusicAsset.h"
+#include "MusicTypes/MusicalAsset.h"
+#include "MusicTypes/MusicHandle.h"
+
+#include "HarmonixMetasoundMusicAsset.generated.h"
+
+/**
+ * This file defines a UHarmonixMetasoundMusicAsset, an implementation of the UHarmonixMusicAsset
+ * class. 
+ */
+
+class UMetaSoundSource;
+class UMidiFile;
+
+/**
+ * UHarmonixMetasoundMusicAsset:
+ * In this case, the music asset is comprised of a MetaSound that uses the Harmonix MetaSound nodes to
+ * generate musical time, and a midi file which maps the tempo and time signatures that define the music's
+ * form.
+ */
+UCLASS(BlueprintType, Category = "Music", Meta = (DisplayName = "Harmonix MetaSound Music"))
+class HARMONIXMETASOUND_API UHarmonixMetasoundMusicAsset : public UHarmonixMusicAsset
+{
+GENERATED_BODY()
+
+public:
+	UPROPERTY(EditAnywhere, Category = "Music", meta = (GetAssetFilter = "AssetIsMissingHarmoixMusicInterface"))
+	TObjectPtr<UMetaSoundSource> MetaSoundSource;
+
+	UFUNCTION()
+	static bool AssetIsMissingHarmoixMusicInterface(const FAssetData& AssetData);
+
+protected:
+	virtual UMetaSoundSource* GetMetaSoundSource() { return MetaSoundSource; }
+};
+
diff --git a/Engine/Plugins/Runtime/Harmonix/Source/HarmonixMetasound/Public/HarmonixMetasound/DataTypes/HarmonixMusicAsset.h b/Engine/Plugins/Runtime/Harmonix/Source/HarmonixMetasound/Public/HarmonixMetasound/DataTypes/HarmonixMusicAsset.h
new file mode 100644
--- /dev/null
+++ b/Engine/Plugins/Runtime/Harmonix/Source/HarmonixMetasound/Public/HarmonixMetasound/DataTypes/HarmonixMusicAsset.h
@@ -0,0 +1,44 @@
+// Copyright Epic Games, Inc. All Rights Reserved.
+
+#pragma once
+#include "MusicTypes/MusicalAsset.h"
+#include "MusicTypes/MusicHandle.h"
+
+#include "HarmonixMusicAsset.generated.h"
+
+/**
+ * This file defines a UHarmonixMusicAsset, an implementation of the pure virtual IMusicalAsset 
+ * interface that will act as a base class for UHarmonixMetasoundMusicAsset and UHarmonixWaveMusicAsset.
+ * 
+ * As such, it can return an FMusicHandle when it is asked to PrepareToPlay or Play.
+ */
+
+class UMetaSoundSource;
+class UMidiFile;
+
+/**
+ * UHarmonixMusicAsset
+ * An implementation of the pure virtual IMusicalAsset interface that will act as a base class
+ * for UHarmonixMetasoundMusicAsset and UHarmonixWaveMusicAsset. As such, it can return an
+ * FMusicHandle when it is asked to PrepareToPlay or Play.
+ */
+UCLASS(BlueprintType, Category = "Music", Meta = (DisplayName = "Harmonix MetaSound Music"))
+class HARMONIXMETASOUND_API UHarmonixMusicAsset : public UObject, public IMusicalAsset
+{
+GENERATED_BODY()
+
+public:
+	UPROPERTY(EditAnywhere, Category = "Music")
+	TObjectPtr<UMidiFile> MidiSongMap;
+
+	virtual UMetaSoundSource* GetMetaSoundSource() { return nullptr; }
+
+	//** BEGIN IMusicalAsset
+public:
+	virtual void CreateFrameBasedMusicMap(UFrameBasedMusicMap* Map) const override;
+	virtual FMarkerProviderResults GatherMarkers(const UFrameBasedMusicMap* Map) const override;
+protected:
+	virtual FMusicHandle PrepareToPlay_Internal(UObject* PlaybackContext, UAudioComponent* OnComponennt = nullptr, float FromSeconds = 0.0f, bool IsAudition = false) override;
+	virtual FMusicHandle Play_Internal(UObject* PlaybackContext, UAudioComponent* OnComponennt = nullptr, float FromSeconds = 0.0f, bool IsAudition = false) override;
+	//** END IMusicalAsset
+};	
\ No newline at end of file
diff --git a/Engine/Plugins/Runtime/Harmonix/Source/HarmonixMetasound/Public/HarmonixMetasound/DataTypes/HarmonixMusicHandle.h b/Engine/Plugins/Runtime/Harmonix/Source/HarmonixMetasound/Public/HarmonixMetasound/DataTypes/HarmonixMusicHandle.h
new file mode 100644
--- /dev/null
+++ b/Engine/Plugins/Runtime/Harmonix/Source/HarmonixMetasound/Public/HarmonixMetasound/DataTypes/HarmonixMusicHandle.h
@@ -0,0 +1,63 @@
+// Copyright Epic Games, Inc. All Rights Reserved.
+
+#pragma once
+
+#include "MusicTypes/MusicalAsset.h"
+#include "MusicTypes/MusicHandle.h"
+#include "MusicEnvironmentClockSource.h"
+#include "UObject/Interface.h"
+#include "HarmonixMusicHandle.generated.h"
+
+class UAudioComponent;
+class UHarmonixMusicAsset;
+class UMusicClockComponent;
+class USoundBase;
+
+/**
+ * UHarmonixMusicHandle
+ * 
+ * This defines a concrete implementation of the IMusicHandle interface. It is a wrapper around
+ * a UAudioComponent and a UMusicClockComponent. Using those two components this class can
+ * control the transport of the playing music (play, pause, continue, seek stop, kill) and
+ * provide the holder with a source of musical time.
+ *
+ * This is the 'flavor' of IMusicHandle returned by the UHarmonixMetasoundMusicAsset class (among others
+ * in the future) when it is asked to instantiate itself (play).
+ */
+
+UCLASS()
+class HARMONIXMETASOUND_API UHarmonixMusicHandle : public UObject, public IMusicHandle
+{
+public:
+	GENERATED_BODY()
+public:
+	static UHarmonixMusicHandle* Instantiate(UHarmonixMusicAsset* Asset, UObject* PlaybackContext, UAudioComponent* OnComponent, float FromSeconds = 0.0f, bool bIsAudition = false);
+
+	virtual bool IsValid() const override;
+	virtual void Tick(float DeltaSeconds) override;
+
+	//** Transport
+	virtual bool Play(float FromSeconds = 0.0f) override;
+	virtual void Pause() override;
+	virtual void Continue() override;
+protected:
+	virtual void Stop_Internal() override;
+	virtual void Kill_Internal() override;
+public:
+
+	virtual bool IsUsingAsset(const UObject* Asset) const override;
+
+	virtual TScriptInterface<IMusicEnvironmentClockSource> GetMusicClockSource() override;
+
+private:
+	UPROPERTY(Transient)
+	TObjectPtr<UAudioComponent> AudioComponent;
+
+	UPROPERTY(Transient)
+	TObjectPtr<UMusicClockComponent> MusicClockComponent;
+
+	UPROPERTY(Transient)
+	TObjectPtr<UHarmonixMusicAsset> PlayingAsset;
+
+	bool bHasStarted = false;
+};
diff --git a/Engine/Plugins/Runtime/Harmonix/Source/HarmonixMetasoundEditor/Private/AssetDefinition_MusicAssets.cpp b/Engine/Plugins/Runtime/Harmonix/Source/HarmonixMetasoundEditor/Private/AssetDefinition_MusicAssets.cpp
new file mode 100644
--- /dev/null
+++ b/Engine/Plugins/Runtime/Harmonix/Source/HarmonixMetasoundEditor/Private/AssetDefinition_MusicAssets.cpp
@@ -0,0 +1,41 @@
+// Copyright Epic Games, Inc. All Rights Reserved.
+
+#include "AssetDefinition_MusicAssets.h"
+#include "HarmonixMetasound/DataTypes/HarmonixMetasoundMusicAsset.h"
+#include "HarmonixMetasoundSlateStyle.h"
+
+TSoftClassPtr<UObject> UAssetDefinition_MetasoundMusic::GetAssetClass() const
+{
+	return UHarmonixMetasoundMusicAsset::StaticClass();
+}
+
+FText UAssetDefinition_MetasoundMusic::GetAssetDisplayName() const
+{
+	return NSLOCTEXT("AssetTypeActions", "HarmonixMetasoundMusicDefinition", "Harmonix MetaSound Music Asset");
+}
+
+FLinearColor  UAssetDefinition_MetasoundMusic::GetAssetColor() const
+{
+	return FLinearColor(1.0f, 0.5f, 0.0f);
+}
+
+TConstArrayView<FAssetCategoryPath> UAssetDefinition_MetasoundMusic::GetAssetCategories() const
+{
+	static const auto Categories = { EAssetCategoryPaths::Audio / NSLOCTEXT("Harmonix", "HmxAssetCategoryName", "Harmonix") };
+	return Categories;
+}
+
+bool UAssetDefinition_MetasoundMusic::CanImport() const
+{
+	return false;
+}
+
+const FSlateBrush* UAssetDefinition_MetasoundMusic::GetIconBrush(const FAssetData& InAssetData, const FName InClassName) const
+{
+	return HarmonixMetasoundEditor::FSlateStyle::Get().GetBrush("HarmonixMetasoundEditor.MetasoundMusic.Icon");
+}
+
+const FSlateBrush* UAssetDefinition_MetasoundMusic::GetThumbnailBrush(const FAssetData& InAssetData, const FName InClassName) const
+{
+	return HarmonixMetasoundEditor::FSlateStyle::Get().GetBrush("HarmonixMetasoundEditor.MetasoundMusic.Thumbnail");
+}
diff --git a/Engine/Plugins/Runtime/Harmonix/Source/HarmonixMetasoundEditor/Private/AssetDefinition_MusicAssets.h b/Engine/Plugins/Runtime/Harmonix/Source/HarmonixMetasoundEditor/Private/AssetDefinition_MusicAssets.h
new file mode 100644
--- /dev/null
+++ b/Engine/Plugins/Runtime/Harmonix/Source/HarmonixMetasoundEditor/Private/AssetDefinition_MusicAssets.h
@@ -0,0 +1,25 @@
+// Copyright Epic Games, Inc. All Rights Reserved.
+
+#pragma once
+
+
+#include "AssetDefinitionDefault.h"
+
+#include "AssetDefinition_MusicAssets.generated.h"
+
+#define ASSET_MENU_ITEM_GROUPING_CHEKED_IN 0
+
+UCLASS()
+class UAssetDefinition_MetasoundMusic : public UAssetDefinitionDefault
+{
+	GENERATED_BODY()
+
+public:
+	virtual FText GetAssetDisplayName() const override;
+	virtual FLinearColor GetAssetColor() const override;
+	virtual TSoftClassPtr<UObject> GetAssetClass() const override;
+	virtual TConstArrayView<FAssetCategoryPath> GetAssetCategories() const override;
+	virtual bool CanImport() const override;
+	virtual const FSlateBrush* GetIconBrush(const FAssetData& InAssetData, const FName InClassName) const override;
+	virtual const FSlateBrush* GetThumbnailBrush(const FAssetData& InAssetData, const FName InClassName) const override;
+};
diff --git a/Engine/Plugins/Runtime/Harmonix/Source/HarmonixMetasoundEditor/Private/HarmonixMetasoundMusicFactory.cpp b/Engine/Plugins/Runtime/Harmonix/Source/HarmonixMetasoundEditor/Private/HarmonixMetasoundMusicFactory.cpp
new file mode 100644
--- /dev/null
+++ b/Engine/Plugins/Runtime/Harmonix/Source/HarmonixMetasoundEditor/Private/HarmonixMetasoundMusicFactory.cpp
@@ -0,0 +1,20 @@
+// Copyright Epic Games, Inc. All Rights Reserved.
+
+#include "HarmonixMetasoundMusicFactory.h"
+#include "HarmonixMetasound/DataTypes/HarmonixMetasoundMusicAsset.h"
+
+UHarmonixMetasoundMusicFactory::UHarmonixMetasoundMusicFactory()
+{
+	SupportedClass = UHarmonixMetasoundMusicAsset::StaticClass();
+	bCreateNew = true;
+}
+
+FText UHarmonixMetasoundMusicFactory::GetDisplayName() const
+{
+	return NSLOCTEXT("Harmonix", "MetasoundMusicAssetDisplayName", "Harmonix MetaSound Music Asset");
+}
+
+UObject* UHarmonixMetasoundMusicFactory::FactoryCreateNew(UClass* Class, UObject* InParent, FName Name, EObjectFlags Flags, UObject* Context, FFeedbackContext* Warn)
+{
+	return NewObject<UHarmonixMetasoundMusicAsset>(InParent, Class, Name, Flags, Context);
+}
diff --git a/Engine/Plugins/Runtime/Harmonix/Source/HarmonixMetasoundEditor/Private/HarmonixMetasoundMusicFactory.h b/Engine/Plugins/Runtime/Harmonix/Source/HarmonixMetasoundEditor/Private/HarmonixMetasoundMusicFactory.h
new file mode 100644
--- /dev/null
+++ b/Engine/Plugins/Runtime/Harmonix/Source/HarmonixMetasoundEditor/Private/HarmonixMetasoundMusicFactory.h
@@ -0,0 +1,19 @@
+// Copyright Epic Games, Inc. All Rights Reserved.
+
+#pragma once
+
+#include "Factories/Factory.h"
+#include "HarmonixMetasoundMusicFactory.generated.h"
+
+UCLASS()
+class UHarmonixMetasoundMusicFactory : public UFactory
+{
+	GENERATED_BODY()
+public:
+	UHarmonixMetasoundMusicFactory();
+	//~ BEGIN UFactory interface
+	virtual FText GetDisplayName() const override;
+	virtual UObject* FactoryCreateNew(UClass* Class, UObject* InParent, FName Name, EObjectFlags Flags, UObject* Context, FFeedbackContext* Warn) override;
+	//~ END UFactory interface
+};
+
diff --git a/Engine/Plugins/Runtime/Harmonix/Source/HarmonixMetasoundEditor/Private/HarmonixMetasoundSlateStyle.cpp b/Engine/Plugins/Runtime/Harmonix/Source/HarmonixMetasoundEditor/Private/HarmonixMetasoundSlateStyle.cpp
--- a/Engine/Plugins/Runtime/Harmonix/Source/HarmonixMetasoundEditor/Private/HarmonixMetasoundSlateStyle.cpp
+++ b/Engine/Plugins/Runtime/Harmonix/Source/HarmonixMetasoundEditor/Private/HarmonixMetasoundSlateStyle.cpp
@@ -1,6 +1,7 @@
 // Copyright Epic Games, Inc. All Rights Reserved.
 
 #include "HarmonixMetasoundSlateStyle.h"
+#include "Styling/SlateStyleMacros.h"
 #include "Styling/SlateStyleRegistry.h"
 #include "Interfaces/IPluginManager.h"
 
@@ -20,15 +21,19 @@ namespace HarmonixMetasoundEditor
 		const FVector2D Icon22x22(22.0f, 22.0f);
 		const FVector2D Icon18x10(18.0f, 10.0f);
 		const FVector2D Icon18x18(18.0f, 18.0f);
+		const FVector2D Icon20x20(20.0f, 20.0f);
 		TSharedPtr<IPlugin> Plugin = IPluginManager::Get().FindPlugin(TEXT("Harmonix"));
 		check(Plugin);
 		SetContentRoot(Plugin->GetContentDir() / TEXT("Editor/Slate"));
 
-#define IMAGE_BRUSH( RelativePath, ... ) FSlateImageBrush( RootToContentDir( RelativePath, TEXT(".png") ), __VA_ARGS__ )
 		SetCustomPinStyle("MIDIStream", MidiStreamColor, new IMAGE_BRUSH(TEXT("Icons/MidiConnectedPin"), Icon22x22), new IMAGE_BRUSH(TEXT("Icons/MidiDisconnectedPin"), Icon22x22));
 		SetCustomPinStyle("MIDIClock", MidiClockColor, new IMAGE_BRUSH(TEXT("Icons/ClockConnectedPin"), Icon22x22), new IMAGE_BRUSH(TEXT("Icons/ClockDisconnectedPin"), Icon22x22));
 		SetCustomPinStyle("MusicTransport", TransportColor, new IMAGE_BRUSH(TEXT("Icons/TransportConnectedPin"), Icon18x10), new IMAGE_BRUSH(TEXT("Icons/TransportDisconnectedPin"), Icon18x10));
-#undef IMAGE_BRUSH
+
+		Set("HarmonixMetasoundEditor.MetasoundMusic.Icon", new IMAGE_BRUSH_SVG(TEXT("Icons/harmonixmetasoundmusic_icon"), Icon20x20));
+		Set("HarmonixMetasoundEditor.WaveMusic.Icon", new IMAGE_BRUSH_SVG(TEXT("Icons/harmonixwavemusic_icon"), Icon20x20));
+		Set("HarmonixMetasoundEditor.MetasoundMusic.Thumbnail", new IMAGE_BRUSH_SVG(TEXT("Icons/harmonixmetasoundmusic_thumbnail"), Icon20x20));
+		Set("HarmonixMetasoundEditor.WaveMusic.Thumbnail", new IMAGE_BRUSH_SVG(TEXT("Icons/harmonixwavemusic_thumbnail"), Icon20x20));
 
 		FSlateStyleRegistry::RegisterSlateStyle(*this);
 	}
