Add ability to seek even if there is no midi clock input... Seek using the seek target in the transport stream.

diff --git a/Engine/Plugins/Runtime/Harmonix/Source/HarmonixMetasound/Private/HarmonixMetasound/Nodes/TransportWavePlayerControllerNode.cpp b/Engine/Plugins/Runtime/Harmonix/Source/HarmonixMetasound/Private/HarmonixMetasound/Nodes/TransportWavePlayerControllerNode.cpp
--- a/Engine/Plugins/Runtime/Harmonix/Source/HarmonixMetasound/Private/HarmonixMetasound/Nodes/TransportWavePlayerControllerNode.cpp
+++ b/Engine/Plugins/Runtime/Harmonix/Source/HarmonixMetasound/Private/HarmonixMetasound/Nodes/TransportWavePlayerControllerNode.cpp
@@ -10,16 +10,36 @@
 #include "HarmonixMetasound/Common.h"
 #include "HarmonixMetasound/DataTypes/MidiClock.h"
 #include "HarmonixMetasound/DataTypes/MusicTransport.h"
+#include "HarmonixMetasound/Nodes/TransportWavePlayerControllerNode.h"
 
 #define LOCTEXT_NAMESPACE "HarmonixMetaSound"
 
-namespace HarmonixMetasound
+namespace HarmonixMetasound::Nodes::TransportWavePlayerControllerNode
 {
 	using namespace Metasound;
 
-	namespace TransportWavePlayerControllerVertexNames
+	const FNodeClassName& GetClassName()
 	{
-		METASOUND_PARAM(OutputStartTime, "Start Time", "Time into the wave asset to start (seek) the wave asset.")
+		static const FNodeClassName ClassName { HarmonixNodeNamespace, TEXT("TransportWavePlayerController"), TEXT("") };
+		return ClassName;
+	}
+
+	int32 GetCurrentMajorVersion()
+	{
+		return 0;
+	}
+
+	namespace Inputs
+	{
+		DEFINE_METASOUND_PARAM_ALIAS(Transport, CommonPinNames::Inputs::Transport);
+		DEFINE_METASOUND_PARAM_ALIAS(MidiClock, CommonPinNames::Inputs::MidiClock);
+	}
+
+	namespace Outputs
+	{
+		DEFINE_INPUT_METASOUND_PARAM(StartTime, "Start Time", "Time into the wave asset to start (seek) the wave asset.")
+		DEFINE_METASOUND_PARAM_ALIAS(TransportPlay, CommonPinNames::Outputs::TransportPlay)
+		DEFINE_METASOUND_PARAM_ALIAS(TransportStop, CommonPinNames::Outputs::TransportStop)
 	}
 
 	class FTransportWavePlayerControllerOperator : public TExecutableOperator<FTransportWavePlayerControllerOperator>, public FMusicTransportControllable
@@ -27,10 +47,12 @@ namespace HarmonixMetasound
 	public:
 		FTransportWavePlayerControllerOperator(const FOperatorSettings& InSettings,
 											   const FMusicTransportEventStreamReadRef& InTransport,
-											   const FMidiClockReadRef& InMidiClock)
+											   const FMidiClockReadRef& InMidiClock,
+		                                       bool bInHasClockInput)
 			: FMusicTransportControllable(EMusicPlayerTransportState::Prepared)
 			, TransportInPin(InTransport)
 			, MidiClockInPin(InMidiClock)
+			, bHasClockInput(bInHasClockInput)
 			, PlayOutPin(FTriggerWriteRef::CreateNew(InSettings))
 			, StopOutPin(FTriggerWriteRef::CreateNew(InSettings))
 			, StartTimeOutPin(FTimeWriteRef::CreateNew())
@@ -41,9 +63,6 @@ namespace HarmonixMetasound
 
 		static const FVertexInterface& GetVertexInterface()
 		{
-			using namespace CommonPinNames;
-			using namespace TransportWavePlayerControllerVertexNames;
-
 			static const FVertexInterface Interface(
 				FInputVertexInterface(
 					TInputDataVertex<FMusicTransportEventStream>(METASOUND_GET_PARAM_NAME_AND_METADATA(Inputs::Transport)),
@@ -52,7 +71,7 @@ namespace HarmonixMetasound
 				FOutputVertexInterface(
 					TOutputDataVertex<FTrigger>(METASOUND_GET_PARAM_NAME_AND_METADATA(Outputs::TransportPlay)),
 					TOutputDataVertex<FTrigger>(METASOUND_GET_PARAM_NAME_AND_METADATA(Outputs::TransportStop)),
-					TOutputDataVertex<FTime>(METASOUND_GET_PARAM_NAME_AND_METADATA(OutputStartTime))
+					TOutputDataVertex<FTime>(METASOUND_GET_PARAM_NAME_AND_METADATA(Outputs::StartTime))
 				)
 			);
 
@@ -64,8 +83,8 @@ namespace HarmonixMetasound
 			auto InitNodeInfo = []() -> FNodeClassMetadata
 			{
 				FNodeClassMetadata Info;
-				Info.ClassName = { HarmonixNodeNamespace, TEXT("TransportWavePlayerController"), TEXT("") };
-				Info.MajorVersion = 0;
+				Info.ClassName = GetClassName();
+				Info.MajorVersion = GetCurrentMajorVersion();
 				Info.MinorVersion = 1;
 				Info.DisplayName = METASOUND_LOCTEXT("TransportWavePlayerControllerNode_DisplayName", "Music Transport Wave Player Controller");
 				Info.Description = METASOUND_LOCTEXT("TransportWavePlayerControllerNode_Description", "An interface between a music transport and a wave player");
@@ -83,22 +102,16 @@ namespace HarmonixMetasound
 
 		virtual void BindInputs(FInputVertexInterfaceData& InVertexData) override
 		{
-			using namespace CommonPinNames;
-
 			InVertexData.BindReadVertex(METASOUND_GET_PARAM_NAME(Inputs::Transport), TransportInPin);
 			InVertexData.BindReadVertex(METASOUND_GET_PARAM_NAME(Inputs::MidiClock), MidiClockInPin);
-
-			bNeedsTransportInit = true;
+			bHasClockInput = InVertexData.IsVertexBound(METASOUND_GET_PARAM_NAME(Inputs::MidiClock));
 		}
 
 		virtual void BindOutputs(FOutputVertexInterfaceData& InVertexData) override
 		{
-			using namespace CommonPinNames;
-			using namespace TransportWavePlayerControllerVertexNames;
-
 			InVertexData.BindReadVertex(METASOUND_GET_PARAM_NAME(Outputs::TransportPlay), PlayOutPin);
 			InVertexData.BindReadVertex(METASOUND_GET_PARAM_NAME(Outputs::TransportStop), StopOutPin);
-			InVertexData.BindReadVertex(METASOUND_GET_PARAM_NAME(OutputStartTime), StartTimeOutPin);
+			InVertexData.BindReadVertex(METASOUND_GET_PARAM_NAME(Outputs::StartTime), StartTimeOutPin);
 
 			bNeedsTransportInit = true;
 		}
@@ -110,8 +123,9 @@ namespace HarmonixMetasound
 			const FInputVertexInterfaceData& InputData = InParams.InputData;
 			FMusicTransportEventStreamReadRef InTransport = InputData.GetOrCreateDefaultDataReadReference<FMusicTransportEventStream>(METASOUND_GET_PARAM_NAME(Inputs::Transport), InParams.OperatorSettings);
 			FMidiClockReadRef InMidiClock = InputData.GetOrCreateDefaultDataReadReference<FMidiClock>(METASOUND_GET_PARAM_NAME(Inputs::MidiClock), InParams.OperatorSettings);
-			
-			return MakeUnique<FTransportWavePlayerControllerOperator>(InParams.OperatorSettings, InTransport, InMidiClock);
+			bool bHasClockInput = InputData.IsVertexBound(METASOUND_GET_PARAM_NAME(Inputs::MidiClock));
+
+			return MakeUnique<FTransportWavePlayerControllerOperator>(InParams.OperatorSettings, InTransport, InMidiClock, bHasClockInput);
 		}
 
 		void Execute()
@@ -147,10 +161,15 @@ namespace HarmonixMetasound
 					return EMusicPlayerTransportState::Playing;
 
 				case EMusicPlayerTransportState::Seeking:
+				{
+					float NextSeekMs = bHasClockInput ? MidiClockInPin->GetCurrentSongPosMs() :
+						TransportInPin->GetNextSeekDestination().Type == ESeekPointType::Millisecond ? TransportInPin->GetNextSeekDestination().Ms :
+						0.0f;
+					float NextSeekSeconds = NextSeekMs * 0.001f;
 					if (ReceivedSeekWhileStopped())
 					{
 						// Assumes the MidiClock is stopped for the remainder of the block.
-						*StartTimeOutPin = FTime(MidiClockInPin->GetCurrentSongPosMs() * 0.001f);
+						*StartTimeOutPin = FTime(NextSeekSeconds);
 					}
 					else
 					{
@@ -158,13 +177,13 @@ namespace HarmonixMetasound
 						int32 PlayFrameIndex = FMath::Min(StartFrameIndex + 1, EndFrameIndex);
 
 						// Assumes the MidiClock is playing for the remainder of the block.
-						*StartTimeOutPin = FTime(MidiClockInPin->GetCurrentSongPosMs() * 0.001f - (BlockSizeFrames - PlayFrameIndex) / SampleRate);
+						*StartTimeOutPin = FTime(NextSeekSeconds - (BlockSizeFrames - PlayFrameIndex) / SampleRate);
 						PlayOutPin->TriggerFrame(PlayFrameIndex);
 					}
 					// Here we will return that we want to be in the same state we were in before this request to 
 					// seek since we can seek "instantaneously"...
 					return GetTransportState();
-
+				}
 				case EMusicPlayerTransportState::Continuing:
 					// Assumes the StartTimeOutPin won't change for the remainder of the block.
 					PlayOutPin->TriggerFrame(StartFrameIndex);
@@ -219,6 +238,7 @@ namespace HarmonixMetasound
 		// Inputs
 		FMusicTransportEventStreamReadRef TransportInPin;
 		FMidiClockReadRef MidiClockInPin;
+		bool bHasClockInput;
 
 		// Outputs
 		FTriggerWriteRef PlayOutPin;
diff --git a/Engine/Plugins/Runtime/Harmonix/Source/HarmonixMetasound/Public/HarmonixMetasound/Nodes/TransportWavePlayerControllerNode.h b/Engine/Plugins/Runtime/Harmonix/Source/HarmonixMetasound/Public/HarmonixMetasound/Nodes/TransportWavePlayerControllerNode.h
new file mode 100644
--- /dev/null
+++ b/Engine/Plugins/Runtime/Harmonix/Source/HarmonixMetasound/Public/HarmonixMetasound/Nodes/TransportWavePlayerControllerNode.h
@@ -0,0 +1,26 @@
+// Copyright Epic Games, Inc. All Rights Reserved.
+
+#pragma once
+
+#include "MetasoundNodeInterface.h"
+
+#include "HarmonixMetasound/Common.h"
+
+namespace HarmonixMetasound::Nodes::TransportWavePlayerControllerNode
+{
+	const HARMONIXMETASOUND_API Metasound::FNodeClassName& GetClassName();
+	HARMONIXMETASOUND_API int32 GetCurrentMajorVersion();
+
+	namespace Inputs
+	{
+		DECLARE_METASOUND_PARAM_ALIAS(Transport);
+		DECLARE_METASOUND_PARAM_ALIAS(MidiClock);
+	}
+
+	namespace Outputs
+	{
+		DECLARE_METASOUND_PARAM_ALIAS(TransportPlay);
+		DECLARE_METASOUND_PARAM_ALIAS(TransportStop);
+		DECLARE_METASOUND_PARAM_EXTERN(StartTime);
+	}
+}
\ No newline at end of file
