Trivial tweaks to make it easier to build metasound graphs from code.

diff --git a/Engine/Plugins/Runtime/Harmonix/Source/HarmonixMetasound/Private/HarmonixMetasound/Common.cpp b/Engine/Plugins/Runtime/Harmonix/Source/HarmonixMetasound/Private/HarmonixMetasound/Common.cpp
--- a/Engine/Plugins/Runtime/Harmonix/Source/HarmonixMetasound/Private/HarmonixMetasound/Common.cpp
+++ b/Engine/Plugins/Runtime/Harmonix/Source/HarmonixMetasound/Private/HarmonixMetasound/Common.cpp
@@ -83,11 +83,12 @@ namespace HarmonixMetasound
 			DEFINE_OUTPUT_METASOUND_PARAM(MidiVelocity,       "Velocity",                    "MIDI velocity (0 - 127).")
 			DEFINE_OUTPUT_METASOUND_PARAM(NormalizedVelocity, "Normalized Velocity",         "Velocity (0.0 - 1.0).")
 			DEFINE_OUTPUT_METASOUND_PARAM(TransportPrepare,   "Prepare",                     "Causes receiving nodes to prepare to play (pre-buffer, cache, etc.)")
-			DEFINE_OUTPUT_METASOUND_PARAM(TransportPlay,      "Play",                        "Causes a transport play request.")
-			DEFINE_OUTPUT_METASOUND_PARAM(TransportPause,     "Pause",                       "Causes a transport pause request.")
-			DEFINE_OUTPUT_METASOUND_PARAM(TransportContinue,  "Continue",                    "Causes a transport continue request.")
-			DEFINE_OUTPUT_METASOUND_PARAM(TransportStop,      "Stop",                        "Causes a transport stop request.")
-			DEFINE_OUTPUT_METASOUND_PARAM(TransportKill,      "Kill",                        "Causes a transport kill request.")
+			DEFINE_OUTPUT_METASOUND_PARAM(TransportPlay,      "Play",                        "Transport play requested.")
+			DEFINE_OUTPUT_METASOUND_PARAM(TransportPause,     "Pause",                       "Transport pause requested.")
+			DEFINE_OUTPUT_METASOUND_PARAM(TransportContinue,  "Continue",                    "Transport continue requested.")
+			DEFINE_OUTPUT_METASOUND_PARAM(TransportStop,      "Stop",                        "Transport stop requested.")
+			DEFINE_OUTPUT_METASOUND_PARAM(TransportKill,      "Kill",                        "Transport kill requested.")
+			DEFINE_OUTPUT_METASOUND_PARAM(TransportSeek,      "Seek",                        "Transport seek requested.")
 			DEFINE_OUTPUT_METASOUND_PARAM(MusicTimestamp,     "Music Timestamp",             "A single structure representing a musical time (bar & beat).")
 			DEFINE_OUTPUT_METASOUND_PARAM(SeekTarget,         "Seek Target",                 "A well formed musical seek target.")
 			DEFINE_OUTPUT_METASOUND_PARAM(AudioMono,          "Audio Mono",					 "The mono audio.")
diff --git a/Engine/Plugins/Runtime/Harmonix/Source/HarmonixMetasound/Private/HarmonixMetasound/Nodes/TransportToTriggerNode.cpp b/Engine/Plugins/Runtime/Harmonix/Source/HarmonixMetasound/Private/HarmonixMetasound/Nodes/TransportToTriggerNode.cpp
--- a/Engine/Plugins/Runtime/Harmonix/Source/HarmonixMetasound/Private/HarmonixMetasound/Nodes/TransportToTriggerNode.cpp
+++ b/Engine/Plugins/Runtime/Harmonix/Source/HarmonixMetasound/Private/HarmonixMetasound/Nodes/TransportToTriggerNode.cpp
@@ -9,13 +9,41 @@
 #include "MetasoundVertex.h"
 #include "HarmonixMetasound/Common.h"
 #include "HarmonixMetasound/DataTypes/MusicTransport.h"
+#include "HarmonixMetasound/Nodes/TransportToTriggerNode.h"
 
 #define LOCTEXT_NAMESPACE "HarmonixMetaSound"
 
-namespace HarmonixMetasound
+namespace HarmonixMetasound::Nodes::TransportToTriggerNode
 {
 	using namespace Metasound;
 	
+	const FNodeClassName& GetClassName()
+	{
+		static const FNodeClassName ClassName { HarmonixNodeNamespace, TEXT("TransportToTrigger"), TEXT("") };
+		return ClassName;
+	}
+
+	int32 GetCurrentMajorVersion()
+	{
+		return 0;
+	}
+
+	namespace Inputs
+	{
+		DEFINE_METASOUND_PARAM_ALIAS(Transport, CommonPinNames::Inputs::Transport);
+	}
+
+	namespace Outputs
+	{
+		DEFINE_METASOUND_PARAM_ALIAS(TransportPrepare, CommonPinNames::Outputs::TransportPrepare);
+		DEFINE_METASOUND_PARAM_ALIAS(TransportPlay, CommonPinNames::Outputs::TransportPlay);
+		DEFINE_METASOUND_PARAM_ALIAS(TransportPause, CommonPinNames::Outputs::TransportPause);
+		DEFINE_METASOUND_PARAM_ALIAS(TransportContinue, CommonPinNames::Outputs::TransportContinue);
+		DEFINE_METASOUND_PARAM_ALIAS(TransportStop, CommonPinNames::Outputs::TransportStop);
+		DEFINE_METASOUND_PARAM_ALIAS(TransportKill, CommonPinNames::Outputs::TransportKill);
+		DEFINE_METASOUND_PARAM_ALIAS(TransportSeek, CommonPinNames::Outputs::TransportSeek);
+	}
+
 	class FTransportToTriggerOperator : public TExecutableOperator<FTransportToTriggerOperator>
 	{
 	public:
@@ -28,6 +56,7 @@ namespace HarmonixMetasound
 			, ContinueOutPin(FTriggerWriteRef::CreateNew(InParams.OperatorSettings))
 			, StopOutPin(FTriggerWriteRef::CreateNew(InParams.OperatorSettings))
 			, KillOutPin(FTriggerWriteRef::CreateNew(InParams.OperatorSettings))
+			, SeekOutPin(FTriggerWriteRef::CreateNew(InParams.OperatorSettings))
 		{
 			Reset(InParams);
 		}
@@ -46,7 +75,8 @@ namespace HarmonixMetasound
 					TOutputDataVertex<FTrigger>(METASOUND_GET_PARAM_NAME_AND_METADATA(Outputs::TransportPause)),
 					TOutputDataVertex<FTrigger>(METASOUND_GET_PARAM_NAME_AND_METADATA(Outputs::TransportContinue)),
 					TOutputDataVertex<FTrigger>(METASOUND_GET_PARAM_NAME_AND_METADATA(Outputs::TransportStop)),
-					TOutputDataVertex<FTrigger>(METASOUND_GET_PARAM_NAME_AND_METADATA(Outputs::TransportKill))
+					TOutputDataVertex<FTrigger>(METASOUND_GET_PARAM_NAME_AND_METADATA(Outputs::TransportKill)),
+					TOutputDataVertex<FTrigger>(METASOUND_GET_PARAM_NAME_AND_METADATA(Outputs::TransportSeek))
 				)
 			);
 
@@ -58,8 +88,8 @@ namespace HarmonixMetasound
 			auto InitNodeInfo = []() -> FNodeClassMetadata
 			{
 				FNodeClassMetadata Info;
-				Info.ClassName = { HarmonixNodeNamespace, TEXT("TransportToTrigger"), TEXT("") };
-				Info.MajorVersion = 0;
+				Info.ClassName = GetClassName();
+				Info.MajorVersion = GetCurrentMajorVersion();
 				Info.MinorVersion = 1;
 				Info.DisplayName = METASOUND_LOCTEXT("TransportToTriggerNode_DisplayName", "Music Transport to Trigger");
 				Info.Description = METASOUND_LOCTEXT("TransportToTriggerNode_Description", "Breaks out music transport requests into individual triggers");
@@ -92,6 +122,7 @@ namespace HarmonixMetasound
 			InVertexData.BindReadVertex(METASOUND_GET_PARAM_NAME(Outputs::TransportContinue), ContinueOutPin);
 			InVertexData.BindReadVertex(METASOUND_GET_PARAM_NAME(Outputs::TransportStop), StopOutPin);
 			InVertexData.BindReadVertex(METASOUND_GET_PARAM_NAME(Outputs::TransportKill), KillOutPin);
+			InVertexData.BindReadVertex(METASOUND_GET_PARAM_NAME(Outputs::TransportSeek), SeekOutPin);
 		}
 
 		static TUniquePtr<IOperator> CreateOperator(const FBuildOperatorParams& InParams, FBuildResults& OutResults)
@@ -112,6 +143,7 @@ namespace HarmonixMetasound
 			ContinueOutPin->Reset();
 			StopOutPin->Reset();
 			KillOutPin->Reset();
+			SeekOutPin->Reset();
 		}
 
 		void Execute()
@@ -123,6 +155,7 @@ namespace HarmonixMetasound
 			ContinueOutPin->AdvanceBlock();
 			StopOutPin->AdvanceBlock();
 			KillOutPin->AdvanceBlock();
+			SeekOutPin->AdvanceBlock();
 
 			// if a transport event comes through, pass it to the appropriate output pin
 			const FMusicTransportEventStream::FEventList& TransportEventsThisBlock = TransportInPin->GetTransportEventsInBlock();
@@ -149,6 +182,9 @@ namespace HarmonixMetasound
 				case EMusicPlayerTransportRequest::Kill:
 					KillOutPin->TriggerFrame(Event.SampleIndex);
 					break;
+				case EMusicPlayerTransportRequest::Seek:
+					SeekOutPin->TriggerFrame(Event.SampleIndex);
+					break;
 				default:
 					break;
 				}
@@ -167,6 +203,7 @@ namespace HarmonixMetasound
 		FTriggerWriteRef ContinueOutPin;
 		FTriggerWriteRef StopOutPin;
 		FTriggerWriteRef KillOutPin;
+		FTriggerWriteRef SeekOutPin;
 	};
 
 	using FTransportToTriggerNode = Metasound::TNodeFacade<FTransportToTriggerOperator>;
diff --git a/Engine/Plugins/Runtime/Harmonix/Source/HarmonixMetasound/Public/HarmonixMetasound/Common.h b/Engine/Plugins/Runtime/Harmonix/Source/HarmonixMetasound/Public/HarmonixMetasound/Common.h
--- a/Engine/Plugins/Runtime/Harmonix/Source/HarmonixMetasound/Public/HarmonixMetasound/Common.h
+++ b/Engine/Plugins/Runtime/Harmonix/Source/HarmonixMetasound/Public/HarmonixMetasound/Common.h
@@ -126,6 +126,7 @@ namespace HarmonixMetasound
 			DECLARE_METASOUND_PARAM_EXTERN(TransportContinue)
 			DECLARE_METASOUND_PARAM_EXTERN(TransportStop)
 			DECLARE_METASOUND_PARAM_EXTERN(TransportKill)
+			DECLARE_METASOUND_PARAM_EXTERN(TransportSeek)
 			DECLARE_METASOUND_PARAM_EXTERN(MusicTimestamp)
 			DECLARE_METASOUND_PARAM_EXTERN(SeekTarget)
 			DECLARE_METASOUND_PARAM_EXTERN(AudioMono)
diff --git a/Engine/Plugins/Runtime/Harmonix/Source/HarmonixMetasound/Public/HarmonixMetasound/Nodes/TransportToTriggerNode.h b/Engine/Plugins/Runtime/Harmonix/Source/HarmonixMetasound/Public/HarmonixMetasound/Nodes/TransportToTriggerNode.h
new file mode 100644
--- /dev/null
+++ b/Engine/Plugins/Runtime/Harmonix/Source/HarmonixMetasound/Public/HarmonixMetasound/Nodes/TransportToTriggerNode.h
@@ -0,0 +1,29 @@
+// Copyright Epic Games, Inc. All Rights Reserved.
+
+#pragma once
+
+#include "MetasoundNodeInterface.h"
+
+#include "HarmonixMetasound/Common.h"
+
+namespace HarmonixMetasound::Nodes::TransportToTriggerNode
+{
+	const HARMONIXMETASOUND_API Metasound::FNodeClassName& GetClassName();
+	HARMONIXMETASOUND_API int32 GetCurrentMajorVersion();
+
+	namespace Inputs
+	{
+		DECLARE_METASOUND_PARAM_ALIAS(Transport);
+	}
+
+	namespace Outputs
+	{
+		DECLARE_METASOUND_PARAM_ALIAS(TransportPrepare);
+		DECLARE_METASOUND_PARAM_ALIAS(TransportPlay);
+		DECLARE_METASOUND_PARAM_ALIAS(TransportPause);
+		DECLARE_METASOUND_PARAM_ALIAS(TransportContinue);
+		DECLARE_METASOUND_PARAM_ALIAS(TransportStop);
+		DECLARE_METASOUND_PARAM_ALIAS(TransportKill);
+		DECLARE_METASOUND_PARAM_ALIAS(TransportSeek);
+	}
+}
\ No newline at end of file
