FMusicClockDriverBase::GetWallClockTime should use GetRealTimeSeconds which ignores any pausing or time dilation that might be applied to the world's game time.

#rnx

diff --git a/Engine/Plugins/Runtime/Harmonix/Source/HarmonixMetasound/Private/HarmonixMetasound/Components/MusicClockDriverBase.cpp b/Engine/Plugins/Runtime/Harmonix/Source/HarmonixMetasound/Private/HarmonixMetasound/Components/MusicClockDriverBase.cpp
--- a/Engine/Plugins/Runtime/Harmonix/Source/HarmonixMetasound/Private/HarmonixMetasound/Components/MusicClockDriverBase.cpp
+++ b/Engine/Plugins/Runtime/Harmonix/Source/HarmonixMetasound/Private/HarmonixMetasound/Components/MusicClockDriverBase.cpp
@@ -28,13 +28,13 @@ double FMusicClockDriverBase::GetWallClockTime() const
 	{
 		if (Obj->GetWorld())
 		{
-			return Obj->GetWorld()->GetTimeSeconds();
+			return Obj->GetWorld()->GetRealTimeSeconds();
 		}
 	}
 
 	if (GWorld)
 	{
-		return GWorld->GetTimeSeconds();
+		return GWorld->GetRealTimeSeconds();
 	}
 	
 	return 0.0;
