Clamp Tempo, Speed, & Time Signature throughout the codebase

#rnx

diff --git a/Engine/Plugins/Runtime/Harmonix/Source/HarmonixDsp/Private/HarmonixDsp/FusionSampler/FusionSampler.cpp b/Engine/Plugins/Runtime/Harmonix/Source/HarmonixDsp/Private/HarmonixDsp/FusionSampler/FusionSampler.cpp
--- a/Engine/Plugins/Runtime/Harmonix/Source/HarmonixDsp/Private/HarmonixDsp/FusionSampler/FusionSampler.cpp
+++ b/Engine/Plugins/Runtime/Harmonix/Source/HarmonixDsp/Private/HarmonixDsp/FusionSampler/FusionSampler.cpp
@@ -16,6 +16,12 @@ namespace Defaults
 	
 	static const float kMinPitchBend = -200.0f;
 	static const float kMaxPitchBend = 200.0f;
+	
+	constexpr float kMinRenderSpeed = 0.001f;
+	constexpr float kMaxRenderSpeed = 10.0f;
+
+	constexpr float kMinPitchShift = 0.0001;
+	constexpr float kMaxPitchShift = 10.0f;
 
 	static const int32 kMaxNumVoices = 100;
 }
@@ -230,11 +236,15 @@ void FFusionSampler::NoteOffWithFrameOffset(FMidiVoiceId InVoiceId, int8 InMidiN
 void FFusionSampler::SetSpeed(float InSpeed, bool InMaintainPitch)
 {
 	FScopeLock Lock(&GetBusLock());
-	if (InSpeed <= .0f)
+	if (InSpeed <= Defaults::kMinRenderSpeed)
+	{
+		UE_LOG(LogFusionSampler, Warning, TEXT("Speed value cannot be less than %.8f. Clamping"), Defaults::kMinRenderSpeed);
+		Speed = Defaults::kMinRenderSpeed;	
+	}
+	else if (InSpeed > Defaults::kMaxRenderSpeed)
 	{
-		UE_LOG(LogFusionSampler, Warning, TEXT("Speed value cannot be zero or negative. Forced to 0.0001"));
-		// lower boundary
-		Speed = .0001f;	
+		UE_LOG(LogFusionSampler, Warning, TEXT("Speed value cannot greater than %.2f. Clamping"), Defaults::kMaxRenderSpeed);
+		Speed = Defaults::kMaxRenderSpeed;
 	}
 	else
 	{
diff --git a/Engine/Plugins/Runtime/Harmonix/Source/HarmonixDsp/Private/HarmonixDsp/FusionSampler/FusionVoice.cpp b/Engine/Plugins/Runtime/Harmonix/Source/HarmonixDsp/Private/HarmonixDsp/FusionSampler/FusionVoice.cpp
--- a/Engine/Plugins/Runtime/Harmonix/Source/HarmonixDsp/Private/HarmonixDsp/FusionSampler/FusionVoice.cpp
+++ b/Engine/Plugins/Runtime/Harmonix/Source/HarmonixDsp/Private/HarmonixDsp/FusionSampler/FusionVoice.cpp
@@ -83,6 +83,11 @@ bool FFusionVoice::AssignIDs(FFusionSampler* InSampler, const FKeyzoneSettings*
 
 	KeyZone = InKeyZone;
 
+	if (KeyZone->TimeStretchConfig.bMaintainTime)
+	{
+		check(KeyZone->TimeStretchConfig.OriginalTempo > 0.0f);
+	}
+
 	RelinquishHandler = Handler;
 
 	MySampler = InSampler;
diff --git a/Engine/Plugins/Runtime/Harmonix/Source/HarmonixMetasound/Private/HarmonixMetasound/Nodes/MetronomeNode.cpp b/Engine/Plugins/Runtime/Harmonix/Source/HarmonixMetasound/Private/HarmonixMetasound/Nodes/MetronomeNode.cpp
--- a/Engine/Plugins/Runtime/Harmonix/Source/HarmonixMetasound/Private/HarmonixMetasound/Nodes/MetronomeNode.cpp
+++ b/Engine/Plugins/Runtime/Harmonix/Source/HarmonixMetasound/Private/HarmonixMetasound/Nodes/MetronomeNode.cpp
@@ -164,10 +164,10 @@ namespace HarmonixMetasound::Nodes::MetronomeNode
 		, MonotonicallyIncreasingClock(MakeShared<FMidiClock, ESPMode::NotThreadSafe>(InParams.OperatorSettings))
 		, BlockSize(InParams.OperatorSettings.GetNumFramesPerBlock())
 		, SampleRate(InParams.OperatorSettings.GetSampleRate())
-		, CurrentTempo(*TempoInPin)
-		, CurrentTimeSigNum(FMath::Clamp(*TimeSigNumInPin, 1, 64))
-		, CurrentTimeSigDenom(FMath::Clamp(*TimeSigDenomInPin, 1, 64))
 	{
+		CurrentTempo = GetInputTempoClamped();
+		CurrentTimeSigNum = GetInputTimeSigNumeratorClamped();
+		CurrentTimeSigDenom = GetInputTimeSigDenominatorClamped();
 		Reset(InParams);
 		Init();
 	}
@@ -320,13 +320,27 @@ namespace HarmonixMetasound::Nodes::MetronomeNode
 		MarkClocksAsExecuted();
 	}
 
+	float FMetronomeOperatorBase::GetInputTempoClamped() const
+	{
+		return FMath::Clamp(*TempoInPin, kMinTempoBpm, kMaxTempoBpm);
+	}
+
+	int32 FMetronomeOperatorBase::GetInputTimeSigDenominatorClamped() const
+	{
+		return FMath::Clamp(*TimeSigNumInPin, kMinTimeSigNumerator, kMaxTimeSigNumerator);;
+	}
+
+	int32 FMetronomeOperatorBase::GetInputTimeSigNumeratorClamped() const
+	{
+		return FMath::Clamp(*TimeSigDenomInPin, kMinTimeSigDenominator, kMaxTimeSigDenominator);
+	}
+	
 	void FMetronomeOperatorBase::BuildSongMaps(bool ResetToStart)
 	{
 		// make sure we have valid values
-		CurrentTempo = FMath::Max(1.0f, *TempoInPin);
-		CurrentTimeSigNum = FMath::Max(1, *TimeSigNumInPin);
-		CurrentTimeSigDenom = FMath::Max(1, *TimeSigDenomInPin);
-	
+		CurrentTempo = GetInputTempoClamped();
+		CurrentTimeSigNum = GetInputTimeSigNumeratorClamped();
+		CurrentTimeSigDenom = GetInputTimeSigDenominatorClamped();
 		SongMaps = MakeShared<FSongMaps>(CurrentTempo, CurrentTimeSigNum, CurrentTimeSigDenom);
 		
 		if (LoopInPin)
@@ -355,9 +369,10 @@ namespace HarmonixMetasound::Nodes::MetronomeNode
 	{
 		FMidiClock& DrivingMidiClock = GetDrivingMidiClock();
 		bool HasMidiChanges = false;
-		if (*TempoInPin > 0 && !FMath::IsNearlyEqual(CurrentTempo, *TempoInPin))
+		float NewTempo = GetInputTempoClamped();
+		if (!FMath::IsNearlyEqual(CurrentTempo, NewTempo))
 		{
-			AddTempoChangeForMidi(*TempoInPin);
+			AddTempoChangeForMidi(NewTempo);
 			HasMidiChanges = true;
 		};
 
diff --git a/Engine/Plugins/Runtime/Harmonix/Source/HarmonixMetasound/Private/HarmonixMetasound/Nodes/MidiPlayerNode.cpp b/Engine/Plugins/Runtime/Harmonix/Source/HarmonixMetasound/Private/HarmonixMetasound/Nodes/MidiPlayerNode.cpp
--- a/Engine/Plugins/Runtime/Harmonix/Source/HarmonixMetasound/Private/HarmonixMetasound/Nodes/MidiPlayerNode.cpp
+++ b/Engine/Plugins/Runtime/Harmonix/Source/HarmonixMetasound/Private/HarmonixMetasound/Nodes/MidiPlayerNode.cpp
@@ -1,5 +1,7 @@
 // Copyright Epic Games, Inc. All Rights Reserved.
 
+#include "HarmonixMetasound/Nodes/MidiPlayerNode.h"
+
 #include "MetasoundExecutableOperator.h"
 #include "MetasoundFacade.h"
 #include "MetasoundNodeInterface.h"
@@ -119,6 +121,11 @@ namespace HarmonixMetasound::Nodes::MidiPlayerNode
 
 		void RenderMidiForClockEvents();
 		void SendAllNotesOff(int32 BlockFrameIndex, int32 Tick);
+
+		FORCEINLINE float GetInputClockSpeedClamped() const
+		{
+			return FMath::Clamp(*SpeedMultInPin, kMinClockSpeed, kMaxClockSpeed);
+		}
 	};
 
 	class FExternallyClockedMidiPlayerOperator : public FMidiPlayerOperator
@@ -404,7 +411,7 @@ namespace HarmonixMetasound::Nodes::MidiPlayerNode
 			MidiClockOut->SetTransportState(StartFrameIndex, CurrentState);
 			if (CurrentState == EMusicPlayerTransportState::Playing || CurrentState == EMusicPlayerTransportState::Continuing)
 			{
-				MidiClockOut->SetSpeed(StartFrameIndex, *SpeedMultInPin);
+				MidiClockOut->SetSpeed(StartFrameIndex, GetInputClockSpeedClamped());
 				MidiClockOut->Advance(*MidiClockIn, StartFrameIndex, NumFrames);
 			}
 		};
@@ -458,7 +465,7 @@ namespace HarmonixMetasound::Nodes::MidiPlayerNode
 	{
 		FMidiPlayerOperator::Execute();
 
-		MidiClockOut->SetSpeed(0, *SpeedMultInPin);
+		MidiClockOut->SetSpeed(0, GetInputClockSpeedClamped());
 
 		TransportSpanPostProcessor MidiClockTransportHandler = [&](int32 StartFrameIndex, int32 EndFrameIndex, EMusicPlayerTransportState CurrentState)
 		{
diff --git a/Engine/Plugins/Runtime/Harmonix/Source/HarmonixMetasound/Public/HarmonixMetasound/Nodes/MetronomeNode.h b/Engine/Plugins/Runtime/Harmonix/Source/HarmonixMetasound/Public/HarmonixMetasound/Nodes/MetronomeNode.h
--- a/Engine/Plugins/Runtime/Harmonix/Source/HarmonixMetasound/Public/HarmonixMetasound/Nodes/MetronomeNode.h
+++ b/Engine/Plugins/Runtime/Harmonix/Source/HarmonixMetasound/Public/HarmonixMetasound/Nodes/MetronomeNode.h
@@ -25,6 +25,15 @@ namespace HarmonixMetasound::Nodes::MetronomeNode
 	HARMONIXMETASOUND_API const Metasound::FNodeClassName& GetClassName();
 	HARMONIXMETASOUND_API int32 GetCurrentMajorVersion();
 
+	constexpr float kMinTempoBpm = 1.0f;
+	constexpr float kMaxTempoBpm = 999.0f;
+
+	constexpr int32 kMinTimeSigNumerator = 1;
+	constexpr int32 kMaxTimeSigNumerator = 64;
+
+	constexpr int32 kMinTimeSigDenominator = 1;
+	constexpr int32 kMaxTimeSigDenominator = 64;
+
 	namespace Inputs
 	{
 		DECLARE_METASOUND_PARAM_ALIAS(Transport);
@@ -83,14 +92,17 @@ namespace HarmonixMetasound::Nodes::MetronomeNode
 		TSharedPtr<FSongMaps> SongMaps;
 		FSampleCount BlockSize;
 		float        SampleRate;
-		float        CurrentTempo;
-		int32        CurrentTimeSigNum;
-		int32        CurrentTimeSigDenom;
+		float        CurrentTempo = 0.0f;
+		int32        CurrentTimeSigNum = 0;
+		int32        CurrentTimeSigDenom = 0;
 		int32		 LastProcessedClockTick = -1;
 		int32		 NextClockTickToProcess = 0;
 		bool		 bClocksArePreparedForExecute = true;
 
 		UE_API void BuildSongMaps(bool ResetToStart = true);
+		UE_API float GetInputTempoClamped() const;
+		UE_API int32 GetInputTimeSigNumeratorClamped() const;
+		UE_API int32 GetInputTimeSigDenominatorClamped() const;
 		UE_API void UpdateMidi();
 		UE_API void AddTempoChangeForMidi(float TempoBPM);
 		UE_API virtual void HandleTimeSigChangeForMidi(int32 TimeSigNum, int32 TimeSigDenom);
diff --git a/Engine/Plugins/Runtime/Harmonix/Source/HarmonixMetasound/Public/HarmonixMetasound/Nodes/MidiPlayerNode.h b/Engine/Plugins/Runtime/Harmonix/Source/HarmonixMetasound/Public/HarmonixMetasound/Nodes/MidiPlayerNode.h
--- a/Engine/Plugins/Runtime/Harmonix/Source/HarmonixMetasound/Public/HarmonixMetasound/Nodes/MidiPlayerNode.h
+++ b/Engine/Plugins/Runtime/Harmonix/Source/HarmonixMetasound/Public/HarmonixMetasound/Nodes/MidiPlayerNode.h
@@ -15,6 +15,9 @@ namespace HarmonixMetasound::Nodes::MidiPlayerNode
 	const HARMONIXMETASOUND_API FNodeClassName& GetClassName();
 	HARMONIXMETASOUND_API int32 GetCurrentMajorVersion();
 
+	static constexpr float kMinClockSpeed = 0.01f;
+	static constexpr float kMaxClockSpeed = 10.0f;
+
 	namespace Inputs
 	{
 		DECLARE_METASOUND_PARAM_ALIAS(MidiFileAsset);
