Add more warnings to FusionVoicePool so we don't creating fusion voices with bad data

#rnx

diff --git a/Engine/Plugins/Runtime/Harmonix/Source/HarmonixDsp/Private/HarmonixDsp/FusionSampler/FusionVoicePool.cpp b/Engine/Plugins/Runtime/Harmonix/Source/HarmonixDsp/Private/HarmonixDsp/FusionSampler/FusionVoicePool.cpp
--- a/Engine/Plugins/Runtime/Harmonix/Source/HarmonixDsp/Private/HarmonixDsp/FusionSampler/FusionVoicePool.cpp
+++ b/Engine/Plugins/Runtime/Harmonix/Source/HarmonixDsp/Private/HarmonixDsp/FusionSampler/FusionVoicePool.cpp
@@ -294,16 +294,29 @@ FFusionVoice* FFusionVoicePool::GetFreeVoice(
 
 	if (!InKeyzone->SoundWaveProxy)
 	{
-		UE_LOG(LogFusionVoicePool, Warning, TEXT("Asked to allocate a InSampler voice for a sample resource that failed to load -> %s\n"), *InKeyzone->SoundWaveProxy->GetFName().ToString());
+		UE_LOG(LogFusionVoicePool, Warning, TEXT("Asked to allocate a InSampler voice for an invalid sample"));
 		return nullptr;
 	}
 
 	// if maintain time is true and we didn't get a pitch shitfer factory, then we can't do anything!
-	if (InKeyzone->TimeStretchConfig.bMaintainTime && !PitchShifterFactory)
+	if (InKeyzone->TimeStretchConfig.bMaintainTime)
 	{
-		FName FactoryName = InKeyzone->TimeStretchConfig.PitchShifter.Name;
-		ensureMsgf(false, TEXT("Attempting to play a fusion InKeyzone that is set to \"maintain time\", but no shifter factory was available for assigned Pitch Shifter: %s. Check project configuration! (%s)"), *FactoryName.ToString(), *InKeyzone->SoundWaveProxy->GetFName().ToString());
-		return nullptr;
+		if (!PitchShifterFactory)
+		{
+			FName FactoryName = InKeyzone->TimeStretchConfig.PitchShifter.Name;
+			UE_LOG(LogFusionVoicePool, Warning,
+				TEXT("Attempting to play a fusion InKeyzone that is set to \"maintain time\", but no shifter factory was available for assigned Pitch Shifter: %s. Check project configuration! (%s)"),
+				*FactoryName.ToString(), *InKeyzone->SoundWaveProxy->GetFName().ToString());
+			return nullptr;
+		}
+
+		if (InKeyzone->TimeStretchConfig.OriginalTempo <= 0.0f)
+		{
+			UE_LOG(LogFusionVoicePool, Warning,
+				TEXT("Attempting to play a fusion InKeyzone that is set to \"maintain time\", but the original tempo for the keyzone is invalid (OriginalTempo = %.2f)! Unable to perform pitch shifting or time stretching."),
+				InKeyzone->TimeStretchConfig.OriginalTempo);
+			return nullptr;
+		}
 	}
 
 	FFusionVoice* Voice = nullptr;
