Implement "Duty Cycle" on Pulse Generator Node

#tests Harmonix Metasound Tests
[RN] Add "Duty Cycle" to Harmonix "MIDI Pulse Generator" Node

diff --git a/Engine/Plugins/Runtime/Harmonix/Source/HarmonixMetasound/Private/HarmonixMetasound/DataTypes/MusicTimeInterval.cpp b/Engine/Plugins/Runtime/Harmonix/Source/HarmonixMetasound/Private/HarmonixMetasound/DataTypes/MusicTimeInterval.cpp
--- a/Engine/Plugins/Runtime/Harmonix/Source/HarmonixMetasound/Private/HarmonixMetasound/DataTypes/MusicTimeInterval.cpp
+++ b/Engine/Plugins/Runtime/Harmonix/Source/HarmonixMetasound/Private/HarmonixMetasound/DataTypes/MusicTimeInterval.cpp
@@ -5,9 +5,22 @@
 
 namespace Harmonix
 {
-	float GetIntervalInBeats(const FMusicTimeInterval& Interval, const FTimeSignature& TimeSignature)
+	float FMusicTimeInterval::GetIntervalBeats(const FTimeSignature& TimeSignature) const
 	{
-		return Harmonix::Midi::Constants::SubdivisionToBeats(Interval.Interval, TimeSignature) * Interval.IntervalMultiplier;
+		return Midi::Constants::SubdivisionToBeats(Interval, TimeSignature) * IntervalMultiplier; 
+	}
+
+	float FMusicTimeInterval::GetOffsetBeats(const FTimeSignature& TimeSignature) const
+	{
+		const float IntervalBeats = GetIntervalBeats(TimeSignature);
+		
+		if (FMath::IsNearlyZero(IntervalBeats))
+		{
+			return 0.0f;
+		}
+		
+		const float OffsetBeats = Midi::Constants::SubdivisionToBeats(Offset, TimeSignature) * OffsetMultiplier;
+		return FMath::Fmod(OffsetBeats, IntervalBeats);
 	}
 
 	void IncrementTimestampByBeats(FMusicTimestamp& Timestamp, float Beats, const FTimeSignature& TimeSignature)
@@ -30,7 +43,7 @@ namespace Harmonix
 		const FTimeSignature& TimeSignature)
 	{
 		// Get the pulse interval in beats
-		const float IntervalBeats = GetIntervalInBeats(Interval, TimeSignature);
+		const float IntervalBeats = Interval.GetIntervalBeats(TimeSignature);
 
 		if (FMath::IsNearlyZero(IntervalBeats))
 		{
@@ -40,26 +53,13 @@ namespace Harmonix
 		IncrementTimestampByBeats(Timestamp, IntervalBeats, TimeSignature);
 	}
 
-	float GetOffsetInBeats(const FMusicTimeInterval& Interval, const FTimeSignature& TimeSignature)
-	{
-		const float IntervalBeats = GetIntervalInBeats(Interval, TimeSignature);
-		
-		if (FMath::IsNearlyZero(IntervalBeats))
-		{
-			return 0.0f;
-		}
-		
-		const float OffsetBeats = Harmonix::Midi::Constants::SubdivisionToBeats(Interval.Offset, TimeSignature) * Interval.OffsetMultiplier;
-		return FMath::Fmod(OffsetBeats, IntervalBeats);
-	}
-
 	void IncrementTimestampByOffset(
 		FMusicTimestamp& Timestamp,
 		const FMusicTimeInterval& Interval,
 		const FTimeSignature& TimeSignature)
 	{
 		// Get the offset in beats
-		const float OffsetBeats = GetOffsetInBeats(Interval, TimeSignature);
+		const float OffsetBeats = Interval.GetOffsetBeats(TimeSignature);
 
 		if (FMath::IsNearlyZero(OffsetBeats))
 		{
diff --git a/Engine/Plugins/Runtime/Harmonix/Source/HarmonixMetasound/Private/HarmonixMetasound/MidiOps/PulseGenerator.cpp b/Engine/Plugins/Runtime/Harmonix/Source/HarmonixMetasound/Private/HarmonixMetasound/MidiOps/PulseGenerator.cpp
--- a/Engine/Plugins/Runtime/Harmonix/Source/HarmonixMetasound/Private/HarmonixMetasound/MidiOps/PulseGenerator.cpp
+++ b/Engine/Plugins/Runtime/Harmonix/Source/HarmonixMetasound/Private/HarmonixMetasound/MidiOps/PulseGenerator.cpp
@@ -17,30 +17,39 @@ namespace Harmonix::Midi::Ops
 		Interval.IntervalMultiplier = FMath::Max(Interval.IntervalMultiplier, static_cast<uint16>(1));
 	}
 
+	void FPulseGenerator::SetDutyCycle(float InDutyCycle)
+	{
+		DutyCycle = FMath::Clamp(InDutyCycle, 0.0f, 1.0f);
+	}
+
 	void FPulseGenerator::Reset()
 	{
-		NextPulseTimestamp = { -1, -1 };
+		NextPulses.Reset();
 	}
 
-	void FPulseGenerator::Process(const HarmonixMetasound::FMidiClock& MidiClock, const TFunctionRef<void(const FPulseInfo&)>& OnPulse)
+	void FPulseGenerator::Process(
+		const HarmonixMetasound::FMidiClock& MidiClock,
+		const TFunctionRef<void(const FPulseInfo&)>& TriggerPulseOn,
+		const TFunctionRef<void(const FPulseInfo&)>& TriggerPulseOff)
 	{
 		bool bIntervalIsValid = Interval.Interval != EMidiClockSubdivisionQuantization::None;
 		bool bTimeSigIsValid = CurrentTimeSignature.Numerator > 0 && CurrentTimeSignature.Denominator > 0;
 		
 		// ensure the pulse generator is lined up with the current clock phase
-		if (NextPulseTimestamp == FMusicTimestamp(-1, -1))
+		if (NextPulses.IsEmpty())
 		{
 			CurrentTimeSignature = *MidiClock.GetSongMapEvaluator().GetTimeSignatureAtTick(MidiClock.GetLastProcessedMidiTick());
 			bTimeSigIsValid = CurrentTimeSignature.Numerator > 0 && CurrentTimeSignature.Denominator > 0;
 			
 			// Find the next pulse and line up phase with the current bar
 			const FMusicTimestamp ClockCurrentTimestamp = MidiClock.GetMusicTimestampAtBlockOffset(0);
-			NextPulseTimestamp.Bar = ClockCurrentTimestamp.Bar;
-			NextPulseTimestamp.Beat = 1;
-			IncrementTimestampByOffset(NextPulseTimestamp, Interval, CurrentTimeSignature);
-			while (bIntervalIsValid && bTimeSigIsValid && NextPulseTimestamp < ClockCurrentTimestamp)
+			FPulse& PulseOn = NextPulses.Add_GetRef(FPulse::On(FMusicTimestamp()));
+			PulseOn.Timestamp.Bar = ClockCurrentTimestamp.Bar;
+			PulseOn.Timestamp.Beat = 1;
+			IncrementTimestampByOffset(PulseOn.Timestamp, Interval, CurrentTimeSignature);
+			while (bIntervalIsValid && bTimeSigIsValid && PulseOn.Timestamp < ClockCurrentTimestamp)
 			{
-				IncrementTimestampByInterval(NextPulseTimestamp, Interval, CurrentTimeSignature);
+				IncrementTimestampByInterval(PulseOn.Timestamp, Interval, CurrentTimeSignature);
 			}
 		}
 		
@@ -51,20 +60,52 @@ namespace Harmonix::Midi::Ops
 		{
 			if (const FAdvance* AsAdvance = ClockEvent.TryGet<FAdvance>())
 			{
-				if (!NextPulseTimestamp.IsValid())
+				if (NextPulses.IsEmpty())
 				{
 					return;
 				}
-				
-				int32 NextPulseTick = MidiClock.GetSongMapEvaluator().MusicTimestampToTick(NextPulseTimestamp);
-
-				while (bIntervalIsValid && bTimeSigIsValid && AsAdvance->LastTickToProcess() >= NextPulseTick)
-				{
-					OnPulse({ ClockEvent.BlockFrameIndex, NextPulseTick });
 
-					IncrementTimestampByInterval(NextPulseTimestamp, Interval, CurrentTimeSignature);
+				// copy the value of the next pulse
+				FPulse Pulse = NextPulses[0];
+				int32 PulseTick = MidiClock.GetSongMapEvaluator().MusicTimestampToTick(Pulse.Timestamp);
 
-					NextPulseTick = MidiClock.GetSongMapEvaluator().MusicTimestampToTick(NextPulseTimestamp);
+				while (bIntervalIsValid && bTimeSigIsValid && AsAdvance->LastTickToProcess() >= PulseTick)
+				{
+					if (Pulse.Type == FPulse::EType::On)
+					{
+						FPulse NextPulseOff = FPulse::Off(Pulse.Timestamp);
+						const float IntervalBeats = Interval.GetIntervalBeats(CurrentTimeSignature);
+						const float DutyCycleClamped = FMath::Clamp(DutyCycle, 0.0f, 1.0f);
+						IncrementTimestampByBeats(NextPulseOff.Timestamp, IntervalBeats * DutyCycleClamped, CurrentTimeSignature);
+
+						int32 NextPulseOffTick = MidiClock.GetSongMapEvaluator().MusicTimestampToTick(NextPulseOff.Timestamp);
+
+						// only send the Pulse if the PulseOff happens in the future
+						if (NextPulseOffTick > PulseTick)
+						{
+							TriggerPulseOn({ ClockEvent.BlockFrameIndex, PulseTick });
+							NextPulses.Add(NextPulseOff);
+						}
+
+						FPulse NextPulseOn = FPulse::On(Pulse.Timestamp);
+						IncrementTimestampByInterval(NextPulseOn.Timestamp, Interval, CurrentTimeSignature);
+						NextPulses.Add(NextPulseOn);
+						
+						NextPulses.RemoveAt(0);
+					}
+					else
+					{
+						TriggerPulseOff({ ClockEvent.BlockFrameIndex, PulseTick });
+						NextPulses.RemoveAt(0);
+					}
+
+					if (NextPulses.IsEmpty())
+					{
+						break;
+					}
+					
+					Pulse = NextPulses[0];
+					PulseTick = MidiClock.GetSongMapEvaluator().MusicTimestampToTick(Pulse.Timestamp);
 				}
 			}
 			else if (const FTimeSignatureChange* AsTimeSigChange = ClockEvent.TryGet<FTimeSignatureChange>())
@@ -74,25 +115,34 @@ namespace Harmonix::Midi::Ops
 				
 				// Time sig changes will come on the downbeat, and if we change time signature,
 				// we want to reset the pulse, so the next pulse is now plus the offset
-				NextPulseTimestamp = MidiClock.GetSongMapEvaluator().TickToMusicTimestamp(AsTimeSigChange->Tick);
-				IncrementTimestampByOffset(NextPulseTimestamp, Interval, CurrentTimeSignature);
+				for (FPulse& Pulse : NextPulses)
+				{
+					Pulse.Timestamp = MidiClock.GetSongMapEvaluator().TickToMusicTimestamp(AsTimeSigChange->Tick);
+					IncrementTimestampByOffset(Pulse.Timestamp, Interval, CurrentTimeSignature);
+				}
 			}
 			else if (const FLoop* AsLoop = ClockEvent.TryGet<FLoop>())
 			{
 				// We assume the pulse should be reset on loop, and the loop start should imply the phase of the pulse
-				NextPulseTimestamp = MidiClock.GetSongMapEvaluator().TickToMusicTimestamp(AsLoop->FirstTickInLoop);
-				IncrementTimestampByOffset(NextPulseTimestamp, Interval, CurrentTimeSignature);
+				for (FPulse& Pulse : NextPulses)
+				{
+					Pulse.Timestamp = MidiClock.GetSongMapEvaluator().TickToMusicTimestamp(AsLoop->FirstTickInLoop);
+					IncrementTimestampByOffset(Pulse.Timestamp, Interval, CurrentTimeSignature);
+				}
 			}
 			else if (const FSeek* AsSeek = ClockEvent.TryGet<FSeek>())
 			{
-				// When we seek, reset the pulse phase to the current bar
-				const FMusicTimestamp ClockCurrentTimestamp = MidiClock.GetSongMapEvaluator().TickToMusicTimestamp(AsSeek->NewNextTick);
-				NextPulseTimestamp.Bar = ClockCurrentTimestamp.Bar;
-				NextPulseTimestamp.Beat = 1;
-				IncrementTimestampByOffset(NextPulseTimestamp, Interval, CurrentTimeSignature);
-				while (bIntervalIsValid && bTimeSigIsValid && NextPulseTimestamp < ClockCurrentTimestamp)
+				for (FPulse& Pulse : NextPulses)
 				{
-					IncrementTimestampByInterval(NextPulseTimestamp, Interval, CurrentTimeSignature);
+					// When we seek, reset the pulse phase to the current bar
+					const FMusicTimestamp ClockCurrentTimestamp = MidiClock.GetSongMapEvaluator().TickToMusicTimestamp(AsSeek->NewNextTick);
+					Pulse.Timestamp.Bar = ClockCurrentTimestamp.Bar;
+					Pulse.Timestamp.Beat = 1;
+					IncrementTimestampByOffset(Pulse.Timestamp, Interval, CurrentTimeSignature);
+					while (bIntervalIsValid && bTimeSigIsValid && Pulse.Timestamp < ClockCurrentTimestamp)
+					{
+						IncrementTimestampByInterval(Pulse.Timestamp, Interval, CurrentTimeSignature);
+					}
 				}
 			}
 		}
@@ -129,51 +179,52 @@ namespace Harmonix::Midi::Ops
 			LastNoteOn.Reset();	
 		}
 
-		FPulseGenerator::Process(MidiClock, [this, &OutStream](const FPulseInfo& Pulse)
-		{
-			AddPulseNote(Pulse.BlockFrameIndex, Pulse.Tick, OutStream);
-		});
+		FPulseGenerator::Process(MidiClock,
+			[this, &OutStream](const FPulseInfo& Pulse)
+			{
+				PulseNoteOn(Pulse.BlockFrameIndex, Pulse.Tick, OutStream);
+			},
+			[this, &OutStream](const FPulseInfo& Pulse)
+			{
+				PulseNoteOff(Pulse.BlockFrameIndex, Pulse.Tick, OutStream);
+			});
 	}
 
-	void FMidiPulseGenerator::AddPulseNote(const int32 BlockFrameIndex, const int32 PulseTick, HarmonixMetasound::FMidiStream& OutStream)
+	void FMidiPulseGenerator::PulseNoteOn(const int32 BlockFrameIndex, const int32 PulseTick, HarmonixMetasound::FMidiStream& OutStream)
 	{
-		int32 NoteOnSample = BlockFrameIndex;
+		if (Enabled)
+		{
+			FMidiMsg Msg{ FMidiMsg::CreateNoteOn(Channel - 1, NoteNumber, Velocity) };
+			HarmonixMetasound::FMidiStreamEvent Event{ &VoiceGenerator, Msg };
+			Event.BlockSampleFrameIndex = BlockFrameIndex;
+			Event.AuthoredMidiTick = PulseTick;
+			Event.CurrentMidiTick = PulseTick;
+			Event.TrackIndex = Track;
+			OutStream.InsertMidiEvent(Event);
 
+			LastNoteOn.Emplace(MoveTemp(Event));
+		}
+	}
+
+	void FMidiPulseGenerator::PulseNoteOff(const int32 BlockFrameIndex, const int32 PulseTick, HarmonixMetasound::FMidiStream& OutStream)
+	{
 		// Note off if there was a previous note on
 		if (LastNoteOn.IsSet())
 		{
 			check(LastNoteOn->MidiMessage.IsNoteOn());
-
-			// Trigger the note off one sample before the note on
-			const int32 NoteOffSample = BlockFrameIndex > 0 ? BlockFrameIndex - 1 : BlockFrameIndex;
-			NoteOnSample = NoteOffSample + 1;
-
 			// Trigger the note off one tick before the note on
-			const int32 NoteOffTick = PulseTick - 1;
+			const int32 NoteOffTick = PulseTick;
 
 			FMidiMsg Msg{ FMidiMsg::CreateNoteOff(LastNoteOn->MidiMessage.GetStdChannel(), LastNoteOn->MidiMessage.GetStdData1()) };
 			HarmonixMetasound::FMidiStreamEvent Event{ &VoiceGenerator, Msg };
-			Event.BlockSampleFrameIndex = NoteOffSample;
-			Event.AuthoredMidiTick = NoteOffTick;
-			Event.CurrentMidiTick = NoteOffTick;
-			Event.TrackIndex = LastNoteOn->TrackIndex;
-			OutStream.InsertMidiEvent(Event);
-
-			LastNoteOn.Reset();
-		}
-
-		// Note on
-		if (Enabled)
-		{
-			FMidiMsg Msg{ FMidiMsg::CreateNoteOn(Channel - 1, NoteNumber, Velocity) };
-			HarmonixMetasound::FMidiStreamEvent Event{ &VoiceGenerator, Msg };
-			Event.BlockSampleFrameIndex = NoteOnSample;
+			Event.BlockSampleFrameIndex = BlockFrameIndex;
 			Event.AuthoredMidiTick = PulseTick;
 			Event.CurrentMidiTick = PulseTick;
-			Event.TrackIndex = Track;
+			Event.TrackIndex = LastNoteOn->TrackIndex;
 			OutStream.InsertMidiEvent(Event);
 
-			LastNoteOn.Emplace(MoveTemp(Event));
+			LastNoteOn.Reset();
 		}
 	}
+
 }
diff --git a/Engine/Plugins/Runtime/Harmonix/Source/HarmonixMetasound/Private/HarmonixMetasound/Nodes/MidiClockSubdivisionTriggerNode.cpp b/Engine/Plugins/Runtime/Harmonix/Source/HarmonixMetasound/Private/HarmonixMetasound/Nodes/MidiClockSubdivisionTriggerNode.cpp
--- a/Engine/Plugins/Runtime/Harmonix/Source/HarmonixMetasound/Private/HarmonixMetasound/Nodes/MidiClockSubdivisionTriggerNode.cpp
+++ b/Engine/Plugins/Runtime/Harmonix/Source/HarmonixMetasound/Private/HarmonixMetasound/Nodes/MidiClockSubdivisionTriggerNode.cpp
@@ -214,13 +214,15 @@ namespace HarmonixMetasound::Nodes::MidiClockSubdivisionTriggerNode
 
 		// If there were pulses, trigger the output
 		TriggerOutPin->AdvanceBlock();
-		PulseGenerator.Process(*MidiClockInPin, [this](const Harmonix::Midi::Ops::FPulseGenerator::FPulseInfo& Pulse)
-		{
-			if (*EnableInPin)
+		PulseGenerator.Process(*MidiClockInPin,
+			[this](const Harmonix::Midi::Ops::FPulseGenerator::FPulseInfo& Pulse)
 			{
-				TriggerOutPin->TriggerFrame(Pulse.BlockFrameIndex);
-			}
-		});
+				if (*EnableInPin)
+				{
+					TriggerOutPin->TriggerFrame(Pulse.BlockFrameIndex);
+				}
+			},
+			[](const Harmonix::Midi::Ops::FPulseGenerator::FPulseInfo&) {});
 	}
 }
 
diff --git a/Engine/Plugins/Runtime/Harmonix/Source/HarmonixMetasound/Private/HarmonixMetasound/Nodes/MidiPulseGeneratorNode.cpp b/Engine/Plugins/Runtime/Harmonix/Source/HarmonixMetasound/Private/HarmonixMetasound/Nodes/MidiPulseGeneratorNode.cpp
--- a/Engine/Plugins/Runtime/Harmonix/Source/HarmonixMetasound/Private/HarmonixMetasound/Nodes/MidiPulseGeneratorNode.cpp
+++ b/Engine/Plugins/Runtime/Harmonix/Source/HarmonixMetasound/Private/HarmonixMetasound/Nodes/MidiPulseGeneratorNode.cpp
@@ -29,6 +29,7 @@ namespace HarmonixMetasound::Nodes::MidiPulseGeneratorNode
 	{
 		DEFINE_METASOUND_PARAM_ALIAS(MidiClock, CommonPinNames::Inputs::MidiClock);
 		DEFINE_METASOUND_PARAM_ALIAS(Enable, CommonPinNames::Inputs::Enable);
+		DEFINE_INPUT_METASOUND_PARAM(DutyCycle, "Duty Cycle", "The pulse width as a percentage of the Interval [0.0, 1.0]")
 		DEFINE_INPUT_METASOUND_PARAM(Interval, "Interval", "The musical time interval at which to send out the pulse");
 		DEFINE_INPUT_METASOUND_PARAM(IntervalMultiplier, "Interval Multiplier", "Multiplies the interval, 1 to use just the value of Interval");
 		DEFINE_INPUT_METASOUND_PARAM(Offset, "Offset", "Offsets the pulse by a musical time");
@@ -83,6 +84,7 @@ namespace HarmonixMetasound::Nodes::MidiPulseGeneratorNode
 				FInputVertexInterface {
 					TInputDataVertex<FMidiClock>(METASOUND_GET_PARAM_NAME_AND_METADATA(Inputs::MidiClock)),
 					TInputDataVertex<bool>(METASOUND_GET_PARAM_NAME_AND_METADATA(Inputs::Enable), true),
+					TInputDataVertex<float>(METASOUND_GET_PARAM_NAME_AND_METADATA(Inputs::DutyCycle), 1.0f),
 					TInputDataVertex<FEnumMidiClockSubdivisionQuantizationType>(
 						METASOUND_GET_PARAM_NAME_AND_METADATA(Inputs::Interval), static_cast<int32>(DefaultInterval.Interval)),
 					TInputDataVertex<int32>(
@@ -113,6 +115,7 @@ namespace HarmonixMetasound::Nodes::MidiPulseGeneratorNode
 		{
 			FMidiClockReadRef Clock;
 			Metasound::FBoolReadRef Enable;
+			Metasound::FFloatReadRef DutyCycle;
 			Metasound::FEnumMidiClockSubdivisionQuantizationReadRef Interval;
 			Metasound::FInt32ReadRef IntervalMultiplier;
 			Metasound::FEnumMidiClockSubdivisionQuantizationReadRef Offset;
@@ -134,6 +137,7 @@ namespace HarmonixMetasound::Nodes::MidiPulseGeneratorNode
 			{
 				InParams.InputData.GetOrCreateDefaultDataReadReference<FMidiClock>(Inputs::MidiClockName, InParams.OperatorSettings),
 				InParams.InputData.GetOrCreateDefaultDataReadReference<bool>(Inputs::EnableName, InParams.OperatorSettings),
+				InParams.InputData.GetOrCreateDefaultDataReadReference<float>(Inputs::DutyCycleName, InParams.OperatorSettings),
 				InParams.InputData.GetOrCreateDefaultDataReadReference<Metasound::FEnumMidiClockSubdivisionQuantizationType>(
 					Inputs::IntervalName, InParams.OperatorSettings),
 				InParams.InputData.GetOrCreateDefaultDataReadReference<int32>(Inputs::IntervalMultiplierName, InParams.OperatorSettings),
@@ -171,6 +175,7 @@ namespace HarmonixMetasound::Nodes::MidiPulseGeneratorNode
 		{
 			InVertexData.BindReadVertex(Inputs::MidiClockName, Inputs.Clock);
 			InVertexData.BindReadVertex(Inputs::EnableName, Inputs.Enable);
+			InVertexData.BindReadVertex(Inputs::DutyCycleName, Inputs.DutyCycle);
 			InVertexData.BindReadVertex(Inputs::IntervalName, Inputs.Interval);
 			InVertexData.BindReadVertex(Inputs::IntervalMultiplierName, Inputs.IntervalMultiplier);
 			InVertexData.BindReadVertex(Inputs::OffsetName, Inputs.Offset);
@@ -210,6 +215,7 @@ namespace HarmonixMetasound::Nodes::MidiPulseGeneratorNode
 					static_cast<uint16>(*Inputs.IntervalMultiplier),
 					static_cast<uint16>(*Inputs.OffsetMultiplier)
 				});
+			PulseGenerator.SetDutyCycle(*Inputs.DutyCycle);
 		}
 		
 		FInputs Inputs;
diff --git a/Engine/Plugins/Runtime/Harmonix/Source/HarmonixMetasound/Public/HarmonixMetasound/DataTypes/MusicTimeInterval.h b/Engine/Plugins/Runtime/Harmonix/Source/HarmonixMetasound/Public/HarmonixMetasound/DataTypes/MusicTimeInterval.h
--- a/Engine/Plugins/Runtime/Harmonix/Source/HarmonixMetasound/Public/HarmonixMetasound/DataTypes/MusicTimeInterval.h
+++ b/Engine/Plugins/Runtime/Harmonix/Source/HarmonixMetasound/Public/HarmonixMetasound/DataTypes/MusicTimeInterval.h
@@ -12,6 +12,10 @@ namespace Harmonix
 		EMidiClockSubdivisionQuantization Offset = EMidiClockSubdivisionQuantization::Beat;
 		uint16 IntervalMultiplier = 1;
 		uint16 OffsetMultiplier = 0;
+
+		HARMONIXMETASOUND_API float GetIntervalBeats(const FTimeSignature& TimeSignature) const;
+
+		HARMONIXMETASOUND_API float GetOffsetBeats(const FTimeSignature& TimeSignature) const;
 	};
 
 	HARMONIXMETASOUND_API void IncrementTimestampByInterval(
@@ -23,4 +27,9 @@ namespace Harmonix
 			FMusicTimestamp& Timestamp,
 			const FMusicTimeInterval& Interval,
 			const FTimeSignature& TimeSignature);
+
+	HARMONIXMETASOUND_API void IncrementTimestampByBeats(
+		FMusicTimestamp& Timestamp,
+		float Beats,
+		const FTimeSignature& TimeSignature);
 }
diff --git a/Engine/Plugins/Runtime/Harmonix/Source/HarmonixMetasound/Public/HarmonixMetasound/MidiOps/PulseGenerator.h b/Engine/Plugins/Runtime/Harmonix/Source/HarmonixMetasound/Public/HarmonixMetasound/MidiOps/PulseGenerator.h
--- a/Engine/Plugins/Runtime/Harmonix/Source/HarmonixMetasound/Public/HarmonixMetasound/MidiOps/PulseGenerator.h
+++ b/Engine/Plugins/Runtime/Harmonix/Source/HarmonixMetasound/Public/HarmonixMetasound/MidiOps/PulseGenerator.h
@@ -21,9 +21,24 @@ namespace Harmonix::Midi::Ops
 		
 		UE_API void Enable(bool bEnable);
 
+		/**
+		 * Set the interval for each pulse in Musical Time.
+		 */
 		UE_API void SetInterval(const FMusicTimeInterval& NewInterval);
 		FMusicTimeInterval GetInterval() const { return Interval; }
 
+		/**
+		 * Set the Duty Cycle of the pulse
+		 *
+		 * The Duty Cycle defines the duration of the pulse as a percentage of the Interval.
+		 *
+		 * for example:
+		 * - Duty Cycle = 1.0: pulse will span the entire length of the interval
+		 * - Duty Cycle = 0.5: pulse will span half the length of the interval
+		 */
+		UE_API void SetDutyCycle(float InDutyCycle);
+		float GetDutyCycle() const { return DutyCycle; }
+
 		UE_API virtual void Reset();
 		
 		struct FPulseInfo
@@ -32,14 +47,46 @@ namespace Harmonix::Midi::Ops
 			int32 Tick;
 		};
 		
-		UE_API void Process(const HarmonixMetasound::FMidiClock& MidiClock, const TFunctionRef<void(const FPulseInfo&)>& OnPulse);
+		UE_API void Process(
+			const HarmonixMetasound::FMidiClock& MidiClock,
+			const TFunctionRef<void(const FPulseInfo&)>& TriggerPulseOn,
+			const TFunctionRef<void(const FPulseInfo&)>& TriggerPulseOff);
 
 	protected:
 		bool Enabled{ true };
 
+		// the interval of the pulse
 		FMusicTimeInterval Interval{};
+
+		// cache the time signature from the clock to calculate seconds from musical time
 		FTimeSignature CurrentTimeSignature{};
-		FMusicTimestamp NextPulseTimestamp{ -1, -1 };
+
+		// pulse width as a percentage of the current interval
+		float DutyCycle = 1.0f;
+
+		struct FPulse
+		{
+			enum class EType 
+			{
+				Off = 0,
+				On = 1,
+			};
+			
+			FMusicTimestamp Timestamp;
+			EType Type = EType::On;
+
+			static FPulse On(const FMusicTimestamp& Timestamp)
+			{
+				return FPulse{ Timestamp, EType::On };
+			}
+
+			static FPulse Off(const FMusicTimestamp& Timestamp)
+			{
+				return FPulse{ Timestamp, EType::Off };
+			}
+		};
+
+		TArray<FPulse> NextPulses;
 	};
 	
 	class FMidiPulseGenerator : public FPulseGenerator
@@ -60,7 +107,8 @@ namespace Harmonix::Midi::Ops
 		UE_API void Process(const HarmonixMetasound::FMidiClock& MidiClock, HarmonixMetasound::FMidiStream& OutStream);
 
 	private:
-		UE_API void AddPulseNote(const int32 BlockFrameIndex, const int32 PulseTick, HarmonixMetasound::FMidiStream& OutStream);
+		UE_API void PulseNoteOn(const int32 BlockFrameIndex, const int32 PulseTick, HarmonixMetasound::FMidiStream& OutStream);
+		UE_API void PulseNoteOff(const int32 BlockFrameIndex, const int32 PulseTick, HarmonixMetasound::FMidiStream& OutStream);
 
 		FMidiVoiceGeneratorBase VoiceGenerator{};
 		TOptional<HarmonixMetasound::FMidiStreamEvent> LastNoteOn;
diff --git a/Engine/Plugins/Runtime/Harmonix/Source/HarmonixMetasound/Public/HarmonixMetasound/Nodes/MidiPulseGeneratorNode.h b/Engine/Plugins/Runtime/Harmonix/Source/HarmonixMetasound/Public/HarmonixMetasound/Nodes/MidiPulseGeneratorNode.h
--- a/Engine/Plugins/Runtime/Harmonix/Source/HarmonixMetasound/Public/HarmonixMetasound/Nodes/MidiPulseGeneratorNode.h
+++ b/Engine/Plugins/Runtime/Harmonix/Source/HarmonixMetasound/Public/HarmonixMetasound/Nodes/MidiPulseGeneratorNode.h
@@ -14,6 +14,7 @@ namespace HarmonixMetasound::Nodes::MidiPulseGeneratorNode
 	namespace Inputs
 	{
 		DECLARE_METASOUND_PARAM_ALIAS(MidiClock);
+		DECLARE_METASOUND_PARAM_EXTERN(DutyCycle);
 		DECLARE_METASOUND_PARAM_EXTERN(Interval);
 		DECLARE_METASOUND_PARAM_EXTERN(IntervalMultiplier);
 		DECLARE_METASOUND_PARAM_EXTERN(Offset);
diff --git a/Engine/Plugins/Runtime/Harmonix/Source/HarmonixMetasoundTests/Private/MidiOps/PulseGeneratorTests.cpp b/Engine/Plugins/Runtime/Harmonix/Source/HarmonixMetasoundTests/Private/MidiOps/PulseGeneratorTests.cpp
--- a/Engine/Plugins/Runtime/Harmonix/Source/HarmonixMetasoundTests/Private/MidiOps/PulseGeneratorTests.cpp
+++ b/Engine/Plugins/Runtime/Harmonix/Source/HarmonixMetasoundTests/Private/MidiOps/PulseGeneratorTests.cpp
@@ -149,6 +149,7 @@ namespace Harmonix::Midi::Ops::Tests
 		{
 			constexpr int32 PulsesToDo = 10;
 			int32 NumPulsesReceived = 0;
+			int32 NumPulseOffReceived = 0;
 
 			Clock->SeekTo(0,0,0);
 			Clock->SetTransportState(0, HarmonixMetasound::EMusicPlayerTransportState::Playing);
@@ -169,11 +170,18 @@ namespace Harmonix::Midi::Ops::Tests
 
 				// Process
 				TArray<FPulseGenerator::FPulseInfo> PulsesThisBlock;
-				PulseGenerator.Process(*Clock, [&PulsesThisBlock](const FPulseGenerator::FPulseInfo& Pulse)
-				{
-					PulsesThisBlock.Add(Pulse);
-				});
+				TArray<FPulseGenerator::FPulseInfo> PulseOffThisBlock;
+				PulseGenerator.Process(*Clock,
+					[&PulsesThisBlock](const FPulseGenerator::FPulseInfo& Pulse)
+					{
+						PulsesThisBlock.Add(Pulse);
+					},
+					[&PulseOffThisBlock](const FPulseGenerator::FPulseInfo& Pulse)
+					{
+						PulseOffThisBlock.Add(Pulse);
+					});
 				NumPulsesReceived += PulsesThisBlock.Num();
+				NumPulseOffReceived += PulseOffThisBlock.Num();
 				
 				// If this is a block where we should get a pulse, check that we got it
 				if (Clock->GetLastProcessedMidiTick() >= Clock->GetSongMapEvaluator().MusicTimestampToTick(NextPulse))
@@ -188,6 +196,7 @@ namespace Harmonix::Midi::Ops::Tests
 			}
 			
 			UTEST_TRUE("Before seek: Got all the pulses at the right time", NumPulsesReceived >= PulsesToDo);
+			UTEST_TRUE("Before seek: Got all the \"pulse-off\" at the right time", NumPulseOffReceived >= PulsesToDo);
 
 			// Seek and reset the next pulse
 			Clock->SeekTo(0, 0, 0);
@@ -202,11 +211,19 @@ namespace Harmonix::Midi::Ops::Tests
 
 				// Process
 				TArray<FPulseGenerator::FPulseInfo> PulsesThisBlock;
-				PulseGenerator.Process(*Clock, [&PulsesThisBlock](const FPulseGenerator::FPulseInfo& Pulse)
-				{
-					PulsesThisBlock.Add(Pulse);
-				});
+				TArray<FPulseGenerator::FPulseInfo> PulseOffThisBlock;
+				PulseGenerator.Process(*Clock,
+					[&PulsesThisBlock](const FPulseGenerator::FPulseInfo& Pulse)
+					{
+						PulsesThisBlock.Add(Pulse);
+					},
+					[&PulseOffThisBlock](const FPulseGenerator::FPulseInfo& Pulse)
+					{
+						PulseOffThisBlock.Add(Pulse);
+					});
+				
 				NumPulsesReceived += PulsesThisBlock.Num();
+				NumPulseOffReceived += PulseOffThisBlock.Num();
 				
 				// If this is a block where we should get a pulse, check that we got it
 				if (Clock->GetLastProcessedMidiTick() >= Clock->GetSongMapEvaluator().MusicTimestampToTick(NextPulse))
@@ -221,6 +238,7 @@ namespace Harmonix::Midi::Ops::Tests
 			}
 			
 			UTEST_TRUE("After seek: Got all the notes at the right time", NumPulsesReceived >= PulsesToDo);
+			UTEST_TRUE("After seek: Gott all the \"pulse off\" at the right time", NumPulseOffReceived >= PulsesToDo);
 		}
 
 		return true;
diff --git a/Engine/Plugins/Runtime/Harmonix/Source/HarmonixMetasoundTests/Private/Nodes/MidiClockSubdivisionTriggerNodeTests.cpp b/Engine/Plugins/Runtime/Harmonix/Source/HarmonixMetasoundTests/Private/Nodes/MidiClockSubdivisionTriggerNodeTests.cpp
--- a/Engine/Plugins/Runtime/Harmonix/Source/HarmonixMetasoundTests/Private/Nodes/MidiClockSubdivisionTriggerNodeTests.cpp
+++ b/Engine/Plugins/Runtime/Harmonix/Source/HarmonixMetasoundTests/Private/Nodes/MidiClockSubdivisionTriggerNodeTests.cpp
@@ -41,10 +41,15 @@ namespace HarmonixMetasound::Nodes::MidiClockSubdivisionTriggerNode::Tests
 			Generator->OnGenerateAudio(Buffer.GetData(), Buffer.Num());
 
 			TArray<int32> PulseGeneratorTriggers;
-			PulseGenerator.Process(*Clock.GetValue(), [&PulseGeneratorTriggers](const Harmonix::Midi::Ops::FPulseGenerator::FPulseInfo& Pulse)
-			{
-				PulseGeneratorTriggers.Add(Pulse.BlockFrameIndex);
-			});
+			PulseGenerator.Process(*Clock.GetValue(),
+				[&PulseGeneratorTriggers](const Harmonix::Midi::Ops::FPulseGenerator::FPulseInfo& Pulse)
+				{
+					PulseGeneratorTriggers.Add(Pulse.BlockFrameIndex);
+				},
+				[&PulseGeneratorTriggers](const Harmonix::Midi::Ops::FPulseGenerator::FPulseInfo& Pulse)
+				{
+					
+				});
 
 			const TArray<int32>& NodeTriggers = (*NodeTriggerOutput)->GetTriggeredFrames();
 			
