Fix flipped Numerator/Denominator methods

Allow time signature changes if the metronome is stopped.

elevate logging to warning

#rnx
#tests Harmonix Engine Tests

diff --git a/Engine/Plugins/Runtime/Harmonix/Source/HarmonixMetasound/Private/HarmonixMetasound/Nodes/MetronomeNode.cpp b/Engine/Plugins/Runtime/Harmonix/Source/HarmonixMetasound/Private/HarmonixMetasound/Nodes/MetronomeNode.cpp
--- a/Engine/Plugins/Runtime/Harmonix/Source/HarmonixMetasound/Private/HarmonixMetasound/Nodes/MetronomeNode.cpp
+++ b/Engine/Plugins/Runtime/Harmonix/Source/HarmonixMetasound/Private/HarmonixMetasound/Nodes/MetronomeNode.cpp
@@ -325,12 +325,12 @@ namespace HarmonixMetasound::Nodes::MetronomeNode
 		return FMath::Clamp(*TempoInPin, kMinTempoBpm, kMaxTempoBpm);
 	}
 
-	int32 FMetronomeOperatorBase::GetInputTimeSigDenominatorClamped() const
+	int32 FMetronomeOperatorBase::GetInputTimeSigNumeratorClamped() const
 	{
 		return FMath::Clamp(*TimeSigNumInPin, kMinTimeSigNumerator, kMaxTimeSigNumerator);;
 	}
 
-	int32 FMetronomeOperatorBase::GetInputTimeSigNumeratorClamped() const
+	int32 FMetronomeOperatorBase::GetInputTimeSigDenominatorClamped() const
 	{
 		return FMath::Clamp(*TimeSigDenomInPin, kMinTimeSigDenominator, kMaxTimeSigDenominator);
 	}
@@ -367,7 +367,6 @@ namespace HarmonixMetasound::Nodes::MetronomeNode
 
 	void FMetronomeOperatorBase::UpdateMidi()
 	{
-		FMidiClock& DrivingMidiClock = GetDrivingMidiClock();
 		bool HasMidiChanges = false;
 		float NewTempo = GetInputTempoClamped();
 		if (!FMath::IsNearlyEqual(CurrentTempo, NewTempo))
@@ -376,20 +375,20 @@ namespace HarmonixMetasound::Nodes::MetronomeNode
 			HasMidiChanges = true;
 		};
 
-		int32 InTimeSigNum = FMath::Clamp(*TimeSigNumInPin, 1, 64);
-		int32 InTimeSigDenom = FMath::Clamp(*TimeSigDenomInPin, 1, 64);
+		int32 InTimeSigNum = GetInputTimeSigNumeratorClamped();
+		int32 InTimeSigDenom = GetInputTimeSigDenominatorClamped();
 		if (InTimeSigNum != CurrentTimeSigNum || InTimeSigDenom != CurrentTimeSigDenom)
 		{
 			CurrentTimeSigNum = InTimeSigNum;
 			CurrentTimeSigDenom = InTimeSigDenom;
-			if (!MidiClockOutPin->HasPersistentLoop())
+			if (IsEffectivelyStopped() || !MidiClockOutPin->HasPersistentLoop())
 			{
 				HandleTimeSigChangeForMidi(InTimeSigNum, InTimeSigDenom);
 				HasMidiChanges = true;
 			}
 			else
 			{
-				UE_LOG(LogMetronomeNode, Log, TEXT("Changing Time Sig. on looping metronome not supported." 
+				UE_LOG(LogMetronomeNode, Warning, TEXT("Changing Time Sig. on looping metronome not supported." 
 					"Changing time signature will require changing loop length which is currently not supported."))
 			}
 		}
