Add support for previous-frame values to MusicTempometer

#changelist validated
#virtualized

diff --git a/Engine/Plugins/Runtime/Harmonix/Source/HarmonixMetasound/Private/HarmonixMetasound/Components/MusicTempometerComponent.cpp b/Engine/Plugins/Runtime/Harmonix/Source/HarmonixMetasound/Private/HarmonixMetasound/Components/MusicTempometerComponent.cpp
--- a/Engine/Plugins/Runtime/Harmonix/Source/HarmonixMetasound/Private/HarmonixMetasound/Components/MusicTempometerComponent.cpp
+++ b/Engine/Plugins/Runtime/Harmonix/Source/HarmonixMetasound/Private/HarmonixMetasound/Components/MusicTempometerComponent.cpp
@@ -1,6 +1,7 @@
 // Copyright Epic Games, Inc. All Rights Reserved.
 
 #include "HarmonixMetasound/Components/MusicTempometerComponent.h"
+#include "Materials/MaterialParameterCollection.h"
 #include "Materials/MaterialParameterCollectionInstance.h"
 #include "Engine/World.h"
 
@@ -20,7 +21,9 @@ void UMusicTempometerComponent::PostLoad()
 	Super::PostLoad();
 
 	// Fix up deprecated properties
-	const FName DeprecatedErroneousName = GetFName();
+	FString DeprecatedErroneousNameString = GetName();
+	DeprecatedErroneousNameString.RemoveFromEnd(TEXT("_GEN_VARIABLE"));
+	FName DeprecatedErroneousName = *DeprecatedErroneousNameString;
 
 	if (UMusicTempometerComponent* DefaultTempometer = GetClass()->GetDefaultObject<UMusicTempometerComponent>())
 	{
@@ -129,17 +132,9 @@ void UMusicTempometerComponent::UpdateCachedSongPos() const
 {
 	LastFrameCounter = GFrameCounter;
 
-	// Cache or clear various FSongPos
-	if (UMusicClockComponent* Clock = GetMutableClockNoMutex())
-	{
-		SongPos = Clock->GetCurrentVideoRenderSongPos();
-	}
-	else
-	{
-		SongPos.Reset();
-	}
+	PreviousFrameSongPos = SongPos;
 
-	MusicTempometerUtilities::UpdateMaterialParameterCollectionFromSongPos(GetOwner(), MaterialParameterCollectionInstance, MaterialParameterCollection, MPCParameters, SongPos);
+	SongPos = MusicTempometerUtilities::UpdateMaterialParameterCollectionFromClock(GetMutableClockNoMutex(), PreviousFrameSongPos, MaterialParameterCollection, MPCParameters, MaterialParameterCollectionInstance);
 }
 
 void UMusicTempometerComponent::UpdateCachedSongPosIfNeeded() const
diff --git a/Engine/Plugins/Runtime/Harmonix/Source/HarmonixMetasound/Private/HarmonixMetasound/Utils/MusicTempometerUtilities.cpp b/Engine/Plugins/Runtime/Harmonix/Source/HarmonixMetasound/Private/HarmonixMetasound/Utils/MusicTempometerUtilities.cpp
--- a/Engine/Plugins/Runtime/Harmonix/Source/HarmonixMetasound/Private/HarmonixMetasound/Utils/MusicTempometerUtilities.cpp
+++ b/Engine/Plugins/Runtime/Harmonix/Source/HarmonixMetasound/Private/HarmonixMetasound/Utils/MusicTempometerUtilities.cpp
@@ -13,45 +13,70 @@ void MusicTempometerUtilities::UpdateMaterialParameterCollectionFromClock(const
 	if (InClockComponent)
 	{
 		const FMidiSongPos MidiSongPos = InClockComponent->GetCurrentVideoRenderSongPos();
-		UpdateMaterialParameterCollectionFromSongPos(InWorldContextObject, InOutMaterialParameterCollectionInstance, InMaterialParameterCollection, InMCPParameters, MidiSongPos);
+		UpdateMaterialParameterCollectionFromSongPos(InWorldContextObject, MidiSongPos, MidiSongPos, InMaterialParameterCollection, InMCPParameters, InOutMaterialParameterCollectionInstance);
 	}
 }
-
+	
 void MusicTempometerUtilities::UpdateMaterialParameterCollectionFromSongPos(const UObject* InWorldContextObject, TWeakObjectPtr<UMaterialParameterCollectionInstance>& InOutMaterialParameterCollectionInstance, const TObjectPtr<UMaterialParameterCollection>& InMaterialParameterCollection, const FMusicTempometerMPCParameters& InMCPParameters, const FMidiSongPos& InMidiSongPos)
 {
-	// Find a MaterialParameterCollectionInstance to update
-	if (!InOutMaterialParameterCollectionInstance.IsValid())
+	UpdateMaterialParameterCollectionFromSongPos(InWorldContextObject, InMidiSongPos, InMidiSongPos, InMaterialParameterCollection, InMCPParameters, InOutMaterialParameterCollectionInstance);
+}
+
+FMidiSongPos MusicTempometerUtilities::UpdateMaterialParameterCollectionFromClock(const UMusicClockComponent* InClockComponent, const FMidiSongPos& InPreviousFrameMidiSongPos, const UMaterialParameterCollection* InMaterialParameterCollection, const FMusicTempometerMPCParameters& InMPCParameters, TWeakObjectPtr<UMaterialParameterCollectionInstance>& InOutMaterialParameterCollectionInstance)
+{
+	FMidiSongPos MidiSongPos;
+	if (InClockComponent)
+	{
+		MidiSongPos = InClockComponent->GetCurrentVideoRenderSongPos();
+		UpdateMaterialParameterCollectionFromSongPos(InClockComponent, MidiSongPos, InPreviousFrameMidiSongPos, InMaterialParameterCollection, InMPCParameters, InOutMaterialParameterCollectionInstance);
+	}
+	return MidiSongPos;
+}
+
+void MusicTempometerUtilities::UpdateMaterialParameterCollectionFromSongPos(const UObject* InWorldContextObject, const FMidiSongPos& InCurrentFrameMidiSongPos, const FMidiSongPos& InPreviousFrameMidiSongPos, const UMaterialParameterCollection* InMaterialParameterCollection, const FMusicTempometerMPCParameters& InMPCParameters, TWeakObjectPtr<UMaterialParameterCollectionInstance>& InOutMaterialParameterCollectionInstance)
+{
+	// Find a MaterialParameterCollectionInstance to update.
+	UMaterialParameterCollectionInstance* MaterialParameterCollectionInstance = InOutMaterialParameterCollectionInstance.Get();
+	if (!MaterialParameterCollectionInstance)
 	{
-		if (InMCPParameters.IsValid())
+		if (InMPCParameters.IsValid() && InWorldContextObject)
 		{
-			if (InWorldContextObject)
+			if (UWorld* World = InWorldContextObject->GetWorld())
 			{
-				if (UWorld* World = InWorldContextObject->GetWorld())
-				{
-					InOutMaterialParameterCollectionInstance = World->GetParameterCollectionInstance(InMaterialParameterCollection);
-				}
+				MaterialParameterCollectionInstance = World->GetParameterCollectionInstance(InMaterialParameterCollection);
+				InOutMaterialParameterCollectionInstance = MaterialParameterCollectionInstance;
 			}
 		}
 
-		if (!InOutMaterialParameterCollectionInstance.IsValid())
+		if (!MaterialParameterCollectionInstance)
 		{
 			return;
 		}
 	}
 
-	InOutMaterialParameterCollectionInstance->SetScalarParameterValue(InMCPParameters.CurrentFrameParameterNames.SecondsIncludingCountInParameterName, InMidiSongPos.SecondsIncludingCountIn);
-	InOutMaterialParameterCollectionInstance->SetScalarParameterValue(InMCPParameters.CurrentFrameParameterNames.BarsIncludingCountInParameterName, InMidiSongPos.BarsIncludingCountIn);
-	InOutMaterialParameterCollectionInstance->SetScalarParameterValue(InMCPParameters.CurrentFrameParameterNames.BeatsIncludingCountInParameterName, InMidiSongPos.BeatsIncludingCountIn);
+	// Update the parameters for each frame.
+	TTuple<const FMidiSongPos*, const FMusicTempometerMPCParameterNames*> Frames[] = {
+		{ &InCurrentFrameMidiSongPos, &InMPCParameters.CurrentFrameParameterNames },
+		{ &InPreviousFrameMidiSongPos, &InMPCParameters.PreviousFrameParameterNames } };
+	for (auto& Frame : Frames)
+	{
+		const FMidiSongPos& MidiSongPos = *get<0>(Frame);
+		const FMusicTempometerMPCParameterNames& MPCParameterNames = *get<1>(Frame);
 
-	InOutMaterialParameterCollectionInstance->SetScalarParameterValue(InMCPParameters.CurrentFrameParameterNames.SecondsFromBarOneParameterName, InMidiSongPos.SecondsFromBarOne);
-	InOutMaterialParameterCollectionInstance->SetScalarParameterValue(InMCPParameters.CurrentFrameParameterNames.TimestampBarParameterName, static_cast<float>(InMidiSongPos.Timestamp.Bar));
-	InOutMaterialParameterCollectionInstance->SetScalarParameterValue(InMCPParameters.CurrentFrameParameterNames.TimestampBeatInBarParameterName, InMidiSongPos.Timestamp.Beat);
+		MaterialParameterCollectionInstance->SetScalarParameterValue(MPCParameterNames.SecondsIncludingCountInParameterName, MidiSongPos.SecondsIncludingCountIn);
+		MaterialParameterCollectionInstance->SetScalarParameterValue(MPCParameterNames.BarsIncludingCountInParameterName, MidiSongPos.BarsIncludingCountIn);
+		MaterialParameterCollectionInstance->SetScalarParameterValue(MPCParameterNames.BeatsIncludingCountInParameterName, MidiSongPos.BeatsIncludingCountIn);
 
-	InOutMaterialParameterCollectionInstance->SetScalarParameterValue(InMCPParameters.CurrentFrameParameterNames.BarProgressParameterName, FMath::Fractional(InMidiSongPos.BarsIncludingCountIn));
-	InOutMaterialParameterCollectionInstance->SetScalarParameterValue(InMCPParameters.CurrentFrameParameterNames.BeatProgressParameterName, FMath::Fractional(InMidiSongPos.BeatsIncludingCountIn));
+		MaterialParameterCollectionInstance->SetScalarParameterValue(MPCParameterNames.SecondsFromBarOneParameterName, MidiSongPos.SecondsFromBarOne);
+		MaterialParameterCollectionInstance->SetScalarParameterValue(MPCParameterNames.TimestampBarParameterName, float(MidiSongPos.Timestamp.Bar));
+		MaterialParameterCollectionInstance->SetScalarParameterValue(MPCParameterNames.TimestampBeatInBarParameterName, MidiSongPos.Timestamp.Beat);
 
-	InOutMaterialParameterCollectionInstance->SetScalarParameterValue(InMCPParameters.CurrentFrameParameterNames.TimeSignatureNumeratorParameterName, InMidiSongPos.TimeSigNumerator);
-	InOutMaterialParameterCollectionInstance->SetScalarParameterValue(InMCPParameters.CurrentFrameParameterNames.TimeSignatureDenominatorParameterName, InMidiSongPos.TimeSigDenominator);
-	InOutMaterialParameterCollectionInstance->SetScalarParameterValue(InMCPParameters.CurrentFrameParameterNames.TempoParameterName, InMidiSongPos.Tempo);
-	InOutMaterialParameterCollectionInstance->SetScalarParameterValue(InMCPParameters.CurrentFrameParameterNames.TimestampValidParameterName, InMidiSongPos.IsValid() ? 1.0f : 0.0f);
-}
\ No newline at end of file
+		MaterialParameterCollectionInstance->SetScalarParameterValue(MPCParameterNames.BarProgressParameterName, FMath::Fractional(MidiSongPos.BarsIncludingCountIn));
+		MaterialParameterCollectionInstance->SetScalarParameterValue(MPCParameterNames.BeatProgressParameterName, FMath::Fractional(MidiSongPos.BeatsIncludingCountIn));
+
+		MaterialParameterCollectionInstance->SetScalarParameterValue(MPCParameterNames.TimeSignatureNumeratorParameterName, MidiSongPos.TimeSigNumerator);
+		MaterialParameterCollectionInstance->SetScalarParameterValue(MPCParameterNames.TimeSignatureDenominatorParameterName, MidiSongPos.TimeSigDenominator);
+		MaterialParameterCollectionInstance->SetScalarParameterValue(MPCParameterNames.TempoParameterName, MidiSongPos.Tempo);
+		MaterialParameterCollectionInstance->SetScalarParameterValue(MPCParameterNames.TimestampValidParameterName, MidiSongPos.IsValid() ? 1.0f : 0.0f);
+	}
+}
diff --git a/Engine/Plugins/Runtime/Harmonix/Source/HarmonixMetasound/Public/HarmonixMetasound/Components/MusicTempometerComponent.h b/Engine/Plugins/Runtime/Harmonix/Source/HarmonixMetasound/Public/HarmonixMetasound/Components/MusicTempometerComponent.h
--- a/Engine/Plugins/Runtime/Harmonix/Source/HarmonixMetasound/Public/HarmonixMetasound/Components/MusicTempometerComponent.h
+++ b/Engine/Plugins/Runtime/Harmonix/Source/HarmonixMetasound/Public/HarmonixMetasound/Components/MusicTempometerComponent.h
@@ -92,7 +92,7 @@ public:
 	FName TempoParameterName_DEPRECATED = "MusicTempo";
 
 	/**
-	 * Deprecated, has moved into structure within MPCParameters
+	 * Parameter names to use for the values in the material parameter collection.
 	 */
 	UPROPERTY(EditAnywhere, Category = "MusicClock")
 	FMusicTempometerMPCParameters MPCParameters;
@@ -107,6 +107,16 @@ public:
 		return SongPos;
 	}
 
+	/**
+	 * The FSongPos for which the game thread previously issued graphics rendering commands.
+	 */
+	UFUNCTION(BlueprintPure, Category = "MusicClock")
+	const FMidiSongPos& GetPreviousFrameSongPos() const
+	{
+		UpdateCachedSongPosIfNeeded();
+		return PreviousFrameSongPos;
+	}
+
 	/**
 	* Seconds from the beginning of the entire music authoring.
 	* Includes all count-in and pickup bars (ie. won't be negative when
@@ -301,6 +311,9 @@ private:
 	UPROPERTY(BlueprintGetter = GetSongPos, Transient, Category = "MusicClock")
 	mutable FMidiSongPos SongPos;
 
+	UPROPERTY(BlueprintGetter = GetPreviousFrameSongPos, Transient, Category = "MusicClock")
+	mutable FMidiSongPos PreviousFrameSongPos;
+
 	mutable FCriticalSection SongPosUpdateMutex;
 	mutable uint64 LastFrameCounter;
 
diff --git a/Engine/Plugins/Runtime/Harmonix/Source/HarmonixMetasound/Public/HarmonixMetasound/Utils/MusicTempometerUtilities.h b/Engine/Plugins/Runtime/Harmonix/Source/HarmonixMetasound/Public/HarmonixMetasound/Utils/MusicTempometerUtilities.h
--- a/Engine/Plugins/Runtime/Harmonix/Source/HarmonixMetasound/Public/HarmonixMetasound/Utils/MusicTempometerUtilities.h
+++ b/Engine/Plugins/Runtime/Harmonix/Source/HarmonixMetasound/Public/HarmonixMetasound/Utils/MusicTempometerUtilities.h
@@ -126,26 +126,8 @@ struct FMusicTempometerMPCParameterNames
 				!TimestampValidParameterName.IsNone();
 	}
 
-	bool operator==(const FMusicTempometerMPCParameterNames& Other)
-	{
-		return	SecondsIncludingCountInParameterName == Other.SecondsIncludingCountInParameterName &&
-			BarsIncludingCountInParameterName == Other.BarsIncludingCountInParameterName &&
-			BeatsIncludingCountInParameterName == Other.BeatsIncludingCountInParameterName &&
-			SecondsFromBarOneParameterName == Other.SecondsFromBarOneParameterName &&
-			TimestampBarParameterName == Other.TimestampBarParameterName &&
-			TimestampBeatInBarParameterName == Other.TimestampBeatInBarParameterName &&
-			BarProgressParameterName == Other.BarProgressParameterName &&
-			BeatProgressParameterName == Other.BeatProgressParameterName &&
-			TimeSignatureNumeratorParameterName == Other.TimeSignatureNumeratorParameterName &&
-			TimeSignatureDenominatorParameterName == Other.TimeSignatureDenominatorParameterName &&
-			TempoParameterName == Other.TempoParameterName &&
-			TimestampValidParameterName == Other.TimestampValidParameterName;
-	}
-
-	bool operator!=(const FMusicTempometerMPCParameterNames& Other)
-	{
-		return	!operator==(Other);
-	}
+	friend bool operator ==(const FMusicTempometerMPCParameterNames& Lhs, const FMusicTempometerMPCParameterNames& Rhs) = default;
+	friend bool operator !=(const FMusicTempometerMPCParameterNames& Lhs, const FMusicTempometerMPCParameterNames& Rhs) = default;
 };
 /**
  * Parameters used when creating/updating the corresponding Material Parameter Collection
@@ -153,14 +135,29 @@ struct FMusicTempometerMPCParameterNames
 USTRUCT(BlueprintType)
 struct FMusicTempometerMPCParameters
 {
-	GENERATED_BODY();
+	GENERATED_BODY()
 
 	UPROPERTY(EditDefaultsOnly, Category = "MusicClock", meta = (ShowOnlyInnerProperties))
 	FMusicTempometerMPCParameterNames CurrentFrameParameterNames;
 
+	UPROPERTY(EditDefaultsOnly, Category = "MusicClock", meta = (ShowOnlyInnerProperties))
+	FMusicTempometerMPCParameterNames PreviousFrameParameterNames{
+		"PreviousFrameMusicSecondsIncludingCountIn",
+		"PreviousFrameMusicBarsIncludingCountIn",
+		"PreviousFrameMusicBeatsIncludingCountIn",
+		"PreviousFrameMusicSecondsFromBarOne",
+		"PreviousFrameMusicTimestampBar",
+		"PreviousFrameMusicTimestampBeatInBar",
+		"PreviousFrameMusicBarProgress",
+		"PreviousFrameMusicBeatProgress",
+		"PreviousFrameMusicTimeSignatureNumerator",
+		"PreviousFrameMusicTimeSignatureDenominator",
+		"PreviousFrameMusicTempo",
+		"PreviousFrameMusicTimestampValid" };
+
 	bool IsValid() const
 	{
-		return CurrentFrameParameterNames.IsValid();
+		return CurrentFrameParameterNames.IsValid() || PreviousFrameParameterNames.IsValid();
 	}
 };
 
@@ -172,10 +169,22 @@ namespace MusicTempometerUtilities
 	/**
 	 * Update the supplied material parameter collection instance (create if null) with song position data extracted from the supplied clock
 	 */
+	UE_DEPRECATED(5.6, "This implementation of UpdateMaterialParameterCollectionFromClock() is deprecated.  Please use the implementation with an 'InPreviousFrameMidiSongPos' argument instead.")
 	void HARMONIXMETASOUND_API UpdateMaterialParameterCollectionFromClock(const UObject* InWorldContextObject, TWeakObjectPtr<UMaterialParameterCollectionInstance>& InOutMaterialParameterCollectionInstance, const TObjectPtr<UMaterialParameterCollection>& InMaterialParameterCollection, const FMusicTempometerMPCParameters& InMCPParameters, const UMusicClockComponent* InClockComponent);
 	
 	/**
 	 * Update the supplied material parameter collection instance (create if null) with data from the supplied song position
 	 */
+	UE_DEPRECATED(5.6, "This implementation of UpdateMaterialParameterCollectionFromSongPos() is deprecated.  Please use the implementation with an 'InPreviousFrameMidiSongPos' argument instead.")
 	void HARMONIXMETASOUND_API UpdateMaterialParameterCollectionFromSongPos(const UObject* InWorldContextObject, TWeakObjectPtr<UMaterialParameterCollectionInstance>& InOutMaterialParameterCollectionInstance, const TObjectPtr<UMaterialParameterCollection>& InMaterialParameterCollection, const FMusicTempometerMPCParameters& InMCPParameters, const FMidiSongPos& InMidiSongPos);
+
+	/**
+	 * Update the supplied material parameter collection instance (create from the supplied material parameter collection if null) from the supplied clock and song position, and return the current song position according to the supplied clock.
+	 */
+	FMidiSongPos HARMONIXMETASOUND_API UpdateMaterialParameterCollectionFromClock(const UMusicClockComponent* InClockComponent, const FMidiSongPos& InPreviousFrameMidiSongPos, const UMaterialParameterCollection* InMaterialParameterCollection, const FMusicTempometerMPCParameters& InMPCParameters, TWeakObjectPtr<UMaterialParameterCollectionInstance>& InOutMaterialParameterCollectionInstance);
+
+	/**
+	 * Update the supplied material parameter collection instance (create from the supplied material parameter collection if null) from the supplied song positions.
+	 */
+	void HARMONIXMETASOUND_API UpdateMaterialParameterCollectionFromSongPos(const UObject* InWorldContextObject, const FMidiSongPos& InCurrentFrameMidiSongPos, const FMidiSongPos& InPreviousFrameMidiSongPos, const UMaterialParameterCollection* InMaterialParameterCollection, const FMusicTempometerMPCParameters& InMPCParameters, TWeakObjectPtr<UMaterialParameterCollectionInstance>& InOutMaterialParameterCollectionInstance);
 };
