Need a const version of GetData

diff --git a/Engine/Plugins/Runtime/Harmonix/Source/HarmonixDsp/Public/HarmonixDsp/AudioBuffer.h b/Engine/Plugins/Runtime/Harmonix/Source/HarmonixDsp/Public/HarmonixDsp/AudioBuffer.h
--- a/Engine/Plugins/Runtime/Harmonix/Source/HarmonixDsp/Public/HarmonixDsp/AudioBuffer.h
+++ b/Engine/Plugins/Runtime/Harmonix/Source/HarmonixDsp/Public/HarmonixDsp/AudioBuffer.h
@@ -754,6 +754,7 @@ public:
 	}
 
 	DataType** GetData() { return MyData; }
+	DataType* const* GetData() const { return MyData; }
 	const FAudioBufferConfig& GetMaxConfig() const { return MaxConfig; }
 
 	FAudioBufferConfig GetValidConfig() const
