Tape Start/Stop Metasound Node

diff --git a/Engine/Plugins/Runtime/Harmonix/Source/HarmonixDsp/Private/HarmonixDsp/Effects/TapeSS.cpp b/Engine/Plugins/Runtime/Harmonix/Source/HarmonixDsp/Private/HarmonixDsp/Effects/TapeSS.cpp
new file mode 100644
--- /dev/null
+++ b/Engine/Plugins/Runtime/Harmonix/Source/HarmonixDsp/Private/HarmonixDsp/Effects/TapeSS.cpp
@@ -0,0 +1,215 @@
+// Copyright Epic Games, Inc. All Rights Reserved.
+#include "HarmonixDSP/Effects/TapeSS.h"
+#include "HarmonixDSP/AudioBuffer.h"
+
+namespace Harmonix::Dsp::Effects
+{
+
+	FTapeStopStart::FTapeStopStart()
+		: MaxNumChannels(0)
+		, ActiveNumChannels(0)
+		, SampleRate(44100.0f)
+		, BufferWritePos(0)
+		, BufferReadPos(0.0f)
+		, NumSamplesInFullRamp(0)
+		, RampProgressInSamples(0)
+		, CurrentState(FTapeStopStart::EState::Play)
+		, DesiredState(FTapeStopStart::EState::Play)
+		, RampTimeSeconds(1.0f)
+	{
+		SRRamper.SetTarget(1.0f);
+		SRRamper.SnapToTarget();
+		SRRamper.SetRampTimeMs(SampleRate, 1000.0f);
+	}
+
+
+	void FTapeStopStart::Prepare(float SampleRateIn, int32 MaxNumChannelsIn, float MaxDecelerationSeconds)
+	{
+		SampleRate = SampleRateIn;
+		MaxNumChannels = MaxNumChannelsIn;
+		ActiveNumChannels = MaxNumChannels;
+
+		SRRamper.SetTarget(1.0f);
+		SRRamper.SnapToTarget();
+		SRRamper.SetRampTimeMs(SampleRate / float(kHopNum), RampTimeSeconds * 1000.0f);
+
+		NumSamplesInFullRamp = FMath::CeilToInt32(RampTimeSeconds * SampleRate);
+		RampProgressInSamples = 0;
+
+		// We only really need a buffer about half the size of the max slow down time.
+		int32 MaxCaptureFrames = FMath::CeilToInt32(MaxDecelerationSeconds / 1.8f * SampleRate);
+		Buffer.Configure(MaxNumChannels, MaxCaptureFrames, EAudioBufferCleanupMode::Delete);
+		Buffer.ZeroData();
+
+		BufferWritePos = 0;
+		DesiredState = CurrentState = EState::Play;
+		BufferReadPos = 0.0f;
+
+		CrossfadeFloor.SetNum(MaxNumChannels);
+	}
+
+	FTapeStopStart::EPhase FTapeStopStart::Process(const TArray<const float*>& InChannelData, const TArray<float*>& OutChannelData, const int32 NumFrames)
+	{
+		check(InChannelData.Num() <= MaxNumChannels);
+		if (InChannelData.Num() != ActiveNumChannels)
+		{
+			Clear();
+			ActiveNumChannels = InChannelData.Num();
+		}
+
+		CheckForStateUpdates();
+
+		int32 FrameIndex = 0;
+
+		EPhase CurrentPhase = EPhase::Playing;
+
+		// First process frames in the "speeding up" or "slowing down" segment...
+		while (!SRRamper.IsAtTarget() && FrameIndex != NumFrames)
+		{
+			if (SRRamper.GetTarget() > 0.5f)
+			{
+				SpeedUp(FrameIndex, InChannelData, OutChannelData, NumFrames);
+				CurrentPhase = EPhase::SpeedingUp;
+			}
+			else
+			{
+				SlowDown(FrameIndex, InChannelData, OutChannelData, NumFrames);
+				CurrentPhase = EPhase::SlowingDown;
+			}
+		}
+
+		// now any portion that is not in the "speeding up" or "slowing down" segment...
+		if (SRRamper.IsAtTarget() && FrameIndex != NumFrames)
+		{
+			if (SRRamper.GetCurrent() < 1.0f)
+			{
+				// silence!
+				for (int32 ChannelIndex = 0; ChannelIndex < ActiveNumChannels; ++ChannelIndex)
+				{
+					FMemory::Memset(OutChannelData[ChannelIndex] + FrameIndex, 0, sizeof(float) * (NumFrames - FrameIndex));
+				}
+				CurrentPhase = EPhase::Stopped;
+			}
+			else
+			{
+				// input!
+				for (int32 ChannelIndex = 0; ChannelIndex < ActiveNumChannels; ++ChannelIndex)
+				{
+					FMemory::Memcpy(OutChannelData[ChannelIndex] + FrameIndex, InChannelData[ChannelIndex] + FrameIndex, sizeof(float) * (NumFrames - FrameIndex));
+				}
+				CurrentPhase = EPhase::Playing;
+			}
+		}
+		
+		return CurrentPhase;
+	}
+
+	void FTapeStopStart::CheckForStateUpdates()
+	{
+		if (CurrentState == DesiredState)
+			return;
+
+		CurrentState = DesiredState;
+
+		BufferWritePos = 0;
+		BufferReadPos = 0.0f;
+
+		float AmountDone = (1.0f - FMath::Abs(SRRamper.GetCurrent() - SRRamper.GetTarget()));
+		SRRamper.SetRampTimeMs(SampleRate / float(kHopNum), AmountDone * RampTimeSeconds * 1000.0f);
+		NumSamplesInFullRamp = FMath::CeilToInt32(AmountDone * RampTimeSeconds * SampleRate);
+		RampProgressInSamples = 0;
+
+		SRRamper.SetRampTimeMs(SampleRate / float(kHopNum), AmountDone * RampTimeSeconds * 1000.0f);
+		if (CurrentState == EState::Play)
+			SRRamper.SetTarget(1.0f);
+		else
+			SRRamper.SetTarget(0.0f);
+	}
+
+	void FTapeStopStart::SnapState(FTapeStopStart::EState TransportState)
+	{
+		DesiredState = CurrentState = TransportState;
+		SRRamper.SetTarget(CurrentState == EState::Play ? 1.0f : 0.0f);
+		SRRamper.SnapToTarget();
+	}
+
+	void FTapeStopStart::SetState(FTapeStopStart::EState TransportState, float DurationSeconds)
+	{
+		DesiredState = TransportState;
+		RampTimeSeconds = DurationSeconds;
+	}
+
+	void FTapeStopStart::SpeedUp(int32& FrameIndex, const TArray<const float*>& InChannelData, const TArray<float*>& OutChannelData, const int32 NumFrames)
+	{
+		int32 PosB = FMath::CeilToInt32(BufferReadPos);
+		int32 PosA = FMath::Max(0, PosB - 1);
+		float WeightA = (float)PosB - BufferReadPos;
+		if (RampProgressInSamples < kCrossfadeSamps)
+		{
+			float Crossfade = float(RampProgressInSamples) / float(kCrossfadeSamps);
+			for (int32 ChannelIndex = 0; ChannelIndex < ActiveNumChannels; ++ChannelIndex)
+			{
+				Buffer[ChannelIndex][BufferWritePos] = *(InChannelData[ChannelIndex] + FrameIndex);
+				float NewSample = (Buffer[ChannelIndex][PosA] * WeightA + Buffer[ChannelIndex][PosB] * (1.0f - WeightA)) * Crossfade + CrossfadeFloor[ChannelIndex] * (1.0f - Crossfade);
+				*(OutChannelData[ChannelIndex] + FrameIndex) = NewSample;
+			}
+		}
+		else if (RampProgressInSamples + kCrossfadeSamps > NumSamplesInFullRamp)
+		{
+			float Crossfade = float(NumSamplesInFullRamp - RampProgressInSamples) / float(kCrossfadeSamps);
+			for (int32 ChannelIndex = 0; ChannelIndex < ActiveNumChannels; ++ChannelIndex)
+			{
+				Buffer[ChannelIndex][BufferWritePos] = *(InChannelData[ChannelIndex] + FrameIndex);
+				float NewSample = (Buffer[ChannelIndex][PosA] * WeightA + Buffer[ChannelIndex][PosB] * (1.0f - WeightA)) * Crossfade + Buffer[ChannelIndex][BufferWritePos] * (1.0f - Crossfade);
+				*(OutChannelData[ChannelIndex] + FrameIndex) = NewSample;
+			}
+		}
+		else
+		{
+			for (int32 ChannelIndex = 0; ChannelIndex < ActiveNumChannels; ++ChannelIndex)
+			{
+				Buffer[ChannelIndex][BufferWritePos] = *(InChannelData[ChannelIndex] + FrameIndex);
+				CrossfadeFloor[ChannelIndex] = Buffer[ChannelIndex][PosA] * WeightA + Buffer[ChannelIndex][PosB] * (1.0f - WeightA);
+				*(OutChannelData[ChannelIndex] + FrameIndex) = CrossfadeFloor[ChannelIndex];
+			}
+		}
+		if (FrameIndex % kHopNum == 0)
+		{
+			SRRamper.Ramp();
+		}
+		BufferWritePos = (BufferWritePos + 1) % Buffer.GetNumValidFrames();
+		BufferReadPos += FMath::Abs(SRRamper.GetCurrent());
+		if (BufferReadPos >= Buffer.GetNumValidFrames())
+		{
+			BufferReadPos -= Buffer.GetNumValidFrames();
+		}
+		++FrameIndex;
+		++RampProgressInSamples;
+	}
+
+	void FTapeStopStart::SlowDown(int32& FrameIndex, const TArray<const float*>& InChannelData, const TArray<float*>& OutChannelData, const int32 NumFrames)
+	{
+		int32 PosB = FMath::CeilToInt32(BufferReadPos);
+		int32 PosA = FMath::Max(0, PosB - 1);
+		float WeightA = (float)PosB - BufferReadPos;
+		float Crossfade = FMath::Min(float(NumSamplesInFullRamp - RampProgressInSamples) / float(kCrossfadeSamps), 1.0f);
+		for (int32 ChannelIndex = 0; ChannelIndex < ActiveNumChannels; ++ChannelIndex)
+		{
+			Buffer[ChannelIndex][BufferWritePos] = *(InChannelData[ChannelIndex] + FrameIndex);
+			CrossfadeFloor[ChannelIndex] = Crossfade * (Buffer[ChannelIndex][PosA] * WeightA + Buffer[ChannelIndex][PosB] * (1.0f - WeightA));
+			*(OutChannelData[ChannelIndex] + FrameIndex) = CrossfadeFloor[ChannelIndex];
+		}
+		if (FrameIndex % kHopNum == 0)
+		{
+			SRRamper.Ramp();
+		}
+		BufferWritePos = (BufferWritePos + 1) % Buffer.GetNumValidFrames();
+		BufferReadPos += FMath::Abs(SRRamper.GetCurrent());
+		if (BufferReadPos >= Buffer.GetNumValidFrames())
+		{
+			BufferReadPos -= Buffer.GetNumValidFrames();
+		}
+		++FrameIndex;
+		++RampProgressInSamples;
+	}
+}
diff --git a/Engine/Plugins/Runtime/Harmonix/Source/HarmonixDsp/Public/HarmonixDsp/Effects/TapeSS.h b/Engine/Plugins/Runtime/Harmonix/Source/HarmonixDsp/Public/HarmonixDsp/Effects/TapeSS.h
new file mode 100644
--- /dev/null
+++ b/Engine/Plugins/Runtime/Harmonix/Source/HarmonixDsp/Public/HarmonixDsp/Effects/TapeSS.h
@@ -0,0 +1,71 @@
+// Copyright Epic Games, Inc. All Rights Reserved.
+#pragma once
+
+#include "HarmonixDSP/AudioBuffer.h"
+#include "HarmonixDSP/Ramper.h"
+
+namespace Harmonix::Dsp::Effects
+{
+	class FTapeStopStart
+	{
+	public:
+
+		enum class EState
+		{
+			Play,
+			Stop
+		};
+
+		enum class EPhase
+		{
+			SpeedingUp,
+			Playing,
+			Stopped,
+			SlowingDown
+		};
+
+		HARMONIXDSP_API FTapeStopStart();
+
+		HARMONIXDSP_API void Prepare(float SampleRate, int32 MaxNumChannels, float MaxDecelerationSeconds);
+
+		void Clear()
+		{
+			BufferWritePos = 0;
+			BufferReadPos = 0.0f;
+			Buffer.ZeroData();
+		}
+
+		HARMONIXDSP_API EPhase Process(const TArray<const float*>& InChannelData, const TArray<float*>& OutChannelData, const int32 NumFrames);
+
+		HARMONIXDSP_API void  SnapState(EState State);
+		HARMONIXDSP_API void  SetState(EState State, float DurationSeconds);
+		HARMONIXDSP_API EState GetState() const { return CurrentState; }
+		HARMONIXDSP_API float GetSpeed() const { return SRRamper.GetCurrent(); }
+
+	private:
+		void CheckForStateUpdates();
+		void SpeedUp(int32& FrameIndex, const TArray<const float*>& InChannelData, const TArray<float*>& OutChannelData, const int32 NumFrames);
+		void SlowDown(int32& FrameIndex, const TArray<const float*>& InChannelData, const TArray<float*>& OutChannelData, const int32 NumFrames);
+
+		int32 MaxNumChannels;
+		int32 ActiveNumChannels;
+		float SampleRate;
+
+		TAudioBuffer<float>   Buffer;
+		int32 BufferWritePos;
+		float BufferReadPos;
+
+		int32 NumSamplesInFullRamp;
+		int32 RampProgressInSamples;
+		TArray<float> CrossfadeFloor;
+
+		TLinearRamper<float> SRRamper;
+
+		EState CurrentState;
+		EState DesiredState;
+		float RampTimeSeconds;
+
+		static const int32 kHopNum = 4;
+		static const int32 kCrossfadeSamps = 512;
+	};
+}
\ No newline at end of file
diff --git a/Engine/Plugins/Runtime/Harmonix/Source/HarmonixMetasound/Private/HarmonixMetasound/Common.cpp b/Engine/Plugins/Runtime/Harmonix/Source/HarmonixMetasound/Private/HarmonixMetasound/Common.cpp
--- a/Engine/Plugins/Runtime/Harmonix/Source/HarmonixMetasound/Private/HarmonixMetasound/Common.cpp
+++ b/Engine/Plugins/Runtime/Harmonix/Source/HarmonixMetasound/Private/HarmonixMetasound/Common.cpp
@@ -59,16 +59,22 @@ namespace HarmonixMetasound
 			DEFINE_INPUT_METASOUND_PARAM(Transposition,                 "Transposition",                      "Transposition in semitones.")
 			DEFINE_INPUT_METASOUND_PARAM(ClockSpeedToPitch,             "MIDI Clock Speed Affects Pitch",     "When true the clock speed of the incoming MIDI stream will affect the pitch of the synth.")
 			DEFINE_INPUT_METASOUND_PARAM(ClockSpeedToFrequency,         "MIDI Clock Speed Affects Frequency", "When true the clock speed of the incoming MIDI stream will affect the calculated frequency.")
-			DEFINE_INPUT_METASOUND_PARAM(LFOSyncType,					"LFO Sync Type",					  "Specifies if and how the LFO should sync to a clock");
-			DEFINE_INPUT_METASOUND_PARAM(LFOFrequency,				    "LFO Frequency",					  "The period of the LFO. If Sync Type is TempoSync, the unit is cycles per quarter note. Otherwise, the unit is Hz (cycles per second).");
-			DEFINE_INPUT_METASOUND_PARAM(LFOInvert,					    "Invert LFO",						  "Toggle to invert the LFO");
-			DEFINE_INPUT_METASOUND_PARAM(LFOShape,					    "LFO Shape",						  "The shape of the LFO. 0.0-1.0 = morph between square and triangle. 1.0-2.0 = morph between triangle and sawtooth");
-			DEFINE_INPUT_METASOUND_PARAM(AudioMono,					    "Audio Mono",					      "The mono audio.");
-			DEFINE_INPUT_METASOUND_PARAM(AudioLeft,					    "Audio Left",					      "The left channel audio.");
-			DEFINE_INPUT_METASOUND_PARAM(AudioRight,					"Audio Right",					      "The right channel audio.");
-			DEFINE_INPUT_METASOUND_PARAM(Timestamp,                     "Music Timestamp",                    "The Bar and Beat of interest.");
-			DEFINE_INPUT_METASOUND_PARAM(ShouldQuantizeTimestamp,       "Quantize",                           "Quantize the provide timestamp (Bar & Beat) to a musical subdivision.");
-			DEFINE_INPUT_METASOUND_PARAM(QuantizationUnit,              "Quantization Unit",                  "The music subdivision that a music timestamp should be quantized to.");
+			DEFINE_INPUT_METASOUND_PARAM(LFOSyncType,					"LFO Sync Type",					  "Specifies if and how the LFO should sync to a clock")
+			DEFINE_INPUT_METASOUND_PARAM(LFOFrequency,				    "LFO Frequency",					  "The period of the LFO. If Sync Type is TempoSync, the unit is cycles per quarter note. Otherwise, the unit is Hz (cycles per second).")
+			DEFINE_INPUT_METASOUND_PARAM(LFOInvert,					    "Invert LFO",						  "Toggle to invert the LFO")
+			DEFINE_INPUT_METASOUND_PARAM(LFOShape,					    "LFO Shape",						  "The shape of the LFO. 0.0-1.0 = morph between square and triangle. 1.0-2.0 = morph between triangle and sawtooth")
+			DEFINE_INPUT_METASOUND_PARAM(AudioMono,					    "Audio Mono",					      "The mono audio.")
+			DEFINE_INPUT_METASOUND_PARAM(AudioLeft,					    "Audio Left",					      "The left channel audio.")
+			DEFINE_INPUT_METASOUND_PARAM(AudioRight,					"Audio Right",					      "The right channel audio.")
+			DEFINE_INPUT_METASOUND_PARAM(AudioCenter,                   "Audio Center",                       "5.1/7.1: The center channel audio.")
+			DEFINE_INPUT_METASOUND_PARAM(AudioLFE,                      "Audio LFE",                          "5.1/7.1: The LFE channel audio.")
+			DEFINE_INPUT_METASOUND_PARAM(AudioLeftSide,                 "Audio Left Side",                    "Quad/5.1: Left Surround, 7.1: Left Side.")
+			DEFINE_INPUT_METASOUND_PARAM(AudioRightSide,                "Audio Right Side",                   "Quad/5.1: Right Surround, 7.1: Right Side.")
+			DEFINE_INPUT_METASOUND_PARAM(AudioLeftBack,                 "Audio Left Back - Surround",         "7.1: Left Back.")
+			DEFINE_INPUT_METASOUND_PARAM(AudioRightBack,                "Audio Right Back - Surround",        "7.1: Right Back.")
+			DEFINE_INPUT_METASOUND_PARAM(Timestamp,                     "Music Timestamp",                    "The Bar and Beat of interest.")
+			DEFINE_INPUT_METASOUND_PARAM(ShouldQuantizeTimestamp,       "Quantize",                           "Quantize the provide timestamp (Bar & Beat) to a musical subdivision.")
+			DEFINE_INPUT_METASOUND_PARAM(QuantizationUnit,              "Quantization Unit",                  "The music subdivision that a music timestamp should be quantized to.")
 		}
 		namespace Outputs
 		{
diff --git a/Engine/Plugins/Runtime/Harmonix/Source/HarmonixMetasound/Private/HarmonixMetasound/Nodes/TapeSSNode.cpp b/Engine/Plugins/Runtime/Harmonix/Source/HarmonixMetasound/Private/HarmonixMetasound/Nodes/TapeSSNode.cpp
new file mode 100644
--- /dev/null
+++ b/Engine/Plugins/Runtime/Harmonix/Source/HarmonixMetasound/Private/HarmonixMetasound/Nodes/TapeSSNode.cpp
@@ -0,0 +1,318 @@
+// Copyright Epic Games, Inc. All Rights Reserved.
+
+#include "HarmonixMetasound/Nodes/TapeSSNode.h"
+
+#include "MetasoundExecutableOperator.h"
+#include "MetasoundFacade.h"
+#include "MetasoundNodeInterface.h"
+#include "MetasoundParamHelper.h"
+#include "MetasoundSampleCounter.h"
+#include "MetasoundStandardNodesCategories.h"
+#include "MetasoundVertex.h"
+//#include "HarmonixDsp/StridePointer.h"
+#include "HarmonixDsp/Effects/TapeSS.h"
+#include "HarmonixMetasound/Common.h"
+
+#define LOCTEXT_NAMESPACE "HarmonixMetaSound"
+
+namespace HarmonixMetasound::TapeSSNode
+{
+	using namespace Metasound;
+	using namespace Harmonix::Dsp::Effects;
+
+	namespace Inputs
+	{
+		DEFINE_METASOUND_PARAM_ALIAS(Enabled, CommonPinNames::Inputs::Enable);
+		DEFINE_INPUT_METASOUND_PARAM(Stop, "Stop", "When toggled from false to true, causes audio to slow down and stop.");
+		DEFINE_INPUT_METASOUND_PARAM(RampDuration, "Duration", "How long it should take to slow down from full speed to stopped in seconds.");
+		static constexpr int32 kNumNonAudioInputs = 3; // NOTE: Keep in sync with the number of input vertecies defined above!
+	}
+
+	namespace Outputs
+	{
+		DEFINE_OUTPUT_METASOUND_PARAM(InputAudioInUse, "Input Audio In Use", "True when audio from the input is being used to generate the output. False when the input audio is unneeded and can be changed.");
+		DEFINE_OUTPUT_METASOUND_PARAM(OnStopped, "On Stopped", "Triggered when the audio has completely stopped.");
+		DEFINE_OUTPUT_METASOUND_PARAM(OnFullSpeed, "On Full Speed", "Triggered when the audio has sped back up to full speed.");
+		DEFINE_OUTPUT_METASOUND_PARAM(CurrentSpeed, "CurrentSpeed", "A value from 0 (stopped) to 1 (playing normally).");
+	}
+
+	const FLazyName FTapeSSNodeOperatorData::OperatorDataTypeName = "TapeStopStartOperatorData";
+
+	static TArray<FVertexName> CreateVertexNames(const TCHAR* Prefix, int32 NumChannels)
+	{
+		TStringBuilder<32> NameStr;
+		TArray<FVertexName> Names;
+		Names.AddUninitialized(NumChannels);
+
+		for (int ChannelIndex = 0; ChannelIndex < NumChannels; ++ChannelIndex)
+		{
+			NameStr << Prefix << ChannelIndex;
+			Names[ChannelIndex] = *NameStr;
+			NameStr.Reset();
+		}
+
+		return Names;
+	}
+
+	static const TArray<FVertexName> AudioInputNames = CreateVertexNames(TEXT("In "), kMaxChannels);
+	static const TArray<FVertexName> AudioOutputNames = CreateVertexNames(TEXT("Out "), kMaxChannels);
+
+#if WITH_EDITOR
+	static const FText GetAudioInputDisplayName(uint32 ChannelIndex)
+	{
+		return METASOUND_LOCTEXT_FORMAT("AudioInputDisplayName", "In {0}", ChannelIndex);
+	}
+
+	static const FText GetAudioOutputDisplayName(uint32 ChannelIndex)
+	{
+		return METASOUND_LOCTEXT_FORMAT("AudioOutputDisplayName", "Out {0}", ChannelIndex);
+	}
+#endif // WITH_EDITOR
+
+	FVertexInterface GetVertexInterface(int32 NumChannels)
+	{
+		using namespace HarmonixMetasound::TapeSSNode;
+
+		FInputVertexInterface InputInterface;
+		FOutputVertexInterface OutputInterface;
+
+		InputInterface.Add(TInputDataVertex<bool>(METASOUND_GET_PARAM_NAME_AND_METADATA(Inputs::Enabled)));
+		InputInterface.Add(TInputDataVertex<bool>(METASOUND_GET_PARAM_NAME_AND_METADATA(Inputs::Stop)));
+		InputInterface.Add(TInputDataVertex<float>(METASOUND_GET_PARAM_NAME_AND_METADATA(Inputs::RampDuration)));
+
+		OutputInterface.Add(TOutputDataVertex<bool>(METASOUND_GET_PARAM_NAME_AND_METADATA(Outputs::InputAudioInUse)));
+		OutputInterface.Add(TOutputDataVertex<FTrigger>(METASOUND_GET_PARAM_NAME_AND_METADATA(Outputs::OnStopped)));
+		OutputInterface.Add(TOutputDataVertex<FTrigger>(METASOUND_GET_PARAM_NAME_AND_METADATA(Outputs::OnFullSpeed)));
+		OutputInterface.Add(TOutputDataVertex<float>(METASOUND_GET_PARAM_NAME_AND_METADATA(Outputs::CurrentSpeed)));
+
+		check (InputInterface.Num() == Inputs::kNumNonAudioInputs);
+		check (NumChannels <= kMaxChannels);
+
+		TStringBuilder<32> NameStr;
+		for (int32 ChannelIndex = 0; ChannelIndex < NumChannels; ++ChannelIndex)
+		{
+#if WITH_EDITOR
+			const FDataVertexMetadata AudioInputMetadata
+			{
+				METASOUND_LOCTEXT("AudioInputDescription", "Audio Input"),
+				GetAudioInputDisplayName(ChannelIndex)
+			};
+#else 
+			const FDataVertexMetadata AudioInputMetadata;
+#endif // WITH_EDITOR
+			InputInterface.Add(TInputDataVertex<FAudioBuffer>(AudioInputNames[ChannelIndex], AudioInputMetadata));
+
+#if WITH_EDITOR
+			const FDataVertexMetadata AudioOutputMetadata
+			{
+				METASOUND_LOCTEXT("AudioOutputDescription", "Audio Output"),
+				GetAudioOutputDisplayName(ChannelIndex)
+			};
+#else 
+			const FDataVertexMetadata AudioOutputMetadata;
+#endif // WITH_EDITOR
+			OutputInterface.Add(TOutputDataVertex<FAudioBuffer>(AudioOutputNames[ChannelIndex], AudioOutputMetadata));
+		}
+
+		return FVertexInterface
+		{
+			MoveTemp(InputInterface),
+			MoveTemp(OutputInterface)
+		};
+	}
+
+	class FTapeSSOperator final : public TExecutableOperator<FTapeSSOperator>
+	{
+	public:
+		struct FInputs
+		{
+			FBoolReadRef        Enabled;
+			FBoolReadRef        Stop;
+			FFloatReadRef       RampDuration;
+			TArray<FAudioBufferReadRef> AudioInputs;
+		};
+
+		struct FOutputs
+		{
+			FBoolWriteRef        InputAudioInUse;
+			FTriggerWriteRef     OnStopped;
+			FTriggerWriteRef     OnFullSpeed;
+			FFloatWriteRef       CurrentSpeed;
+			TArray<FAudioBufferWriteRef> AudioOutputs;
+		};
+		
+		FTapeSSOperator(const FBuildOperatorParams& Params, FInputs&& Inputs, FOutputs&& Outputs)
+			: Inputs(MoveTemp(Inputs))
+			, Outputs(MoveTemp(Outputs))
+		{
+			TSharedPtr<const FTapeSSNodeOperatorData> InOperatorData = StaticCastSharedPtr<const FTapeSSNodeOperatorData>(Params.Node.GetOperatorData());
+			Reset(Params, InOperatorData ? InOperatorData->MaxRampSeconds : 5.0f);
+		}
+		
+		static const FNodeClassMetadata& GetNodeInfo()
+		{
+			auto InitNodeInfo = []() -> FNodeClassMetadata
+			{
+				FNodeClassMetadata Info = FNodeClassMetadata::GetEmpty();
+				Info.ClassName = { HarmonixNodeNamespace, TEXT("TapeSS"), TEXT("") };
+				Info.MajorVersion = 0;
+				Info.MinorVersion = 1;
+				Info.DisplayName = METASOUND_LOCTEXT("TapeSS_DisplayName", "Tape Stop/Start");
+				Info.CategoryHierarchy = { MetasoundNodeCategories::Harmonix };
+				Info.Description = METASOUND_LOCTEXT("TapeSS_Description", "Simulates the stopping and starting of a tape machine.");
+				Info.Author = PluginAuthor;
+				Info.PromptIfMissing = PluginNodeMissingPrompt;
+				Info.DefaultInterface = GetVertexInterface(1);
+				return Info;
+			};
+
+			static const FNodeClassMetadata Info = InitNodeInfo();
+
+			return Info;
+		}
+
+		static TUniquePtr<IOperator> CreateOperator(const FBuildOperatorParams& InParams, FBuildResults& OutResults)
+		{
+			const FVertexInterface& NodeInterface = InParams.Node.GetVertexInterface();
+			const FInputVertexInterfaceData& InputData = InParams.InputData;
+
+			TArray<FAudioBufferReadRef> AudioInputs;
+			TArray<FAudioBufferWriteRef> AudioOutputs;
+			int32 NumAudioChannels = NodeInterface.GetInputInterface().Num() - Inputs::kNumNonAudioInputs;
+			for (int32 ChannelIndex = 0; ChannelIndex < NumAudioChannels; ++ChannelIndex)
+			{
+				AudioInputs.Emplace(InputData.GetOrCreateDefaultDataReadReference<FAudioBuffer>(AudioInputNames[ChannelIndex], InParams.OperatorSettings));
+				AudioOutputs.Emplace(TDataWriteReferenceFactory<FAudioBuffer>::CreateExplicitArgs(InParams.OperatorSettings));
+			}
+
+			FInputs Inputs {
+					InputData.GetOrCreateDefaultDataReadReference<bool>(Inputs::EnabledName, InParams.OperatorSettings),
+					InputData.GetOrCreateDefaultDataReadReference<bool>(Inputs::StopName, InParams.OperatorSettings),
+					InputData.GetOrCreateDefaultDataReadReference<float>(Inputs::RampDurationName, InParams.OperatorSettings),
+					MoveTemp(AudioInputs)
+			};
+
+			FOutputs Outputs {	
+					FBoolWriteRef::CreateNew(),
+					FTriggerWriteRef::CreateNew(InParams.OperatorSettings),
+					FTriggerWriteRef::CreateNew(InParams.OperatorSettings),
+					FFloatWriteRef::CreateNew(),
+					MoveTemp(AudioOutputs)
+			};
+
+			return MakeUnique<FTapeSSOperator>(
+				InParams,
+				MoveTemp(Inputs),
+				MoveTemp(Outputs)
+			);
+		}
+
+		virtual void BindInputs(FInputVertexInterfaceData& InVertexData) override
+		{
+			InVertexData.BindReadVertex(Inputs::EnabledName, Inputs.Enabled);
+			InVertexData.BindReadVertex(Inputs::StopName, Inputs.Stop);
+			InVertexData.BindReadVertex(Inputs::RampDurationName, Inputs.RampDuration);
+			for (int32 ChannelIndex = 0; ChannelIndex < Inputs.AudioInputs.Num(); ++ChannelIndex)
+			{
+				InVertexData.BindReadVertex(AudioInputNames[ChannelIndex], Inputs.AudioInputs[ChannelIndex]);
+			}
+		}
+
+		virtual void BindOutputs(FOutputVertexInterfaceData& InVertexData) override
+		{
+			InVertexData.BindReadVertex(Outputs::InputAudioInUseName, Outputs.InputAudioInUse);
+			InVertexData.BindReadVertex(Outputs::OnStoppedName, Outputs.OnStopped);
+			InVertexData.BindReadVertex(Outputs::OnFullSpeedName, Outputs.OnFullSpeed);
+			InVertexData.BindReadVertex(Outputs::CurrentSpeedName, Outputs.CurrentSpeed);
+			for (int32 ChannelIndex = 0; ChannelIndex < Outputs.AudioOutputs.Num(); ++ChannelIndex)
+			{
+				InVertexData.BindReadVertex(AudioOutputNames[ChannelIndex], Outputs.AudioOutputs[ChannelIndex]);
+			}
+		}
+		
+		void Reset(const FResetParams& Params, float InMaxRampSeconds)
+		{
+			MaxRampSeconds = InMaxRampSeconds;
+			TapeSS.Prepare(Params.OperatorSettings.GetSampleRate(), Inputs.AudioInputs.Num(), MaxRampSeconds);
+			*Outputs.InputAudioInUse = true;
+			ChannelInputPointers.SetNumUninitialized(Inputs.AudioInputs.Num());
+			ChannelOutputPointers.SetNumUninitialized(Inputs.AudioInputs.Num());
+		}
+
+		void Execute()
+		{
+			Outputs.OnStopped->AdvanceBlock();
+			Outputs.OnFullSpeed->AdvanceBlock();
+
+			if (!*Inputs.Enabled)
+			{
+				*Outputs.InputAudioInUse = true;
+				for (int32 ChannelIndex = 0; ChannelIndex < Outputs.AudioOutputs.Num(); ++ChannelIndex)
+				{
+					FMemory::Memcpy(Outputs.AudioOutputs[ChannelIndex]->GetData(), Inputs.AudioInputs[ChannelIndex]->GetData(), sizeof(float) * Outputs.AudioOutputs[ChannelIndex]->Num());
+				}
+				return;
+			}
+
+			if (*Inputs.Stop && TapeSS.GetState() != FTapeStopStart::EState::Stop)
+			{
+				TapeSS.SetState(FTapeStopStart::EState::Stop, FMath::Min(*Inputs.RampDuration, MaxRampSeconds));
+			}
+			else if (!*Inputs.Stop && TapeSS.GetState() == FTapeStopStart::EState::Stop)
+			{
+				TapeSS.SetState(FTapeStopStart::EState::Play, FMath::Min(*Inputs.RampDuration, MaxRampSeconds));
+			}
+
+			for (int32 ChannelIndex = 0; ChannelIndex < Inputs.AudioInputs.Num(); ++ChannelIndex)
+			{
+				ChannelInputPointers[ChannelIndex] = Inputs.AudioInputs[ChannelIndex]->GetData();
+				ChannelOutputPointers[ChannelIndex] = Outputs.AudioOutputs[ChannelIndex]->GetData();
+			}
+
+			FTapeStopStart::EPhase CurrentPhase = TapeSS.Process(ChannelInputPointers, ChannelOutputPointers, Inputs.AudioInputs[0]->Num());
+			if (CurrentPhase != LastPhase)
+			{
+				if (CurrentPhase == FTapeStopStart::EPhase::Playing)
+				{
+					Outputs.OnFullSpeed->TriggerFrame(0);
+				}
+				else if (CurrentPhase == FTapeStopStart::EPhase::Stopped)
+				{
+					Outputs.OnStopped->TriggerFrame(0);
+				}
+				LastPhase = CurrentPhase;
+			}
+			*Outputs.InputAudioInUse = true;
+			if (CurrentPhase == FTapeStopStart::EPhase::Stopped || (CurrentPhase == FTapeStopStart::EPhase::SlowingDown && TapeSS.GetSpeed() < 0.5f))
+			{
+				*Outputs.InputAudioInUse = false;
+			}
+			*Outputs.CurrentSpeed = TapeSS.GetSpeed();
+		}
+
+	private:
+		FInputs Inputs;
+		FOutputs Outputs;
+		FTapeStopStart TapeSS;
+		FTapeStopStart::EPhase LastPhase = FTapeStopStart::EPhase::Playing;
+		TArray<const float*> ChannelInputPointers;
+		TArray<float*> ChannelOutputPointers;
+		int32 BlockSamples;
+		float MaxRampSeconds = 5.0f;
+	};
+
+	using FTapeSSNode = Metasound::TNodeFacade<FTapeSSOperator>;
+	METASOUND_REGISTER_NODE_AND_CONFIGURATION(FTapeSSNode, FHarmonixTapeSSNodeConfiguration);
+}
+
+TInstancedStruct<FMetasoundFrontendClassInterface> FHarmonixTapeSSNodeConfiguration::OverrideDefaultInterface(const FMetasoundFrontendClass& InNodeClass) const
+{
+	return TInstancedStruct<FMetasoundFrontendClassInterface>::Make(FMetasoundFrontendClassInterface::GenerateClassInterface(HarmonixMetasound::TapeSSNode::GetVertexInterface(NumChannels)));
+}
+
+TSharedPtr<const Metasound::IOperatorData> FHarmonixTapeSSNodeConfiguration::GetOperatorData() const
+{
+	return MakeShared<HarmonixMetasound::TapeSSNode::FTapeSSNodeOperatorData>(MaxRampSeconds);
+}
+
+#undef LOCTEXT_NAMESPACE
diff --git a/Engine/Plugins/Runtime/Harmonix/Source/HarmonixMetasound/Public/HarmonixMetasound/Common.h b/Engine/Plugins/Runtime/Harmonix/Source/HarmonixMetasound/Public/HarmonixMetasound/Common.h
--- a/Engine/Plugins/Runtime/Harmonix/Source/HarmonixMetasound/Public/HarmonixMetasound/Common.h
+++ b/Engine/Plugins/Runtime/Harmonix/Source/HarmonixMetasound/Public/HarmonixMetasound/Common.h
@@ -104,6 +104,12 @@ namespace HarmonixMetasound
 			DECLARE_METASOUND_PARAM_EXTERN(AudioMono);
 			DECLARE_METASOUND_PARAM_EXTERN(AudioLeft);
 			DECLARE_METASOUND_PARAM_EXTERN(AudioRight);
+			DECLARE_METASOUND_PARAM_EXTERN(AudioCenter)
+			DECLARE_METASOUND_PARAM_EXTERN(AudioLFE)
+			DECLARE_METASOUND_PARAM_EXTERN(AudioLeftBack)
+			DECLARE_METASOUND_PARAM_EXTERN(AudioRightBack)
+			DECLARE_METASOUND_PARAM_EXTERN(AudioLeftSide)
+			DECLARE_METASOUND_PARAM_EXTERN(AudioRightSide)
 			DECLARE_METASOUND_PARAM_EXTERN(Timestamp);
 			DECLARE_METASOUND_PARAM_EXTERN(ShouldQuantizeTimestamp);
 			DECLARE_METASOUND_PARAM_EXTERN(QuantizationUnit);
diff --git a/Engine/Plugins/Runtime/Harmonix/Source/HarmonixMetasound/Public/HarmonixMetasound/Nodes/TapeSSNode.h b/Engine/Plugins/Runtime/Harmonix/Source/HarmonixMetasound/Public/HarmonixMetasound/Nodes/TapeSSNode.h
new file mode 100644
--- /dev/null
+++ b/Engine/Plugins/Runtime/Harmonix/Source/HarmonixMetasound/Public/HarmonixMetasound/Nodes/TapeSSNode.h
@@ -0,0 +1,61 @@
+// Copyright Epic Games, Inc. All Rights Reserved.
+
+#pragma once
+
+#include "HarmonixMetasound/Common.h"
+#include "MetasoundFrontendDocument.h"
+#include "MetasoundOperatorData.h"
+
+#include "TapeSSNode.generated.h"
+
+namespace HarmonixMetasound::TapeSSNode
+{
+	class FTapeSSNodeOperatorData;
+}
+
+USTRUCT()
+struct FHarmonixTapeSSNodeConfiguration : public FMetaSoundFrontendNodeConfiguration
+{
+	GENERATED_BODY()
+
+	UPROPERTY(EditAnywhere, Category = General, meta = (ClampMin = "1", ClampMax = "8"))
+	uint32 NumChannels = 1;
+
+	UPROPERTY(EditAnywhere, Category = General, meta = (ClampMin = "1", ClampMax = "10"))
+	float MaxRampSeconds = 5.0f;
+
+	/* Get the current interface for the class based upon the node extension */
+	virtual TInstancedStruct<FMetasoundFrontendClassInterface> OverrideDefaultInterface(const FMetasoundFrontendClass& InNodeClass) const;
+
+	TSharedPtr<const Metasound::IOperatorData> GetOperatorData() const override;
+};
+
+namespace HarmonixMetasound::TapeSSNode
+{
+	class FTapeSSNodeOperatorData : public Metasound::TOperatorData<FTapeSSNodeOperatorData>
+	{
+	public:
+		static const FLazyName OperatorDataTypeName;
+		FTapeSSNodeOperatorData(const float& InMaxRampSeconds)
+			: MaxRampSeconds(InMaxRampSeconds)
+		{}
+		float MaxRampSeconds = 5.0f;
+	};
+
+	namespace Inputs
+	{
+		DECLARE_METASOUND_PARAM_ALIAS(Enabled);
+		DECLARE_METASOUND_PARAM_EXTERN(Stop);
+		DECLARE_METASOUND_PARAM_EXTERN(RampDuration);
+	}
+
+	namespace Outputs
+	{
+		DECLARE_METASOUND_PARAM_EXTERN(InputAudioInUse);
+		DECLARE_METASOUND_PARAM_EXTERN(OnStopped);
+		DECLARE_METASOUND_PARAM_EXTERN(OnFullSpeed);
+		DECLARE_METASOUND_PARAM_EXTERN(CurrentSpeed);
+	}
+
+	static constexpr int32 kMaxChannels = 8;
+}
