Fix Caps in includes

diff --git a/Engine/Plugins/Runtime/Harmonix/Source/HarmonixDsp/Private/HarmonixDsp/Effects/TapeSS.cpp b/Engine/Plugins/Runtime/Harmonix/Source/HarmonixDsp/Private/HarmonixDsp/Effects/TapeSS.cpp
--- a/Engine/Plugins/Runtime/Harmonix/Source/HarmonixDsp/Private/HarmonixDsp/Effects/TapeSS.cpp
+++ b/Engine/Plugins/Runtime/Harmonix/Source/HarmonixDsp/Private/HarmonixDsp/Effects/TapeSS.cpp
@@ -1,6 +1,6 @@
 // Copyright Epic Games, Inc. All Rights Reserved.
-#include "HarmonixDSP/Effects/TapeSS.h"
-#include "HarmonixDSP/AudioBuffer.h"
+#include "HarmonixDsp/Effects/TapeSS.h"
+#include "HarmonixDsp/AudioBuffer.h"
 
 namespace Harmonix::Dsp::Effects
 {
diff --git a/Engine/Plugins/Runtime/Harmonix/Source/HarmonixDsp/Public/HarmonixDsp/Effects/TapeSS.h b/Engine/Plugins/Runtime/Harmonix/Source/HarmonixDsp/Public/HarmonixDsp/Effects/TapeSS.h
--- a/Engine/Plugins/Runtime/Harmonix/Source/HarmonixDsp/Public/HarmonixDsp/Effects/TapeSS.h
+++ b/Engine/Plugins/Runtime/Harmonix/Source/HarmonixDsp/Public/HarmonixDsp/Effects/TapeSS.h
@@ -1,8 +1,8 @@
 // Copyright Epic Games, Inc. All Rights Reserved.
 #pragma once
 
-#include "HarmonixDSP/AudioBuffer.h"
-#include "HarmonixDSP/Ramper.h"
+#include "HarmonixDsp/AudioBuffer.h"
+#include "HarmonixDsp/Ramper.h"
 
 namespace Harmonix::Dsp::Effects
 {
