Fix issue with Metasound Unit Tests and the Reset function in the metasound node.

diff --git a/Engine/Plugins/Runtime/Harmonix/Source/HarmonixMetasound/Private/HarmonixMetasound/Nodes/TapeSSNode.cpp b/Engine/Plugins/Runtime/Harmonix/Source/HarmonixMetasound/Private/HarmonixMetasound/Nodes/TapeSSNode.cpp
--- a/Engine/Plugins/Runtime/Harmonix/Source/HarmonixMetasound/Private/HarmonixMetasound/Nodes/TapeSSNode.cpp
+++ b/Engine/Plugins/Runtime/Harmonix/Source/HarmonixMetasound/Private/HarmonixMetasound/Nodes/TapeSSNode.cpp
@@ -146,7 +146,8 @@ namespace HarmonixMetasound::TapeSSNode
 			, Outputs(MoveTemp(Outputs))
 		{
 			TSharedPtr<const FTapeSSNodeOperatorData> InOperatorData = StaticCastSharedPtr<const FTapeSSNodeOperatorData>(Params.Node.GetOperatorData());
-			Reset(Params, InOperatorData ? InOperatorData->MaxRampSeconds : 5.0f);
+			MaxRampSeconds = InOperatorData ? InOperatorData->MaxRampSeconds : 5.0f;
+			Reset(Params);
 		}
 		
 		static const FNodeClassMetadata& GetNodeInfo()
@@ -230,9 +231,8 @@ namespace HarmonixMetasound::TapeSSNode
 			}
 		}
 		
-		void Reset(const FResetParams& Params, float InMaxRampSeconds)
+		void Reset(const FResetParams& Params)
 		{
-			MaxRampSeconds = InMaxRampSeconds;
 			TapeSS.Prepare(Params.OperatorSettings.GetSampleRate(), Inputs.AudioInputs.Num(), MaxRampSeconds);
 			*Outputs.InputAudioInUse = true;
 			ChannelInputPointers.SetNumUninitialized(Inputs.AudioInputs.Num());
