Reset all outputs on TapeSS node

#rnx
#tests Metasound Automation Tests

diff --git a/Engine/Plugins/Runtime/Harmonix/Source/HarmonixMetasound/Private/HarmonixMetasound/Nodes/TapeSSNode.cpp b/Engine/Plugins/Runtime/Harmonix/Source/HarmonixMetasound/Private/HarmonixMetasound/Nodes/TapeSSNode.cpp
--- a/Engine/Plugins/Runtime/Harmonix/Source/HarmonixMetasound/Private/HarmonixMetasound/Nodes/TapeSSNode.cpp
+++ b/Engine/Plugins/Runtime/Harmonix/Source/HarmonixMetasound/Private/HarmonixMetasound/Nodes/TapeSSNode.cpp
@@ -235,6 +235,9 @@ namespace HarmonixMetasound::TapeSSNode
 		{
 			TapeSS.Prepare(Params.OperatorSettings.GetSampleRate(), Inputs.AudioInputs.Num(), MaxRampSeconds);
 			*Outputs.InputAudioInUse = true;
+			Outputs.OnStopped->Reset();
+			Outputs.OnFullSpeed->Reset();
+			*Outputs.CurrentSpeed = 0.0f;
 			ChannelInputPointers.SetNumUninitialized(Inputs.AudioInputs.Num());
 			ChannelOutputPointers.SetNumUninitialized(Inputs.AudioInputs.Num());
 		}
