Stutter Edit Effect for Metasound
- New Simple Sampler Node
- New Stutter Sequencer Node
- New Midi Stutter Sequence data type

A complete reworking of how the old Harmonix Stutter Edit worked.

diff --git a/Engine/Plugins/Runtime/Harmonix/Source/HarmonixMetasound/Private/HarmonixMetasound/DataTypes/MidiStutterSequence.cpp b/Engine/Plugins/Runtime/Harmonix/Source/HarmonixMetasound/Private/HarmonixMetasound/DataTypes/MidiStutterSequence.cpp
new file mode 100644
--- /dev/null
+++ b/Engine/Plugins/Runtime/Harmonix/Source/HarmonixMetasound/Private/HarmonixMetasound/DataTypes/MidiStutterSequence.cpp
@@ -0,0 +1,67 @@
+// Copyright Epic Games, Inc. All Rights Reserved.
+
+
+#include "HarmonixMetasound/DataTypes/MidiStutterSequence.h"
+
+#include "Harmonix/PropertyUtility.h"
+
+DEFINE_LOG_CATEGORY_STATIC(LogMIDIStutterSequence, Log, All);
+
+DEFINE_AUDIORENDERABLE_ASSET(HarmonixMetasound, FStutterSequenceAsset, StutterSequenceAsset, UMidiStutterSequence)
+
+#if WITH_EDITOR
+void UMidiStutterSequence::PostEditChangeChainProperty(FPropertyChangedChainEvent& PropertyChangedChainEvent)
+{
+	Super::PostEditChangeChainProperty(PropertyChangedChainEvent);
+
+	const FProperty* PropertyChanged = PropertyChangedChainEvent.Property;
+	const EPropertyChangeType::Type PropertyChangeType = PropertyChangedChainEvent.ChangeType;
+
+	// check against if changing the underlying FStutterSequenceTable struct
+	const FProperty* MemberChanged = PropertyChangedChainEvent.PropertyChain.GetHead()->GetValue();
+	if (MemberChanged->GetFName() != GET_MEMBER_NAME_CHECKED(UMidiStutterSequence, StutterTable))
+	{
+		return;
+	}
+
+	// there is no renderable data, so no need to do any copying over on property changes
+	if (!RenderableStutterTable)
+	{
+		return;
+	}
+
+	// Determine what to do based on the property and the change type
+	Harmonix::EPostEditAction PostEditAction = Harmonix::GetPropertyPostEditAction(PropertyChanged, PropertyChangeType);
+
+	if (PostEditAction == Harmonix::EPostEditAction::UpdateTrivial)
+	{
+		// use the PropertyChangedChainEvent to copy the single property changed
+		FStutterSequenceTable* CopyToStruct = *RenderableStutterTable;
+		Harmonix::CopyStructProperty(CopyToStruct, &StutterTable, PropertyChangedChainEvent);
+	}
+	else if (PostEditAction == Harmonix::EPostEditAction::UpdateNonTrivial)
+	{
+		UpdateRenderableForNonTrivialChange();
+	}
+	
+}
+#endif // WITH_EDITOR
+
+TSharedPtr<Audio::IProxyData> UMidiStutterSequence::CreateProxyData(const Audio::FProxyDataInitParams& InitParams)
+{
+	if (!RenderableStutterTable)
+	{
+		RenderableStutterTable = MakeShared<FStutterSequenceTableProxy::QueueType>(StutterTable);
+	}
+	return MakeShared<FStutterSequenceTableProxy>(RenderableStutterTable);
+}
+
+void UMidiStutterSequence::UpdateRenderableForNonTrivialChange()
+{
+	// If no one has requested a renderable proxy we don't have to do anything
+	if (!RenderableStutterTable)
+	{
+		return;
+	}
+	RenderableStutterTable->SetNewSettings(StutterTable);
+}
diff --git a/Engine/Plugins/Runtime/Harmonix/Source/HarmonixMetasound/Private/HarmonixMetasound/Nodes/SimpleSamplerNode.cpp b/Engine/Plugins/Runtime/Harmonix/Source/HarmonixMetasound/Private/HarmonixMetasound/Nodes/SimpleSamplerNode.cpp
new file mode 100644
--- /dev/null
+++ b/Engine/Plugins/Runtime/Harmonix/Source/HarmonixMetasound/Private/HarmonixMetasound/Nodes/SimpleSamplerNode.cpp
@@ -0,0 +1,680 @@
+// Copyright Epic Games, Inc. All Rights Reserved.
+
+#include "HarmonixMetasound/Nodes/SimpleSamplerNode.h"
+
+#include "MetasoundExecutableOperator.h"
+#include "MetasoundFacade.h"
+#include "MetasoundNodeInterface.h"
+#include "MetasoundParamHelper.h"
+#include "MetasoundSampleCounter.h"
+#include "MetasoundStandardNodesCategories.h"
+#include "MetasoundVertex.h"
+#include "HarmonixMetasound/Common.h"
+#include "HarmonixDsp/AudioBuffer.h"
+
+#define LOCTEXT_NAMESPACE "HarmonixMetaSound"
+
+namespace HarmonixMetasound::SimpleSamplerNode
+{
+	using namespace Metasound;
+
+	namespace Inputs
+	{
+		DEFINE_METASOUND_PARAM_ALIAS(Enabled, CommonPinNames::Inputs::Enable);
+		DEFINE_INPUT_METASOUND_PARAM(Capture, "Capture", "When triggered, starts capturing the input audio to a buffer.");
+		DEFINE_INPUT_METASOUND_PARAM(CaptureDuration, "Capture Duration Sec", "How long a sample to capture.");
+		DEFINE_INPUT_METASOUND_PARAM(Play, "Play", "Triggers playback of the captured buffer.");
+		DEFINE_INPUT_METASOUND_PARAM(PlayDuration, "Play Duration Sec", "How much of the captured sample to play.");
+		DEFINE_INPUT_METASOUND_PARAM(Reverse, "Reverse Playback", "If true the captured buffer is played in reverse.");
+		DEFINE_INPUT_METASOUND_PARAM(Reset, "Reset", "Stop capture/playback and start passing audio from input to output.");
+		static constexpr int32 kNumNonAudioInputs = 7; // NOTE: Keep in sync with the number of input vertecies defined above!
+	}
+
+	namespace Outputs
+	{
+		DEFINE_OUTPUT_METASOUND_PARAM(InputAudioInUse, "Input Audio In Use", "True when audio from the input is being used to generate the output. False when the input audio is unneeded and can be changed.");
+	}
+
+	const FLazyName FSimpleSamplerNodeOperatorData::OperatorDataTypeName = "SimpleSamplerOperatorData";
+
+	static TArray<FVertexName> CreateVertexNames(const TCHAR* Prefix, int32 NumChannels)
+	{
+		TStringBuilder<32> NameStr;
+		TArray<FVertexName> Names;
+		Names.AddUninitialized(NumChannels);
+
+		for (int ChannelIndex = 0; ChannelIndex < NumChannels; ++ChannelIndex)
+		{
+			NameStr << Prefix << ChannelIndex;
+			Names[ChannelIndex] = *NameStr;
+			NameStr.Reset();
+		}
+
+		return Names;
+	}
+
+	static const TArray<FVertexName> AudioInputNames = CreateVertexNames(TEXT("In "), kMaxChannels);
+	static const TArray<FVertexName> AudioOutputNames = CreateVertexNames(TEXT("Out "), kMaxChannels);
+
+#if WITH_EDITOR
+	static const FText GetAudioInputDisplayName(uint32 ChannelIndex)
+	{
+		return METASOUND_LOCTEXT_FORMAT("AudioInputDisplayName", "In {0}", ChannelIndex);
+	}
+
+	static const FText GetAudioOutputDisplayName(uint32 ChannelIndex)
+	{
+		return METASOUND_LOCTEXT_FORMAT("AudioOutputDisplayName", "Out {0}", ChannelIndex);
+	}
+#endif // WITH_EDITOR
+
+	FVertexInterface GetVertexInterface(int32 NumChannels)
+	{
+		using namespace HarmonixMetasound::SimpleSamplerNode;
+
+		FInputVertexInterface InputInterface;
+		FOutputVertexInterface OutputInterface;
+
+		InputInterface.Add(TInputDataVertex<bool>(METASOUND_GET_PARAM_NAME_AND_METADATA(Inputs::Enabled)));
+		InputInterface.Add(TInputDataVertex<FTrigger>(METASOUND_GET_PARAM_NAME_AND_METADATA(Inputs::Capture)));
+		InputInterface.Add(TInputDataVertex<float>(METASOUND_GET_PARAM_NAME_AND_METADATA(Inputs::CaptureDuration)));
+		InputInterface.Add(TInputDataVertex<FTrigger>(METASOUND_GET_PARAM_NAME_AND_METADATA(Inputs::Play)));
+		InputInterface.Add(TInputDataVertex<float>(METASOUND_GET_PARAM_NAME_AND_METADATA(Inputs::PlayDuration)));
+		InputInterface.Add(TInputDataVertex<bool>(METASOUND_GET_PARAM_NAME_AND_METADATA(Inputs::Reverse)));
+		InputInterface.Add(TInputDataVertex<FTrigger>(METASOUND_GET_PARAM_NAME_AND_METADATA(Inputs::Reset)));
+
+		OutputInterface.Add(TOutputDataVertex<bool>(METASOUND_GET_PARAM_NAME_AND_METADATA(Outputs::InputAudioInUse)));
+
+		check (InputInterface.Num() == Inputs::kNumNonAudioInputs);
+		check (NumChannels <= kMaxChannels);
+
+		TStringBuilder<32> NameStr;
+		for (int32 ChannelIndex = 0; ChannelIndex < NumChannels; ++ChannelIndex)
+		{
+#if WITH_EDITOR
+			const FDataVertexMetadata AudioInputMetadata
+			{
+				METASOUND_LOCTEXT("AudioInputDescription", "Audio Input"),
+				GetAudioInputDisplayName(ChannelIndex)
+			};
+#else 
+			const FDataVertexMetadata AudioInputMetadata;
+#endif // WITH_EDITOR
+			InputInterface.Add(TInputDataVertex<FAudioBuffer>(AudioInputNames[ChannelIndex], AudioInputMetadata));
+
+#if WITH_EDITOR
+			const FDataVertexMetadata AudioOutputMetadata
+			{
+				METASOUND_LOCTEXT("AudioOutputDescription", "Audio Output"),
+				GetAudioOutputDisplayName(ChannelIndex)
+			};
+#else 
+			const FDataVertexMetadata AudioOutputMetadata;
+#endif // WITH_EDITOR
+			OutputInterface.Add(TOutputDataVertex<FAudioBuffer>(AudioOutputNames[ChannelIndex], AudioOutputMetadata));
+		}
+
+		return FVertexInterface
+		{
+			MoveTemp(InputInterface),
+			MoveTemp(OutputInterface)
+		};
+	}
+
+	class FSimpleSamplerOperator final : public TExecutableOperator<FSimpleSamplerOperator>
+	{
+	public:
+		struct FInputs
+		{
+			FBoolReadRef        Enabled;
+			FTriggerReadRef     Capture;
+			FFloatReadRef       CaptureDuration;
+			FTriggerReadRef     Play;
+			FFloatReadRef       PlayDuration;
+			FBoolReadRef        Reverse;
+			FTriggerReadRef     Reset;
+			TArray<FAudioBufferReadRef> AudioInputs;
+		};
+
+		struct FOutputs
+		{
+			FBoolWriteRef        InputAudioInUse;
+			TArray<FAudioBufferWriteRef> AudioOutputs;
+		};
+		
+		FSimpleSamplerOperator(const FBuildOperatorParams& Params, FInputs&& Inputs, FOutputs&& Outputs)
+			: Inputs(MoveTemp(Inputs))
+			, Outputs(MoveTemp(Outputs))
+		{
+			TSharedPtr<const FSimpleSamplerNodeOperatorData> InOperatorData = StaticCastSharedPtr<const FSimpleSamplerNodeOperatorData>(Params.Node.GetOperatorData());
+			MaxCaptureSeconds = InOperatorData ? InOperatorData->MaxCaptureSeconds : 5.0f;
+			Reset(Params);
+		}
+		
+		static const FNodeClassMetadata& GetNodeInfo()
+		{
+			auto InitNodeInfo = []() -> FNodeClassMetadata
+			{
+				FNodeClassMetadata Info = FNodeClassMetadata::GetEmpty();
+				Info.ClassName = { HarmonixNodeNamespace, TEXT("SimpleSampler"), TEXT("") };
+				Info.MajorVersion = 0;
+				Info.MinorVersion = 1;
+				Info.DisplayName = METASOUND_LOCTEXT("SimpleSampler_DisplayName", "Simple Sampler");
+				Info.CategoryHierarchy = { MetasoundNodeCategories::Harmonix };
+				Info.Description = METASOUND_LOCTEXT("SimpleSampler_Description", "Allows you to capture a buffer of incoming audio and then trigger its playback.");
+				Info.Author = PluginAuthor;
+				Info.PromptIfMissing = PluginNodeMissingPrompt;
+				Info.DefaultInterface = GetVertexInterface(1);
+				return Info;
+			};
+
+			static const FNodeClassMetadata Info = InitNodeInfo();
+
+			return Info;
+		}
+
+		static TUniquePtr<IOperator> CreateOperator(const FBuildOperatorParams& InParams, FBuildResults& OutResults)
+		{
+			const FVertexInterface& NodeInterface = InParams.Node.GetVertexInterface();
+			const FInputVertexInterfaceData& InputData = InParams.InputData;
+
+			TArray<FAudioBufferReadRef> AudioInputs;
+			TArray<FAudioBufferWriteRef> AudioOutputs;
+			int32 NumAudioChannels = NodeInterface.GetInputInterface().Num() - Inputs::kNumNonAudioInputs;
+			for (int32 ChannelIndex = 0; ChannelIndex < NumAudioChannels; ++ChannelIndex)
+			{
+				AudioInputs.Emplace(InputData.GetOrCreateDefaultDataReadReference<FAudioBuffer>(AudioInputNames[ChannelIndex], InParams.OperatorSettings));
+				AudioOutputs.Emplace(TDataWriteReferenceFactory<FAudioBuffer>::CreateExplicitArgs(InParams.OperatorSettings));
+			}
+
+			FInputs Inputs {
+					InputData.GetOrCreateDefaultDataReadReference<bool>(Inputs::EnabledName, InParams.OperatorSettings),
+					InputData.GetOrCreateDefaultDataReadReference<FTrigger>(Inputs::CaptureName, InParams.OperatorSettings),
+					InputData.GetOrCreateDefaultDataReadReference<float>(Inputs::CaptureDurationName, InParams.OperatorSettings),
+					InputData.GetOrCreateDefaultDataReadReference<FTrigger>(Inputs::PlayName, InParams.OperatorSettings),
+					InputData.GetOrCreateDefaultDataReadReference<float>(Inputs::PlayDurationName, InParams.OperatorSettings),
+					InputData.GetOrCreateDefaultDataReadReference<bool>(Inputs::ReverseName, InParams.OperatorSettings),
+					InputData.GetOrCreateDefaultDataReadReference<FTrigger>(Inputs::ResetName, InParams.OperatorSettings),
+					MoveTemp(AudioInputs)
+			};
+
+			FOutputs Outputs {	
+					FBoolWriteRef::CreateNew(),
+					MoveTemp(AudioOutputs)
+			};
+
+			return MakeUnique<FSimpleSamplerOperator>(
+				InParams,
+				MoveTemp(Inputs),
+				MoveTemp(Outputs)
+			);
+		}
+
+		virtual void BindInputs(FInputVertexInterfaceData& InVertexData) override
+		{
+			InVertexData.BindReadVertex(Inputs::EnabledName, Inputs.Enabled);
+			InVertexData.BindReadVertex(Inputs::CaptureName, Inputs.Capture);
+			InVertexData.BindReadVertex(Inputs::CaptureDurationName, Inputs.CaptureDuration);
+			InVertexData.BindReadVertex(Inputs::PlayName, Inputs.Play);
+			InVertexData.BindReadVertex(Inputs::PlayDurationName, Inputs.PlayDuration);
+			InVertexData.BindReadVertex(Inputs::ReverseName, Inputs.Reverse);
+			InVertexData.BindReadVertex(Inputs::ResetName, Inputs.Reset);
+			for (int32 ChannelIndex = 0; ChannelIndex < Inputs.AudioInputs.Num(); ++ChannelIndex)
+			{
+				InVertexData.BindReadVertex(AudioInputNames[ChannelIndex], Inputs.AudioInputs[ChannelIndex]);
+			}
+		}
+
+		virtual void BindOutputs(FOutputVertexInterfaceData& InVertexData) override
+		{
+			InVertexData.BindReadVertex(Outputs::InputAudioInUseName, Outputs.InputAudioInUse);
+			for (int32 ChannelIndex = 0; ChannelIndex < Outputs.AudioOutputs.Num(); ++ChannelIndex)
+			{
+				InVertexData.BindReadVertex(AudioOutputNames[ChannelIndex], Outputs.AudioOutputs[ChannelIndex]);
+			}
+		}
+		
+		void Reset(const FResetParams& Params)
+		{
+			int32 NumFrames = FMath::CeilToInt32(Params.OperatorSettings.GetSampleRate() * MaxCaptureSeconds);
+			CaptureBuffer.Configure(Inputs.AudioInputs.Num(), NumFrames, EAudioBufferCleanupMode::Delete, Params.OperatorSettings.GetSampleRate(), false);
+			CaptureBuffer.SetNumValidFrames(0);
+			CaptureBuffer.SetNumValidChannels(Inputs.AudioInputs.Num());
+			bCurrentlyPlaying = bCurrentlyCapturing = bFullyCaptured = false;
+			CaptureLengthFrames = PlayLengthFrames = NextPlayFrame = 0;
+			*Outputs.InputAudioInUse = false;
+		}
+
+		enum class EInputOrSilence
+		{
+			Input,
+			Silence
+		};
+
+		void BulkOutput(EInputOrSilence Mode, int32 StartFrameIndexInBlock = 0, int32 NumFrames = -1)
+		{
+			if (NumFrames == -1)
+			{
+				NumFrames = Outputs.AudioOutputs[0]->Num() - StartFrameIndexInBlock;
+			}
+			check (StartFrameIndexInBlock + NumFrames <= Outputs.AudioOutputs[0]->Num());
+
+			if (NumFrames <= 0)
+			{
+				return;
+			}
+
+			if (Mode == EInputOrSilence::Input)
+			{
+				*Outputs.InputAudioInUse = true;
+				for (int32 ChannelIndex = 0; ChannelIndex < Outputs.AudioOutputs.Num(); ++ChannelIndex)
+				{
+					FMemory::Memcpy(Outputs.AudioOutputs[ChannelIndex]->GetData() + StartFrameIndexInBlock, Inputs.AudioInputs[ChannelIndex]->GetData() + StartFrameIndexInBlock, sizeof(float) * NumFrames);
+				}
+			}
+			else
+			{
+				for (int32 ChannelIndex = 0; ChannelIndex < Outputs.AudioOutputs.Num(); ++ChannelIndex)
+				{
+					FMemory::Memset(Outputs.AudioOutputs[ChannelIndex]->GetData() + StartFrameIndexInBlock, 0, sizeof(float) * NumFrames);
+				}
+			}
+
+		}
+
+		void Execute()
+		{
+			bool Holding = Inputs.Capture->IsTriggeredInBlock() || CaptureBuffer.GetNumValidFrames() > 0;
+			if (!*Inputs.Enabled || !Holding)
+			{
+				BulkOutput(EInputOrSilence::Input);
+				return;
+			}
+
+			CaptureLengthFrames = FMath::Min(FMath::CeilToInt32(*Inputs.CaptureDuration * CaptureBuffer.GetMaxConfig().GetSampleRate()), CaptureBuffer.GetMaxNumFrames());
+			if (CaptureLengthFrames <= CaptureBuffer.GetNumValidFrames())
+			{
+				// crop our capture buffer...
+				CaptureBuffer.SetNumValidFrames(CaptureLengthFrames);
+				// we may still be capturing... so stop...
+				bCurrentlyCapturing = false;
+				bFullyCaptured = true;
+				*Outputs.InputAudioInUse = false;
+			}
+
+			bool bRenderedOutput = false;
+
+			if (Inputs.Capture->IsTriggeredInBlock() || bCurrentlyCapturing)
+			{
+				bFullyCaptured = false;
+
+				// We pass in a where/when a play is ALSO triggered so the capture code knows
+				// whether to pass input -> output or let the play call do it...
+				bRenderedOutput = DoCapture(Inputs.Play->Last());
+				
+				// bFullyCaptured could have changed inside of DoCapture... so...
+				*Outputs.InputAudioInUse = !bFullyCaptured; 
+			}
+
+			if (Inputs.Play->IsTriggeredInBlock() || bCurrentlyPlaying)
+			{
+				bHavePlayedThisCapture = true;
+
+				if (Inputs.Play->IsTriggeredInBlock()) // a new trigger in this block...
+				{
+					int32 NewPlayPosition = Inputs.Play->Last();
+
+					// Are we still playing a previous trigger?
+					if (bCurrentlyPlaying)
+					{
+						NewPlayPosition = ForceFinishCurrentPlaybackByBlockFrame(NewPlayPosition);
+						bRenderedOutput = true;
+					}
+					else
+					{
+						// No... we MAY want to write silence before this new trigger starts...
+						if (bFullyCaptured && NewPlayPosition > 0)
+						{
+							BulkOutput(EInputOrSilence::Silence, 0, NewPlayPosition);
+						}
+					}
+
+					// Since we are starting... capture any new length at our input...
+					PlayLengthFrames = FMath::Min(FMath::CeilToInt32(*Inputs.PlayDuration * CaptureBuffer.GetMaxConfig().GetSampleRate()), CaptureLengthFrames);
+					// ... and whether we should be playing backward...
+					if (*Inputs.Reverse)
+					{
+						// We can only play in reverse if we have captured enough samples!
+						if (CaptureBuffer.GetNumValidFrames() >= PlayLengthFrames)
+						{
+							bPlayingReverse = true;
+						}
+						else
+						{
+							// We'll take accept anything within a 2 milliseconds...
+							int32 ErrorFrames = PlayLengthFrames - CaptureBuffer.GetNumValidFrames();
+							bPlayingReverse = ErrorFrames < (CaptureBuffer.GetMaxConfig().GetSampleRate() * 2 / 1000);
+							if (bPlayingReverse)
+							{
+								PlayLengthFrames = CaptureBuffer.GetNumValidFrames();
+							}
+						}
+					}
+					else
+					{
+						bPlayingReverse = false;
+					}
+
+					// We will be starting (or restarting) so set up some things...
+					bCurrentlyPlaying = false;
+					NextPlayFrame = 0;
+					int32 NumSamplesToRenderThisBlock = Outputs.AudioOutputs[0]->Num() - NewPlayPosition;
+
+					// Now render...
+					bRenderedOutput = DoPlay(NewPlayPosition, NumSamplesToRenderThisBlock) || bRenderedOutput ;
+				}
+				else // we must be currently playing... so continue...
+				{
+					bRenderedOutput = DoPlay(0, Outputs.AudioOutputs[0]->Num());
+				}
+			}
+
+			// Reset this since we are done capturing & playing this block...
+			bPlayFadeInNotNeededThisBlock = false;
+
+			// If we didn't render anything, and we are fully captured
+			if (!bRenderedOutput && (bFullyCaptured || bHavePlayedThisCapture))
+			{
+				for (FAudioBufferWriteRef& Out : Outputs.AudioOutputs)
+				{
+					Out->Zero();
+				}
+			}
+
+			if (Inputs.Reset->IsTriggeredInBlock())
+			{
+				ResetCaptureAndPlayback(Inputs.Reset->First());
+			}
+		}
+
+		/**
+		Blows out the capture buffer, stops current playback, and starts piping
+		input -> output.
+		@param FromFrame The position of the reset within the block.
+		*/
+		void ResetCaptureAndPlayback(int32 FromFrame)
+		{
+			bCurrentlyPlaying = bCurrentlyCapturing = bFullyCaptured = bHavePlayedThisCapture = false;
+			NextPlayFrame = 0;
+			CaptureBuffer.SetNumValidFrames(0);
+			// Make sure we have at least enough frames to fade in...
+			if (Outputs.AudioOutputs[0]->Num() - FromFrame < kNumFadeFrames)
+			{
+				FromFrame = Outputs.AudioOutputs[0]->Num() - kNumFadeFrames;
+			}
+			// Now write un-faded frames...
+			BulkOutput(EInputOrSilence::Input, FromFrame);
+			// Now fade...
+			float Fade = 0.0f;
+			for (int32 FadeFrame = 0; FadeFrame < kNumFadeFrames; ++FadeFrame)
+			{
+				Fade = float(FadeFrame) / float(kNumFadeFrames);
+				for (int32 ChannelIndex = 0; ChannelIndex < Outputs.AudioOutputs.Num(); ++ChannelIndex)
+				{
+					*(Outputs.AudioOutputs[ChannelIndex]->GetData() + FromFrame + FadeFrame) *= Fade;
+				}
+			}
+		}
+
+		/**
+		Ends the current playback as gracefully as possible.
+		@param CompleteByBlockFrame When the playback needs to be finished.
+		@return The frame we actually completed by, so that the new play
+		        request can be moved a few samples to avoid a click.
+		*/
+		int32 ForceFinishCurrentPlaybackByBlockFrame(int32 CompleteByBlockFrame)
+		{
+			int32 FramesRemainingToBePlayed = PlayLengthFrames - NextPlayFrame;
+			if (FramesRemainingToBePlayed <= CompleteByBlockFrame)
+			{
+				// We have enough space to finish rendering the last play INCLUDING its fade...
+				DoPlay(0, CompleteByBlockFrame, false);
+				// No need to modify the new play frame.
+			}
+			else
+			{
+				// We don't have enough space to finish rendering the last play...
+				// If it has already started fading, let it finish and then fudge the start sample
+				// of the next play...
+				if (FramesRemainingToBePlayed <= kNumFadeFrames)
+				{
+					DoPlay(0, CompleteByBlockFrame, false);
+					CompleteByBlockFrame = FramesRemainingToBePlayed;
+				}
+				else
+				{
+					// It hasn't already started to fade, so let's force it to fade (truncate it) and
+					// fudge the start sample of the next play...
+					DoPlay(0, kNumFadeFrames, true);
+					CompleteByBlockFrame = kNumFadeFrames;
+				}
+			}
+			return CompleteByBlockFrame;
+		}
+
+		/**
+		Plays the captured audio. Will do a little fade in at the beginning, and a little 
+		fade out at the end.
+		@param DestinationFrame Where in the output buffer to write samples.
+		@param NumFrames Number of frames to write.
+		@param bForceFade Whether to force the fade out of playback even
+		                  if we are not at the end of the available sample data.
+						  This typically only happens if we are being asked to start
+						  playback again while we are already playing.
+		@return True if this function ended up writing ANYTHING to the output buffer.
+		*/
+		bool DoPlay(int32 DestinationFrame, int32 NumFrames, bool bForceFade = false)
+		{
+			int32 FramesWriteableInOutput = NumFrames;
+			int32 FramesAvailableToPlay = PlayLengthFrames - NextPlayFrame;
+			int32 FramesToPlay = FMath::Min(FramesWriteableInOutput, FramesAvailableToPlay);
+
+			if (FramesToPlay)
+			{
+				bCurrentlyPlaying = true;
+				if (!bPlayingReverse)
+				{
+					for (int32 ChannelIndex = 0; ChannelIndex < CaptureBuffer.GetNumValidChannels(); ++ChannelIndex)
+					{
+						FMemory::Memcpy(Outputs.AudioOutputs[ChannelIndex]->GetData() + DestinationFrame, CaptureBuffer.GetValidChannelData(ChannelIndex) + NextPlayFrame, FramesToPlay * sizeof(float));
+					}
+				}
+				else
+				{
+					for (int32 ChannelIndex = 0; ChannelIndex < CaptureBuffer.GetNumValidChannels(); ++ChannelIndex)
+					{
+						float* CapturedData = CaptureBuffer.GetValidChannelData(ChannelIndex);
+						float* DestinationAddress = Outputs.AudioOutputs[ChannelIndex]->GetData() + DestinationFrame;
+						for (int32 SourceIndex = PlayLengthFrames - NextPlayFrame - 1, FrameCount = 0; FrameCount < FramesToPlay; ++FrameCount, --SourceIndex)
+						{
+							*(DestinationAddress + FrameCount) = CapturedData[SourceIndex];
+						}
+					}
+				}
+				// fade in the beginning to avoid click...
+				if (NextPlayFrame < kNumFadeFrames && !bPlayFadeInNotNeededThisBlock)
+				{
+					int32 FadeFrames = FMath::Min(kNumFadeFrames, FramesToPlay);
+					float CurrentFade = float(NextPlayFrame) / float(kNumFadeFrames);
+					static float constexpr FadeInc = 1.0f/float(kNumFadeFrames);
+					for (int32 FrameIndex = DestinationFrame; FrameIndex < DestinationFrame + FadeFrames; ++FrameIndex)
+					{
+						for (int32 ChannelIndex = 0; ChannelIndex < CaptureBuffer.GetNumValidChannels(); ++ChannelIndex)
+						{
+							Outputs.AudioOutputs[ChannelIndex]->GetData()[FrameIndex] *= CurrentFade;
+						}
+						CurrentFade += FadeInc;
+					}
+				}
+				// fade out the end to avoid click...
+				if ((NextPlayFrame + FramesToPlay) > (PlayLengthFrames - kNumFadeFrames) || bForceFade)
+				{
+					int32 RenderedLastFrame = NextPlayFrame + FramesToPlay;
+					int32 FramesRemainingToRender = bForceFade ? 0 : PlayLengthFrames - RenderedLastFrame;
+					int32 NumFadeFrames = bForceFade ? FMath::Min(FramesToPlay, kNumFadeFrames) : kNumFadeFrames;
+					int32 MaxFramesToFade = NumFadeFrames - FramesRemainingToRender;
+					float CurrentFade = FMath::Min(1.0f,  float(FramesRemainingToRender) / float(NumFadeFrames));
+					float FadeInc = 1.0f / float(NumFadeFrames);
+					int32 FadeAtIndex = DestinationFrame + FramesToPlay - 1;
+					for (int32 FadeIndex = 0; FadeIndex < MaxFramesToFade; ++FadeIndex)
+					{
+						for (int32 ChannelIndex = 0; ChannelIndex < CaptureBuffer.GetNumValidChannels(); ++ChannelIndex)
+						{
+							Outputs.AudioOutputs[ChannelIndex]->GetData()[FadeAtIndex] *= CurrentFade;
+						}
+						--FadeAtIndex;
+						CurrentFade = FMath::Clamp(CurrentFade + FadeInc, 0.0f, 1.0f);
+					}
+				}
+				NextPlayFrame += FramesToPlay;
+				bCurrentlyPlaying = NextPlayFrame < PlayLengthFrames;
+			}
+			else
+			{
+				NextPlayFrame = 0;
+				bCurrentlyPlaying = false;
+			}
+
+			// we may have some silence to write at the end of playback...
+			if (!bCurrentlyPlaying)
+			{
+				for (int32 ChannelIndex = 0; ChannelIndex < CaptureBuffer.GetNumValidChannels(); ++ChannelIndex)
+				{
+					FMemory::Memset(Outputs.AudioOutputs[ChannelIndex]->GetData() + DestinationFrame + FramesToPlay, 0, (Outputs.AudioOutputs[0]->Num() - (DestinationFrame + FramesToPlay)) * sizeof(float));
+				}
+			}
+
+			return true;
+		}
+
+		/** 
+		Captures the input audio to a buffer.
+		@param FirstFrameOfPlay There MAY be a request to play WHILE we are capturing. If this is > -1 
+		                        then there is a request to play the captured buffer in this same block.
+		@return True if this function ended up writing ANYTHING to the output buffer.
+		*/
+		bool DoCapture(int32 FirstFrameOfPlay)
+		{
+			int32 LastCaptureTriggerBlockIndex = Inputs.Capture->Last();
+			check (bCurrentlyCapturing || LastCaptureTriggerBlockIndex > -1);
+
+			int32 FirstFrameInBlockToCapture = 0;
+			if (LastCaptureTriggerBlockIndex != -1)
+			{
+				// start or restart capturing!
+				CaptureBuffer.SetNumValidFrames(0);
+				FirstFrameInBlockToCapture = LastCaptureTriggerBlockIndex;
+				bCurrentlyCapturing = true;
+			}
+
+			int32 FramesAvailableAtInput = Inputs.AudioInputs[0]->Num() - FirstFrameInBlockToCapture;
+			int32 FramesWriteableInBuffer = CaptureLengthFrames - CaptureBuffer.GetNumValidFrames();
+			int32 NumFramesToCapture = FMath::Min(FramesAvailableAtInput, FramesWriteableInBuffer);
+
+			if (NumFramesToCapture < 1)
+			{
+				bCurrentlyCapturing = false;
+				bFullyCaptured = true;
+				return false; // we did not render anything to the output.
+			}
+
+			for (int32 ChannelIndex = 0; ChannelIndex < CaptureBuffer.GetNumValidChannels(); ++ChannelIndex)
+			{
+				FMemory::Memcpy(CaptureBuffer.GetValidChannelData(ChannelIndex) + CaptureBuffer.GetNumValidFrames(), Inputs.AudioInputs[ChannelIndex]->GetData() + FirstFrameInBlockToCapture, NumFramesToCapture * sizeof(float));
+			}
+
+			CaptureBuffer.SetNumValidFrames(CaptureBuffer.GetNumValidFrames() + NumFramesToCapture);
+
+			if (CaptureBuffer.GetNumValidFrames() == CaptureBuffer.GetMaxNumFrames())
+			{
+				bCurrentlyCapturing = false;
+				bFullyCaptured = true;
+			}
+
+			// we still have input data to copy to the output...
+			if (!bCurrentlyPlaying && !bHavePlayedThisCapture)
+			{
+				if (FirstFrameOfPlay > -1)
+				{
+					// We WILL be playing this block... so limit "capture output thru" up to that point...
+					BulkOutput(EInputOrSilence::Input, 0, FMath::Min(FirstFrameOfPlay, FirstFrameInBlockToCapture + NumFramesToCapture));
+					bPlayFadeInNotNeededThisBlock = FirstFrameOfPlay == FirstFrameInBlockToCapture && CaptureBuffer.GetNumValidFrames() == NumFramesToCapture;
+				}
+				else
+				{
+					// If we are capturing the last little bit, we faded out playback in the previous block,
+					// so don't pass input to output here!
+					if (NumFramesToCapture >= kNumFadeFrames)
+					{
+						BulkOutput(EInputOrSilence::Input, 0, FirstFrameInBlockToCapture + NumFramesToCapture);
+						BulkOutput(EInputOrSilence::Silence, FirstFrameInBlockToCapture + NumFramesToCapture);
+						// We may be done capturing and have to fade to the state where we are outputting silence.
+						// NOTE: in order to prevent a click if the end of capturing is at the very beginning of the NEXT 
+						// block, we will fade NOW, making the capture buffer a little shorter if we have to...
+						if (CaptureBuffer.GetNumValidFrames() >= (CaptureLengthFrames - kNumFadeFrames))
+						{
+							float CurrentFade = 0.0f;
+							float FadeInc = 1.0f/float(kNumFadeFrames);
+							int32 FadeAtIndex = FirstFrameInBlockToCapture + NumFramesToCapture - 1;
+							for (int32 FadeIndex = 0; FadeIndex < kNumFadeFrames; ++FadeIndex)
+							{
+								for (int32 ChannelIndex = 0; ChannelIndex < CaptureBuffer.GetNumValidChannels(); ++ChannelIndex)
+								{
+									Outputs.AudioOutputs[ChannelIndex]->GetData()[FadeAtIndex] *= CurrentFade;
+								}
+								--FadeAtIndex;
+								CurrentFade = FMath::Clamp(CurrentFade + FadeInc, 0.0f, 1.0f);
+							}
+						}
+					}
+					else
+					{
+						BulkOutput(EInputOrSilence::Silence);
+					}
+				}
+				return true; // we rendered something to the output.
+			}
+			return false; // we did not render anything to the output.
+		}
+
+	private:
+		FInputs  Inputs;
+		FOutputs Outputs;
+		float    MaxCaptureSeconds             = 5.0f;
+		bool     bCurrentlyCapturing           = false;
+		bool     bFullyCaptured                = false;
+		bool     bCurrentlyPlaying             = false;
+		int32    CaptureLengthFrames           = 0;
+		bool     bPlayFadeInNotNeededThisBlock = false;
+		int32    PlayLengthFrames              = 0;
+		bool     bPlayingReverse               = false;
+		int32    NextPlayFrame                 = 0;
+		bool     bHavePlayedThisCapture        = false;
+		TAudioBuffer<float> CaptureBuffer;
+	};
+
+	using FSimpleSamplerNode = Metasound::TNodeFacade<FSimpleSamplerOperator>;
+	METASOUND_REGISTER_NODE_AND_CONFIGURATION(FSimpleSamplerNode, FHarmonixSimpleSamplerNodeConfiguration);
+}
+
+TInstancedStruct<FMetasoundFrontendClassInterface> FHarmonixSimpleSamplerNodeConfiguration::OverrideDefaultInterface(const FMetasoundFrontendClass& InNodeClass) const
+{
+	return TInstancedStruct<FMetasoundFrontendClassInterface>::Make(FMetasoundFrontendClassInterface::GenerateClassInterface(HarmonixMetasound::SimpleSamplerNode::GetVertexInterface(NumChannels)));
+}
+
+TSharedPtr<const Metasound::IOperatorData> FHarmonixSimpleSamplerNodeConfiguration::GetOperatorData() const
+{
+	return MakeShared<HarmonixMetasound::SimpleSamplerNode::FSimpleSamplerNodeOperatorData>(MaxCaptureSeconds);
+}
+
+#undef LOCTEXT_NAMESPACE
diff --git a/Engine/Plugins/Runtime/Harmonix/Source/HarmonixMetasound/Private/HarmonixMetasound/Nodes/StutterSequencerNode.cpp b/Engine/Plugins/Runtime/Harmonix/Source/HarmonixMetasound/Private/HarmonixMetasound/Nodes/StutterSequencerNode.cpp
new file mode 100644
--- /dev/null
+++ b/Engine/Plugins/Runtime/Harmonix/Source/HarmonixMetasound/Private/HarmonixMetasound/Nodes/StutterSequencerNode.cpp
@@ -0,0 +1,337 @@
+// Copyright Epic Games, Inc. All Rights Reserved.
+
+#include "HarmonixMetasound/Nodes/StutterSequencerNode.h"
+
+#include "MetasoundExecutableOperator.h"
+#include "MetasoundFacade.h"
+#include "MetasoundNodeInterface.h"
+#include "MetasoundParamHelper.h"
+#include "MetasoundSampleCounter.h"
+#include "MetasoundStandardNodesCategories.h"
+#include "MetasoundVertex.h"
+#include "HarmonixMetasound/Common.h"
+#include "HarmonixMetasound/DataTypes/MidiClock.h"
+#include "HarmonixMetasound/DataTypes/MidiStutterSequence.h"
+#include "IAudioProxyInitializer.h"
+#include "Harmonix/AudioRenderableProxy.h"
+#include "Algo/FindLast.h"
+
+#define LOCTEXT_NAMESPACE "HarmonixMetaSound"
+
+namespace HarmonixMetasound::StutterSequencerNode
+{
+	using namespace Metasound;
+	using namespace Harmonix;
+	using namespace HarmonixMetasound;
+
+	namespace Inputs
+	{
+		DEFINE_METASOUND_PARAM_ALIAS(Enabled, CommonPinNames::Inputs::Enable);
+		DEFINE_INPUT_METASOUND_PARAM(StutterSequence, "Stutter Sequence", "The asset describing the sequence of stutters to play.");
+		DEFINE_METASOUND_PARAM_ALIAS(MidiClock, CommonPinNames::Inputs::MidiClock);
+		DEFINE_INPUT_METASOUND_PARAM(Start, "Start", "Starts the stutter sequence.");
+		DEFINE_INPUT_METASOUND_PARAM(Stop, "Stop", "Stops the stutter sequence, and resets the simple sampler.");
+	}
+
+	namespace Outputs
+	{
+		DEFINE_OUTPUT_METASOUND_PARAM(Capture, "Capture", "Trigger to start capturing the input audio.");
+		DEFINE_OUTPUT_METASOUND_PARAM(CaptureDuration, "Capture Duration Sec", "How long a sample to capture.");
+		DEFINE_OUTPUT_METASOUND_PARAM(Play, "Play", "Trigger to play back the captured buffer.");
+		DEFINE_OUTPUT_METASOUND_PARAM(PlayDuration, "Play Duration Sec", "How much of the captured sample to play.");
+		DEFINE_OUTPUT_METASOUND_PARAM(Reverse, "Reverse Playback", "If true the captured buffer should play in reverse.");
+		DEFINE_OUTPUT_METASOUND_PARAM(Reset, "Reset", "Trigger to stop capture/playback and start passing audio from input to output.");
+	}
+
+	FVertexInterface GetVertexInterface()
+	{
+		FInputVertexInterface InputInterface;
+		FOutputVertexInterface OutputInterface;
+
+		InputInterface.Add(TInputDataVertex<bool>(METASOUND_GET_PARAM_NAME_AND_METADATA(Inputs::Enabled)));
+		InputInterface.Add(TInputDataVertex<FStutterSequenceAsset>(METASOUND_GET_PARAM_NAME_AND_METADATA(Inputs::StutterSequence)));
+		InputInterface.Add(TInputDataVertex<FMidiClock>(METASOUND_GET_PARAM_NAME_AND_METADATA(Inputs::MidiClock)));
+		InputInterface.Add(TInputDataVertex<FTrigger>(METASOUND_GET_PARAM_NAME_AND_METADATA(Inputs::Start)));
+		InputInterface.Add(TInputDataVertex<FTrigger>(METASOUND_GET_PARAM_NAME_AND_METADATA(Inputs::Stop)));
+
+		OutputInterface.Add(TOutputDataVertex<FTrigger>(METASOUND_GET_PARAM_NAME_AND_METADATA(Outputs::Capture)));
+		OutputInterface.Add(TOutputDataVertex<float>(METASOUND_GET_PARAM_NAME_AND_METADATA(Outputs::CaptureDuration)));
+		OutputInterface.Add(TOutputDataVertex<FTrigger>(METASOUND_GET_PARAM_NAME_AND_METADATA(Outputs::Play)));
+		OutputInterface.Add(TOutputDataVertex<float>(METASOUND_GET_PARAM_NAME_AND_METADATA(Outputs::PlayDuration)));
+		OutputInterface.Add(TOutputDataVertex<bool>(METASOUND_GET_PARAM_NAME_AND_METADATA(Outputs::Reverse)));
+		OutputInterface.Add(TOutputDataVertex<FTrigger>(METASOUND_GET_PARAM_NAME_AND_METADATA(Outputs::Reset)));
+
+		return FVertexInterface
+		{
+			MoveTemp(InputInterface),
+			MoveTemp(OutputInterface)
+		};
+	}
+
+	class FStutterSequencerOperator final : public TExecutableOperator<FStutterSequencerOperator>
+	{
+	public:
+		struct FInputs
+		{
+			FBoolReadRef                 Enabled;
+			FStutterSequenceAssetReadRef StutterSequence;
+			FMidiClockReadRef            MidiClock;
+			FTriggerReadRef              Start;
+			FTriggerReadRef              Stop;
+		};
+
+		struct FOutputs
+		{
+			FTriggerWriteRef Capture;
+			FFloatWriteRef   CaptureDuration;
+			FTriggerWriteRef Play;
+			FFloatWriteRef   PlayDuration;
+			FBoolWriteRef    Reverse;
+			FTriggerWriteRef Reset;
+		};
+		
+		FStutterSequencerOperator(const FBuildOperatorParams& Params, FInputs&& Inputs, FOutputs&& Outputs)
+			: Inputs(MoveTemp(Inputs))
+			, Outputs(MoveTemp(Outputs))
+		{
+			Reset(Params);
+		}
+		
+		static const FNodeClassMetadata& GetNodeInfo()
+		{
+			auto InitNodeInfo = []() -> FNodeClassMetadata
+			{
+				FNodeClassMetadata Info = FNodeClassMetadata::GetEmpty();
+				Info.ClassName = { HarmonixNodeNamespace, TEXT("StutterSequencer"), TEXT("") };
+				Info.MajorVersion = 0;
+				Info.MinorVersion = 1;
+				Info.DisplayName = METASOUND_LOCTEXT("StutterSequencer_DisplayName", "Stutter Sequencer");
+				Info.CategoryHierarchy = { MetasoundNodeCategories::Harmonix };
+				Info.Description = METASOUND_LOCTEXT("StutterSequencer_Description", "Takes a Stutter Sequence asset and uses it to control a Simple Sampler node to create a stutter edit effect.");
+				Info.Author = PluginAuthor;
+				Info.PromptIfMissing = PluginNodeMissingPrompt;
+				Info.DefaultInterface = GetVertexInterface();
+				return Info;
+			};
+
+			static const FNodeClassMetadata Info = InitNodeInfo();
+
+			return Info;
+		}
+
+		static TUniquePtr<IOperator> CreateOperator(const FBuildOperatorParams& InParams, FBuildResults& OutResults)
+		{
+			const FInputVertexInterfaceData& InputData = InParams.InputData;
+
+			FInputs Inputs {
+					InputData.GetOrCreateDefaultDataReadReference<bool>(Inputs::EnabledName, InParams.OperatorSettings),
+					InputData.GetOrCreateDefaultDataReadReference<FStutterSequenceAsset>(Inputs::StutterSequenceName, InParams.OperatorSettings),
+					InputData.GetOrCreateDefaultDataReadReference<FMidiClock>(Inputs::MidiClockName, InParams.OperatorSettings),
+					InputData.GetOrCreateDefaultDataReadReference<FTrigger>(Inputs::StartName, InParams.OperatorSettings),
+					InputData.GetOrCreateDefaultDataReadReference<FTrigger>(Inputs::StopName, InParams.OperatorSettings),
+			};
+
+			FOutputs Outputs {	
+					FTriggerWriteRef::CreateNew(InParams.OperatorSettings),
+					FFloatWriteRef::CreateNew(),
+					FTriggerWriteRef::CreateNew(InParams.OperatorSettings),
+					FFloatWriteRef::CreateNew(),
+					FBoolWriteRef::CreateNew(),
+					FTriggerWriteRef::CreateNew(InParams.OperatorSettings),
+			};
+
+			return MakeUnique<FStutterSequencerOperator>(
+				InParams,
+				MoveTemp(Inputs),
+				MoveTemp(Outputs)
+			);
+		}
+
+		virtual void BindInputs(FInputVertexInterfaceData& InVertexData) override
+		{
+			InVertexData.BindReadVertex(Inputs::EnabledName, Inputs.Enabled);
+			InVertexData.BindReadVertex(Inputs::StutterSequenceName, Inputs.StutterSequence);
+			InVertexData.BindReadVertex(Inputs::MidiClockName, Inputs.MidiClock);
+			InVertexData.BindReadVertex(Inputs::StartName, Inputs.Start);
+			InVertexData.BindReadVertex(Inputs::StopName, Inputs.Stop);
+		}
+
+		virtual void BindOutputs(FOutputVertexInterfaceData& InVertexData) override
+		{
+			InVertexData.BindReadVertex(Outputs::CaptureName, Outputs.Capture);
+			InVertexData.BindReadVertex(Outputs::CaptureDurationName, Outputs.CaptureDuration);
+			InVertexData.BindReadVertex(Outputs::PlayName, Outputs.Play);
+			InVertexData.BindReadVertex(Outputs::PlayDurationName, Outputs.PlayDuration);
+			InVertexData.BindReadVertex(Outputs::ReverseName, Outputs.Reverse);
+			InVertexData.BindReadVertex(Outputs::ResetName, Outputs.Reset);
+		}
+		
+		void Reset(const FResetParams& Params)
+		{
+			StutterTable = Inputs.StutterSequence->GetRenderable();
+			bRunning = false;
+			StartTick = -1;
+			NextStutterTick = -1;
+			CurrentStutterEntryIndex = 0;
+			NumStuttersPlayedThisEntry = 0;
+			Outputs.Capture->Reset();
+			*Outputs.CaptureDuration = 0.0f;
+			Outputs.Play->Reset();
+			*Outputs.PlayDuration = 0.0f;
+			*Outputs.Reverse = false;
+			Outputs.Reset->Reset();
+		}
+		
+		void Execute()
+		{
+			Outputs.Capture->AdvanceBlock();
+			Outputs.Play->AdvanceBlock();
+			Outputs.Reset->AdvanceBlock();
+
+			CheckForUpdatedStutterTable();
+
+			// should we be running at all?
+			if (!StutterTable || !*Inputs.Enabled || (!bRunning && !Inputs.Start->IsTriggeredInBlock()))
+			{
+				if (bRunning)
+				{
+					bRunning = false;
+					if (StutterTable && StutterTable->ResetOnCompletion)
+					{
+						Outputs.Reset->TriggerFrame(0);
+					}
+				}
+				return;
+			}
+
+			if (Inputs.MidiClock->GetTransportStateAtStartOfBlock() != EMusicPlayerTransportState::Playing &&
+				Inputs.MidiClock->GetTransportStateAtEndOfBlock() != EMusicPlayerTransportState::Playing)
+			{
+				// nothing happening!
+				return;
+			}
+
+			if (StutterTable->ResetOnCompletion)
+			{
+				for (int32 StopFrame : Inputs.Stop->GetTriggeredFrames())
+				{
+					Outputs.Reset->TriggerFrame(StopFrame);
+				}
+			}
+
+			if (!bRunning)
+			{
+				// There must be a new start trigger this block or we wouldn't be here!
+				int32 BlockFrame = Inputs.Start->Last();
+				Outputs.Capture->TriggerFrame(BlockFrame);
+				StartTick = Inputs.MidiClock->GetNextTickToProcessAtBlockFrame(BlockFrame);
+				NextStutterTick = StartTick;
+				CurrentStutterEntryIndex = 0;
+				NumStuttersPlayedThisEntry = 0;
+				const ISongMapEvaluator& SongMaps = Inputs.MidiClock->GetSongMapEvaluator();
+				int32 NumTicksToCapture = MusicalTimeSpan::CalculateMusicalSpanLengthTicks(StutterTable->CaptureDuration, SongMaps, StartTick);
+				float MsOfCaptureStart = SongMaps.TickToMs(StartTick);
+				float MsOfCaptureEnd = SongMaps.TickToMs(StartTick + NumTicksToCapture);
+				*Outputs.CaptureDuration = ((MsOfCaptureEnd - MsOfCaptureStart) / 1000.0f) / Inputs.MidiClock->GetSpeedAtBlockSampleFrame(BlockFrame);
+				bRunning = true;
+			}
+
+			for (const FMidiClockEvent& Event : Inputs.MidiClock->GetMidiClockEventsInBlock())
+			{
+				using namespace MidiClockMessageTypes;
+
+				if (const FAdvance* AsAdvance = Event.TryGet<FAdvance>())
+				{
+					if (AsAdvance->LastTickToProcess() >= NextStutterTick)
+					{
+						// We may actually be at the end of the last stutter in the table...
+						if (CurrentStutterEntryIndex == StutterTable->Stutters.Num())
+						{
+							// yup... done...
+							NextStutterTick = -1;
+							bRunning = false;
+							if (StutterTable->ResetOnCompletion)
+							{
+								Outputs.Reset->TriggerFrame(Event.BlockFrameIndex);
+							}
+							break;
+						}
+
+						// we have a stutter to play...
+						Outputs.Play->TriggerFrame(Event.BlockFrameIndex);
+
+						// Calculate the play duration and figure out when this stutter will be done
+						// so we can come back here and trigger the next one...
+						const ISongMapEvaluator& SongMaps = Inputs.MidiClock->GetSongMapEvaluator();
+						int32 NumTicksToPlay = MusicalTimeSpan::CalculateMusicalSpanLengthTicks(StutterTable->Stutters[CurrentStutterEntryIndex].AudibleDuration, SongMaps, NextStutterTick);
+						int32 NumTicksToNextStutter = MusicalTimeSpan::CalculateMusicalSpanLengthTicks(StutterTable->Stutters[CurrentStutterEntryIndex].Spacing, SongMaps, NextStutterTick);
+						if (NumTicksToNextStutter < NumTicksToPlay)
+						{
+							NumTicksToPlay = NumTicksToNextStutter;
+						}
+						float MsOfPlayStart = SongMaps.TickToMs(NextStutterTick);
+						float MsOfPlayEnd = SongMaps.TickToMs(NextStutterTick + NumTicksToPlay);
+						float NewPlayDuration = ((MsOfPlayEnd - MsOfPlayStart) / 1000.0f) / Inputs.MidiClock->GetSpeedAtBlockSampleFrame(Event.BlockFrameIndex);
+						*Outputs.PlayDuration = NewPlayDuration;
+
+						*Outputs.Reverse = StutterTable->Stutters[CurrentStutterEntryIndex].Reverse;
+
+						NextStutterTick += NumTicksToNextStutter;
+						++NumStuttersPlayedThisEntry;
+						if (NumStuttersPlayedThisEntry >= StutterTable->Stutters[CurrentStutterEntryIndex].Count)
+						{
+							++CurrentStutterEntryIndex;
+							NumStuttersPlayedThisEntry = 0;
+						}
+					}
+				}
+			}
+
+			// and finally... 
+			if (Inputs.Start->Last() < Inputs.Stop->Last())
+			{
+				bRunning = false;
+				StartTick = -1;
+				NextStutterTick = -1;
+				if (StutterTable->ResetOnCompletion)
+				{
+					Outputs.Reset->TriggerFrame(Inputs.Stop->Last());
+				}
+			}
+		}
+
+		void CheckForUpdatedStutterTable()
+		{
+			FStutterSequenceTableProxy::NodePtr Tester = Inputs.StutterSequence->GetRenderable();
+			if (Tester != StutterTable)
+			{
+				StutterTable = Inputs.StutterSequence->GetRenderable();
+				return;
+			}
+			if (StutterTable)
+			{
+				TRefCountedAudioRenderableWithQueuedChanges<FStutterSequenceTable>* AsTableWithPossibleChanges = StutterTable;
+				if (AsTableWithPossibleChanges->HasUpdate())
+				{
+					StutterTable = AsTableWithPossibleChanges->GetUpdate();
+				}
+			}
+		}
+
+	private:
+		FInputs Inputs;
+		FOutputs Outputs;
+
+		TSharedAudioRenderableDataPtr<FStutterSequenceTable, TRefCountedAudioRenderableWithQueuedChanges<FStutterSequenceTable>> StutterTable;
+		bool bRunning = false;
+		int32 StartTick = -1;
+		int32 NextStutterTick = -1;
+		int32 CurrentStutterEntryIndex = 0;
+		int32 NumStuttersPlayedThisEntry = 0;
+	};
+
+	using FStutterSequencerNode = Metasound::TNodeFacade<FStutterSequencerOperator>;
+	METASOUND_REGISTER_NODE(FStutterSequencerNode);
+}
+
+#undef LOCTEXT_NAMESPACE
diff --git a/Engine/Plugins/Runtime/Harmonix/Source/HarmonixMetasound/Public/HarmonixMetasound/DataTypes/MidiStutterSequence.h b/Engine/Plugins/Runtime/Harmonix/Source/HarmonixMetasound/Public/HarmonixMetasound/DataTypes/MidiStutterSequence.h
new file mode 100644
--- /dev/null
+++ b/Engine/Plugins/Runtime/Harmonix/Source/HarmonixMetasound/Public/HarmonixMetasound/DataTypes/MidiStutterSequence.h
@@ -0,0 +1,156 @@
+// Copyright Epic Games, Inc. All Rights Reserved.
+
+#pragma once
+
+#include "IAudioProxyInitializer.h"
+#include "Harmonix/AudioRenderableProxy.h"
+#include "AudioRenderableAsset.h"
+#include "MetasoundDataReference.h"
+#include "MetasoundDataTypeRegistrationMacro.h"
+#include "HarmonixMidi/MusicTimeSpan.h"
+
+#include "MidiStutterSequence.generated.h"
+
+#define UE_API HARMONIXMETASOUND_API
+
+USTRUCT(BlueprintType)
+struct FStutterSequenceEntry
+{
+	GENERATED_BODY()
+
+public:
+	FStutterSequenceEntry() = default;
+	FStutterSequenceEntry(const FStutterSequenceEntry&) = default;
+	FStutterSequenceEntry& operator=(const FStutterSequenceEntry&) = default;
+	FStutterSequenceEntry(FStutterSequenceEntry&&) = default;
+	FStutterSequenceEntry(EMusicTimeSpanLengthUnits Spacing, EMusicTimeSpanLengthUnits AudibleDuration,
+		bool Reverse, int32 Count)
+		: Spacing(Spacing)
+		, AudibleDuration(AudibleDuration)
+		, Reverse(Reverse)
+		, Count(Count)
+	{ }
+
+	/** Space between stutters */
+	UPROPERTY(EditAnywhere, BlueprintReadOnly, Category = "MidiStutterSequence", Meta = (PostEditType = "Trivial"))
+	EMusicTimeSpanLengthUnits Spacing = EMusicTimeSpanLengthUnits::QuarterNotes;
+
+	/** Duration of the audible portion of the stutter */
+	UPROPERTY(EditAnywhere, BlueprintReadOnly, Category = "MidiStutterSequence", Meta = (PostEditType = "Trivial"))
+	EMusicTimeSpanLengthUnits AudibleDuration = EMusicTimeSpanLengthUnits::EighthNotes;
+
+	/** Play the audio backwards? */
+	UPROPERTY(EditAnywhere, BlueprintReadOnly, Category = "MidiStutterSequence", Meta = (PostEditType = "Trivial"))
+	bool Reverse = false;
+
+	/** How many stutters to play with these settings */
+	UPROPERTY(EditAnywhere, BlueprintReadOnly, Category = "MidiStutterSequence", Meta = (PostEditType = "Trivial", ClampMin = "1", ClampMax = "64"))
+	int32 Count = 1;
+};
+
+USTRUCT(BlueprintType)
+struct FStutterSequenceTable
+{
+	GENERATED_BODY()
+
+public:
+	// This struct is going to be the "root" of the data we want to share between
+	// the asset's UObject instance and one or more Metasound nodes running on
+	// the audio thread. For that reason, we include this macro so other
+	// templates and macros work.
+	IMPL_AUDIORENDERABLE_PROXYABLE(FStutterSequenceTable)
+
+	FStutterSequenceTable()
+		: Stutters()
+	{}
+	FStutterSequenceTable(const FStutterSequenceTable& Other)
+		: CaptureDuration(Other.CaptureDuration)
+		, ResetOnCompletion(Other.ResetOnCompletion)
+		, Stutters(Other.Stutters)
+	{}
+	FStutterSequenceTable& operator=(const FStutterSequenceTable& Other)
+	{
+		CaptureDuration = Other.CaptureDuration;
+		ResetOnCompletion = Other.ResetOnCompletion;
+		Stutters = Other.Stutters;
+		return *this;
+	}
+	FStutterSequenceTable(FStutterSequenceTable&& Other) noexcept
+		: CaptureDuration(Other.CaptureDuration)
+		, ResetOnCompletion(Other.ResetOnCompletion)
+		, Stutters(MoveTemp(Other.Stutters))
+	{}
+
+	FStutterSequenceTable& operator=(FStutterSequenceTable&& Other)
+	{
+		CaptureDuration = Other.CaptureDuration;
+		ResetOnCompletion = Other.ResetOnCompletion;
+		Stutters = MoveTemp(Other.Stutters);
+		return *this;
+	}
+	~FStutterSequenceTable() = default;
+
+	/** The amount of audio to tell the sampler node to capture */
+	UPROPERTY(EditAnywhere, BlueprintReadOnly, Category = "MidiStutterSequence", Meta = (PostEditType = "Trivial"))
+	EMusicTimeSpanLengthUnits CaptureDuration = EMusicTimeSpanLengthUnits::HalfNotes;
+
+	/** Reset the Sampler node after the last stutter plays? This will cause the sampler node to start passing its input to its output again. */
+	UPROPERTY(EditAnywhere, BlueprintReadOnly, Category = "MidiStutterSequence", Meta = (PostEditType = "Trivial"))
+	bool ResetOnCompletion = true;
+
+	/** A sequence of stutter 'sections'... */
+	UPROPERTY(EditAnywhere, BlueprintReadOnly, Category = "MidiStutterSequence")
+	TArray<FStutterSequenceEntry> Stutters;
+};
+
+USING_AUDIORENDERABLE_PROXY(FStutterSequenceTable, FStutterSequenceTableProxy)
+
+// This next macro does a few things. 
+// - It says, "I want a Metasound exposed asset called 'FStutterSequenceAsset' with 
+//   corresponding TypeInfo, ReadRef, and WriteRef classes."
+// - That asset is a wrapper around a proxy class that acts as the go-between from the 
+//   UObject (GC'able) side to  the audio render thread side. So here I tell the macro to 
+//   wrap "FStutterSequenceTable" in a proxy named "FStutterSequenceTableProxy" and use that 
+//   as the "guts" of the FMidiStepSequenceAsset asset.
+//NOTE: This macro has a corresponding "DEFINE_AUDIORENDERABLE_ASSET" that must be added to the cpp file. 
+DECLARE_AUDIORENDERABLE_ASSET(HarmonixMetasound, FStutterSequenceAsset, FStutterSequenceTableProxy, HARMONIXMETASOUND_API)
+
+// Now I can define the UCLASS that is the UObject side of the asset. Notice it is an 
+// IAudioProxyDataFactory. In the code we will see that the CreateNewProxyData override
+// returns an instance of a proxy class defined with the macro above.
+
+// This class represents a stutter sequence. It is used by the MetasSound Stutter Sequencer node
+// to generate triggers for a Simple Sampler node to create a classic Stutter Edit effect.
+UCLASS(MinimalAPI, BlueprintType, Category = "Music", Meta = (DisplayName = "MIDI Stutter Sequence"))
+class UMidiStutterSequence : public UObject, public IAudioProxyDataFactory
+{
+	GENERATED_BODY()
+
+public:
+
+#if WITH_EDITOR
+	UE_API virtual void PostEditChangeChainProperty(FPropertyChangedChainEvent& PropertyChangedChainEvent) override;
+#endif // WITH_EDITOR
+
+	UE_API TSharedPtr<Audio::IProxyData> CreateProxyData(const Audio::FProxyDataInitParams& InitParams) override;
+
+
+protected:
+	UPROPERTY(EditAnywhere, BlueprintReadOnly, Category = "MidiStepSequence")
+	FStutterSequenceTable StutterTable;
+
+private:
+
+	// Notice here that we cache a pointer the Proxy's "Queue" so we can...
+	// 1 - Supply it to all instances of Metasound nodes rendering this data. How?
+	//     CreateNewProxyData instantiates a NEW unique ptr to an FStepSequenceTableProxy
+	//     every time it is called. All of those unique proxy instances refer to the same
+	//     queue... this one that we have cached.
+	// 2 - Modify that data in response to changes to this class's UPROPERTIES
+	//     so that we can hear data changes reflected in the rendered audio.
+	TSharedPtr<FStutterSequenceTableProxy::QueueType> RenderableStutterTable;
+
+	UE_API void UpdateRenderableForNonTrivialChange();
+};
+
+#undef UE_API
diff --git a/Engine/Plugins/Runtime/Harmonix/Source/HarmonixMetasound/Public/HarmonixMetasound/Nodes/SimpleSamplerNode.h b/Engine/Plugins/Runtime/Harmonix/Source/HarmonixMetasound/Public/HarmonixMetasound/Nodes/SimpleSamplerNode.h
new file mode 100644
--- /dev/null
+++ b/Engine/Plugins/Runtime/Harmonix/Source/HarmonixMetasound/Public/HarmonixMetasound/Nodes/SimpleSamplerNode.h
@@ -0,0 +1,63 @@
+// Copyright Epic Games, Inc. All Rights Reserved.
+
+#pragma once
+
+#include "HarmonixMetasound/Common.h"
+#include "MetasoundFrontendDocument.h"
+#include "MetasoundOperatorData.h"
+
+#include "SimpleSamplerNode.generated.h"
+
+namespace HarmonixMetasound::SimpleSamplerNode
+{
+	class FSimpleSamplerNodeOperatorData;
+}
+
+USTRUCT()
+struct FHarmonixSimpleSamplerNodeConfiguration : public FMetaSoundFrontendNodeConfiguration
+{
+	GENERATED_BODY()
+
+	UPROPERTY(EditAnywhere, Category = General, meta = (ClampMin = "1", ClampMax = "8"))
+	uint32 NumChannels = 1;
+
+	UPROPERTY(EditAnywhere, Category = General, meta = (ClampMin = "1", ClampMax = "10"))
+	float MaxCaptureSeconds = 5.0f;
+
+	/* Get the current interface for the class based upon the node extension */
+	virtual TInstancedStruct<FMetasoundFrontendClassInterface> OverrideDefaultInterface(const FMetasoundFrontendClass& InNodeClass) const;
+
+	TSharedPtr<const Metasound::IOperatorData> GetOperatorData() const override;
+};
+
+namespace HarmonixMetasound::SimpleSamplerNode
+{
+	class FSimpleSamplerNodeOperatorData : public Metasound::TOperatorData<FSimpleSamplerNodeOperatorData>
+	{
+	public:
+		static const FLazyName OperatorDataTypeName;
+		FSimpleSamplerNodeOperatorData(const float& InMaxCaptureSeconds)
+			: MaxCaptureSeconds(InMaxCaptureSeconds)
+		{}
+		float MaxCaptureSeconds = 5.0f;
+	};
+
+	namespace Inputs
+	{
+		DECLARE_METASOUND_PARAM_ALIAS(Enabled);
+		DECLARE_METASOUND_PARAM_EXTERN(Capture);
+		DECLARE_METASOUND_PARAM_EXTERN(CaptureDuration);
+		DECLARE_METASOUND_PARAM_EXTERN(Play);
+		DECLARE_METASOUND_PARAM_EXTERN(PlayDuration);
+		DECLARE_METASOUND_PARAM_EXTERN(Reverse);
+		DECLARE_METASOUND_PARAM_EXTERN(Reset);
+	}
+
+	namespace Outputs
+	{
+		DECLARE_METASOUND_PARAM_EXTERN(InputAudioInUse);
+	}
+
+	static constexpr int32 kMaxChannels = 8;
+	static constexpr int32 kNumFadeFrames = 32;
+}
diff --git a/Engine/Plugins/Runtime/Harmonix/Source/HarmonixMetasound/Public/HarmonixMetasound/Nodes/StutterSequencerNode.h b/Engine/Plugins/Runtime/Harmonix/Source/HarmonixMetasound/Public/HarmonixMetasound/Nodes/StutterSequencerNode.h
new file mode 100644
--- /dev/null
+++ b/Engine/Plugins/Runtime/Harmonix/Source/HarmonixMetasound/Public/HarmonixMetasound/Nodes/StutterSequencerNode.h
@@ -0,0 +1,29 @@
+// Copyright Epic Games, Inc. All Rights Reserved.
+
+#pragma once
+
+#include "HarmonixMetasound/Common.h"
+#include "MetasoundFrontendDocument.h"
+#include "MetasoundOperatorData.h"
+
+namespace HarmonixMetasound::StutterSequencerNode
+{
+	namespace Inputs
+	{
+		DECLARE_METASOUND_PARAM_ALIAS(Enabled);
+		DECLARE_METASOUND_PARAM_EXTERN(StutterSequence);
+		DECLARE_METASOUND_PARAM_ALIAS(MidiClock);
+		DECLARE_METASOUND_PARAM_EXTERN(Start);
+		DECLARE_METASOUND_PARAM_EXTERN(Stop);
+	}
+
+	namespace Outputs
+	{
+		DECLARE_METASOUND_PARAM_EXTERN(Capture);
+		DECLARE_METASOUND_PARAM_EXTERN(CaptureDuration);
+		DECLARE_METASOUND_PARAM_EXTERN(Play);
+		DECLARE_METASOUND_PARAM_EXTERN(PlayDuration);
+		DECLARE_METASOUND_PARAM_EXTERN(Reverse);
+		DECLARE_METASOUND_PARAM_EXTERN(Reset);
+	}
+}
diff --git a/Engine/Plugins/Runtime/Harmonix/Source/HarmonixMetasoundEditor/Private/AssetDefinition_MidiStutterSequence.cpp b/Engine/Plugins/Runtime/Harmonix/Source/HarmonixMetasoundEditor/Private/AssetDefinition_MidiStutterSequence.cpp
new file mode 100644
--- /dev/null
+++ b/Engine/Plugins/Runtime/Harmonix/Source/HarmonixMetasoundEditor/Private/AssetDefinition_MidiStutterSequence.cpp
@@ -0,0 +1,31 @@
+// Copyright Epic Games, Inc. All Rights Reserved.
+
+#include "AssetDefinition_MidiStutterSequence.h"
+#include "HarmonixMetasound/DataTypes/MidiStutterSequence.h"
+
+TSoftClassPtr<UObject> UAssetDefinition_MidiStutterSequence::GetAssetClass() const
+{
+	return UMidiStutterSequence::StaticClass();
+}
+
+FText UAssetDefinition_MidiStutterSequence::GetAssetDisplayName() const
+{
+	return NSLOCTEXT("AssetTypeActions", "MIDIStutterSequenceDefinition", "MIDI Stutter Sequence");
+}
+
+FLinearColor  UAssetDefinition_MidiStutterSequence::GetAssetColor() const
+{
+	return FLinearColor(1.0f, 0.5f, 0.0f);
+}
+
+TConstArrayView<FAssetCategoryPath> UAssetDefinition_MidiStutterSequence::GetAssetCategories() const
+{
+	static const auto Categories = { EAssetCategoryPaths::Audio / NSLOCTEXT("Harmonix", "HmxAssetCategoryName", "Harmonix") };
+	return Categories;
+}
+
+bool UAssetDefinition_MidiStutterSequence::CanImport() const
+{
+	return false;
+}
+
diff --git a/Engine/Plugins/Runtime/Harmonix/Source/HarmonixMetasoundEditor/Private/AssetDefinition_MidiStutterSequence.h b/Engine/Plugins/Runtime/Harmonix/Source/HarmonixMetasoundEditor/Private/AssetDefinition_MidiStutterSequence.h
new file mode 100644
--- /dev/null
+++ b/Engine/Plugins/Runtime/Harmonix/Source/HarmonixMetasoundEditor/Private/AssetDefinition_MidiStutterSequence.h
@@ -0,0 +1,24 @@
+// Copyright Epic Games, Inc. All Rights Reserved.
+
+#pragma once
+
+
+#include "AssetDefinitionDefault.h"
+
+#include "AssetDefinition_MidiStutterSequence.generated.h"
+
+#define ASSET_MENU_ITEM_GROUPING_CHEKED_IN 0
+
+UCLASS()
+class UAssetDefinition_MidiStutterSequence : public UAssetDefinitionDefault
+{
+	GENERATED_BODY()
+
+public:
+	virtual FText GetAssetDisplayName() const override;
+	virtual FLinearColor GetAssetColor() const override;
+	virtual TSoftClassPtr<UObject> GetAssetClass() const override;
+	virtual TConstArrayView<FAssetCategoryPath> GetAssetCategories() const override;
+	virtual bool CanImport() const override;
+
+};
diff --git a/Engine/Plugins/Runtime/Harmonix/Source/HarmonixMetasoundEditor/Private/HarmonixMetasoundEditorModule.cpp b/Engine/Plugins/Runtime/Harmonix/Source/HarmonixMetasoundEditor/Private/HarmonixMetasoundEditorModule.cpp
--- a/Engine/Plugins/Runtime/Harmonix/Source/HarmonixMetasoundEditor/Private/HarmonixMetasoundEditorModule.cpp
+++ b/Engine/Plugins/Runtime/Harmonix/Source/HarmonixMetasoundEditor/Private/HarmonixMetasoundEditorModule.cpp
@@ -19,6 +19,7 @@ void FHarmonixMetasoundEditorModule::StartupModule()
 	Metasound::Editor::IMetasoundEditorModule& MetasoundEditorModule = FModuleManager::GetModuleChecked<Metasound::Editor::IMetasoundEditorModule>("MetasoundEditor");
 	MetasoundEditorModule.RegisterPinType("MIDIAsset");
 	MetasoundEditorModule.RegisterPinType("MIDIStepSequenceAsset");
+	MetasoundEditorModule.RegisterPinType("StutterSequenceAsset");
 	MetasoundEditorModule.RegisterPinType("FusionPatchAsset");
 	MetasoundEditorModule.RegisterPinType("Enum:SubdivisionQuantizationType","Int32");
 	MetasoundEditorModule.RegisterPinType("Enum:DelayFilterType", "Int32");
diff --git a/Engine/Plugins/Runtime/Harmonix/Source/HarmonixMetasoundEditor/Private/MidiStutterSequenceFactory.cpp b/Engine/Plugins/Runtime/Harmonix/Source/HarmonixMetasoundEditor/Private/MidiStutterSequenceFactory.cpp
new file mode 100644
--- /dev/null
+++ b/Engine/Plugins/Runtime/Harmonix/Source/HarmonixMetasoundEditor/Private/MidiStutterSequenceFactory.cpp
@@ -0,0 +1,20 @@
+// Copyright Epic Games, Inc. All Rights Reserved.
+
+#include "MidiStutterSequenceFactory.h"
+#include "HarmonixMetasound/DataTypes/MidiStutterSequence.h"
+
+UMidiStutterSequenceFactory::UMidiStutterSequenceFactory()
+{
+	SupportedClass = UMidiStutterSequence::StaticClass();
+	bCreateNew = true;
+}
+
+FText UMidiStutterSequenceFactory::GetDisplayName() const
+{
+	return NSLOCTEXT("MIDI", "MIDIStutterSequenceFactoryName", "MIDI Stutter Sequence");
+}
+
+UObject* UMidiStutterSequenceFactory::FactoryCreateNew(UClass* Class, UObject* InParent, FName Name, EObjectFlags Flags, UObject* Context, FFeedbackContext* Warn)
+{
+	return NewObject<UMidiStutterSequence>(InParent, Class, Name, Flags, Context);
+}
diff --git a/Engine/Plugins/Runtime/Harmonix/Source/HarmonixMetasoundEditor/Private/MidiStutterSequenceFactory.h b/Engine/Plugins/Runtime/Harmonix/Source/HarmonixMetasoundEditor/Private/MidiStutterSequenceFactory.h
new file mode 100644
--- /dev/null
+++ b/Engine/Plugins/Runtime/Harmonix/Source/HarmonixMetasoundEditor/Private/MidiStutterSequenceFactory.h
@@ -0,0 +1,19 @@
+// Copyright Epic Games, Inc. All Rights Reserved.
+
+#pragma once
+
+#include "Factories/Factory.h"
+#include "MidiStutterSequenceFactory.generated.h"
+
+UCLASS()
+class UMidiStutterSequenceFactory : public UFactory
+{
+	GENERATED_BODY()
+public:
+	UMidiStutterSequenceFactory();
+	//~ BEGIN UFactory interface
+	virtual FText GetDisplayName() const override;
+	virtual UObject* FactoryCreateNew(UClass* Class, UObject* InParent, FName Name, EObjectFlags Flags, UObject* Context, FFeedbackContext* Warn) override;
+	//~ END UFactory interface
+};
+
diff --git a/Engine/Plugins/Runtime/Harmonix/Source/HarmonixMidi/Private/HarmonixMidi/MusicTimeSpan.cpp b/Engine/Plugins/Runtime/Harmonix/Source/HarmonixMidi/Private/HarmonixMidi/MusicTimeSpan.cpp
--- a/Engine/Plugins/Runtime/Harmonix/Source/HarmonixMidi/Private/HarmonixMidi/MusicTimeSpan.cpp
+++ b/Engine/Plugins/Runtime/Harmonix/Source/HarmonixMidi/Private/HarmonixMidi/MusicTimeSpan.cpp
@@ -212,3 +212,79 @@ float FMusicalTimeSpan::CalculateEnclosingFixedSizeSpanExtents(const FMidiSongPo
 	float Progress = FMath::Fractional(InCell);
 	return Progress < 0.0f ? 1.0f + Progress : Progress;
 }
+
+namespace MusicalTimeSpan
+{
+	int32 CalculateMusicalSpanLengthTicks(EMusicTimeSpanLengthUnits Unit, const ISongMapEvaluator& Maps, int32 AtTick)
+	{
+		int32 MidiTTQ = Maps.GetTicksPerQuarterNote();
+		int32 GridUnitTicks = 1;
+		switch (Unit)
+		{
+		case EMusicTimeSpanLengthUnits::Bars:
+			{
+				const FTimeSignature* TimeSig = Maps.GetTimeSignatureAtTick(AtTick);
+				check(TimeSig);
+				int32 TicksInBeat = MidiTTQ * 4 / TimeSig->Denominator;
+				GridUnitTicks = TicksInBeat * TimeSig->Numerator;
+			}
+			break;
+		case EMusicTimeSpanLengthUnits::Beats:
+			{
+				const FTimeSignature* TimeSig = Maps.GetTimeSignatureAtTick(AtTick);
+				check(TimeSig);
+				GridUnitTicks = MidiTTQ * 4 / TimeSig->Denominator;
+			}
+			break;
+		case EMusicTimeSpanLengthUnits::ThirtySecondNotes:
+			GridUnitTicks = MidiTTQ / 8;
+			break;
+		case EMusicTimeSpanLengthUnits::SixteenthNotes:
+			GridUnitTicks = MidiTTQ / 4;
+			break;
+		case EMusicTimeSpanLengthUnits::EighthNotes:
+			GridUnitTicks = MidiTTQ / 2;
+			break;
+		case EMusicTimeSpanLengthUnits::QuarterNotes:
+			GridUnitTicks = MidiTTQ;
+			break;
+		case EMusicTimeSpanLengthUnits::HalfNotes:
+			GridUnitTicks = MidiTTQ * 2;
+			break;
+		case EMusicTimeSpanLengthUnits::WholeNotes:
+			GridUnitTicks = MidiTTQ * 4;
+			break;
+		case EMusicTimeSpanLengthUnits::DottedSixteenthNotes:
+			GridUnitTicks = MidiTTQ * 3 / 8;
+			break;
+		case EMusicTimeSpanLengthUnits::DottedEighthNotes:
+			GridUnitTicks = MidiTTQ * 3 / 4;
+			break;
+		case EMusicTimeSpanLengthUnits::DottedQuarterNotes:
+			GridUnitTicks = MidiTTQ * 3 / 2;
+			break;
+		case EMusicTimeSpanLengthUnits::DottedHalfNotes:
+			GridUnitTicks = MidiTTQ * 3;
+			break;
+		case EMusicTimeSpanLengthUnits::DottedWholeNotes:
+			GridUnitTicks = MidiTTQ * 6;
+			break;
+		case EMusicTimeSpanLengthUnits::SixteenthNoteTriplets:
+			GridUnitTicks = MidiTTQ / 6;
+			break;
+		case EMusicTimeSpanLengthUnits::EighthNoteTriplets:
+			GridUnitTicks = MidiTTQ / 3;
+			break;
+		case EMusicTimeSpanLengthUnits::QuarterNoteTriplets:
+			GridUnitTicks = MidiTTQ * 2 / 3;
+			break;
+		case EMusicTimeSpanLengthUnits::HalfNoteTriplets:
+			GridUnitTicks = MidiTTQ * 4 / 3;
+			break;
+		default:
+			checkNoEntry();
+			break;
+		}
+		return GridUnitTicks;
+	}
+}
diff --git a/Engine/Plugins/Runtime/Harmonix/Source/HarmonixMidi/Public/HarmonixMidi/MusicTimeSpan.h b/Engine/Plugins/Runtime/Harmonix/Source/HarmonixMidi/Public/HarmonixMidi/MusicTimeSpan.h
--- a/Engine/Plugins/Runtime/Harmonix/Source/HarmonixMidi/Public/HarmonixMidi/MusicTimeSpan.h
+++ b/Engine/Plugins/Runtime/Harmonix/Source/HarmonixMidi/Public/HarmonixMidi/MusicTimeSpan.h
@@ -53,6 +53,11 @@ enum class EMusicTimeSpanLengthUnits : uint8
 	HalfNoteTriplets	  UMETA(DisplayName = "1/2 (triplets)"),
 };
 
+namespace MusicalTimeSpan
+{
+	UE_API int32 CalculateMusicalSpanLengthTicks(EMusicTimeSpanLengthUnits Unit, const ISongMapEvaluator& Maps, int32 AtTick);
+}
+
 USTRUCT(BlueprintType, Category = "HarmonixMidi", meta = (ScriptName = MusicalTimeSpan))
 struct FMusicalTimeSpan
 {
