Harmonix Music Clock Component Functional Tests

#rnx

#virtualized

diff --git a/Engine/Plugins/Runtime/Harmonix/Source/HarmonixMetasoundTests/Private/HarmonixFunctionalTestAction.cpp b/Engine/Plugins/Runtime/Harmonix/Source/HarmonixMetasoundTests/Private/HarmonixFunctionalTestAction.cpp
--- a/Engine/Plugins/Runtime/Harmonix/Source/HarmonixMetasoundTests/Private/HarmonixFunctionalTestAction.cpp
+++ b/Engine/Plugins/Runtime/Harmonix/Source/HarmonixMetasoundTests/Private/HarmonixFunctionalTestAction.cpp
@@ -157,14 +157,13 @@ void UHarmonixFunctionalTestActionParallel::Tick_Implementation(AFunctionalTest*
 
 void UHarmonixFunctionalTestActionParallel::OnFinished_Implementation()
 {
-	FinishAllActions(true);
+	FinishAllActions(ShouldContinue());
 }
 
 void UHarmonixFunctionalTestActionParallel::FinishAllActions(bool bContinue)
 {
 	while (ActionStack.Num() > 0)
 	{
-		
 		ActionStack.Pop()->Finish(bContinue);
 	}
 }
@@ -192,6 +191,6 @@ void UHarmonixFunctionalTestActionDelay::Tick_Implementation(AFunctionalTest* Te
 void UHarmonixFunctionalTestActionFinishTest::OnStart_Implementation(AFunctionalTest* Test)
 {
 	Test->FinishTest(Result, Message);
-	Finish(true);
+	Finish(false);
 }
 
diff --git a/Engine/Plugins/Runtime/Harmonix/Source/HarmonixMetasoundTests/Private/HarmonixMetasoundFunctionalTest.cpp b/Engine/Plugins/Runtime/Harmonix/Source/HarmonixMetasoundTests/Private/HarmonixMetasoundFunctionalTest.cpp
--- a/Engine/Plugins/Runtime/Harmonix/Source/HarmonixMetasoundTests/Private/HarmonixMetasoundFunctionalTest.cpp
+++ b/Engine/Plugins/Runtime/Harmonix/Source/HarmonixMetasoundTests/Private/HarmonixMetasoundFunctionalTest.cpp
@@ -286,7 +286,10 @@ void AHarmonixMetasoundFunctionalTest::StartTest()
 
 void AHarmonixMetasoundFunctionalTest::FinishTest(EFunctionalTestResult TestResult, const FString& Message)
 {
-	ActionSequence->Finish(true);
+	if (ActionSequence)
+	{
+		ActionSequence->Finish(false);
+	}
 	
 	if (AudioComponent)
 	{
diff --git a/Engine/Plugins/Runtime/Harmonix/Source/HarmonixMetasoundTests/Private/HarmonixMetasoundFunctionalTestAction.cpp b/Engine/Plugins/Runtime/Harmonix/Source/HarmonixMetasoundTests/Private/HarmonixMetasoundFunctionalTestAction.cpp
--- a/Engine/Plugins/Runtime/Harmonix/Source/HarmonixMetasoundTests/Private/HarmonixMetasoundFunctionalTestAction.cpp
+++ b/Engine/Plugins/Runtime/Harmonix/Source/HarmonixMetasoundTests/Private/HarmonixMetasoundFunctionalTestAction.cpp
@@ -2,8 +2,9 @@
 
 #include "HarmonixMetasoundFunctionalTestAction.h"
 #include "Components/AudioComponent.h"
-
 #include "FunctionalTest.h"
+#include "HarmonixMetasound/Components/MusicClockComponent.h"
+#include "HAL/FileManager.h"
 
 void UHarmonixMetasoundFunctionalTestActionSetAudioParameter::OnStart_Implementation(AFunctionalTest* Test)
 {
@@ -36,3 +37,91 @@ void UHarmonixMetasoundFunctionalTestActionSetWaitForAudioFinished::OnAudioFinis
 {
 	Finish(true);
 }
+
+void UHarmonixMetasoundFunctonalTestActionRecordClockOutput::Tick_Implementation(AFunctionalTest* Test, float DeltaSeconds)
+{
+	if (RecordTime <= 0.0f)
+	{
+		Finish(true);
+	}
+	
+	if (UMusicClockComponent* MusicClock = Test->GetComponentByClass<UMusicClockComponent>())
+	{
+
+		FMusicTimestamp Timestamp(4, 1.0f);
+		int32 NumChanges = MusicClock->GetSongMaps().GetNumTimeSignatureChanges();
+
+		float TotalBeats = 0.0f;
+		int32 BarIndex = Timestamp.Bar - 1;
+		int32 BeatIndex = Timestamp.Beat - 1.0f;
+
+		int32 LastBarIndex = 0;
+		int32 LastBeatIndex = 0;
+		int32 LastTimeSigNum = 0;
+		int32 LastTimeSigDenom = 0;
+		for (int32 Idx = 0; Idx < NumChanges; ++Idx)
+		{
+			const FTimeSignaturePoint* TimeSigPoint = MusicClock->GetSongMaps().GetTimeSignaturePoint(Idx);
+
+			int32 NumBars = TimeSigPoint->BarIndex - LastBarIndex;
+
+			TotalBeats += NumBars * LastTimeSigNum;
+
+			LastTimeSigNum = TimeSigPoint->TimeSignature.Numerator;
+		}
+		if (MusicClock->GetState() == EMusicClockState::Running)
+		{
+			if (!bClockStarted)
+			{
+				bClockStarted = true;
+				StartTime = TotalTime;
+			}
+			FMidiSongPos SongPos = MusicClock->GetSongPos();
+			float BarDuration = MusicClock->GetCurrentBarDurationSeconds();
+			float BeatDuration = MusicClock->GetCurrentBeatDurationSeconds();
+			ClockFrameData.Add({ TotalTime - StartTime, SongPos , BarDuration, BeatDuration, MusicClock->GetState()});
+		}
+	}
+	TotalTime += DeltaSeconds;
+	if (TotalTime >= RecordTime)
+	{
+		Finish(true);
+	}
+	
+}
+
+void UHarmonixMetasoundFunctonalTestActionRecordClockOutput::OnFinished_Implementation()
+{
+	IConsoleVariable* WriteOutputToFileCVar = IConsoleManager::Get().FindConsoleVariable(TEXT("harmonix.tests.WriteOutputToFile"));
+	int32 CVarValue = WriteOutputToFileCVar ? WriteOutputToFileCVar->GetInt() : 0;
+	bool bWriteOutputToFile = CVarValue == 1;
+
+	if (!bWriteOutputToFile)
+	{
+		return;
+	}
+	
+	const FString FilePath = FString::Printf(TEXT("%s%s.csv"), *FPaths::ProjectSavedDir(), *OutputFilename);
+	if (TUniquePtr<FArchive> Archive(IFileManager::Get().CreateFileWriter(*FilePath)); Archive.IsValid())
+	{
+		const FString Header(TEXT("Time,State,SecondsFromBarOne,SecondsIncludingCountIn,TimeSig,Tempo,BarsIncludingCountIn,BeatsIncludingCountIn,BeatType,SecondsPerBar,SecondsPerBeat,Timestamp,SongSection\n"));
+		Archive->Serialize(TCHAR_TO_ANSI(*Header), Header.Len());
+		for (const FClockFrameData& Data : ClockFrameData)
+		{
+			const FString Entry = FString::Printf(
+				TEXT("%.4f,%s,%.4f,%.4f,%d/%d,%.4f,%.4f,%.4f,%s,%.2f,%.2f,%d:%.4f,%s\n"),
+				Data.Time, *StaticEnum<EMusicClockState>()->GetNameByValue((int64)Data.ClockState).ToString(),
+				Data.SongPos.SecondsFromBarOne, Data.SongPos.SecondsIncludingCountIn,
+				Data.SongPos.TimeSigNumerator, Data.SongPos.TimeSigDenominator, Data.SongPos.Tempo,
+				Data.SongPos.BarsIncludingCountIn, Data.SongPos.BeatsIncludingCountIn,
+				*StaticEnum<EMusicalBeatType>()->GetNameByValue((int64)Data.SongPos.BeatType).ToString(),
+				Data.BarDuration, Data.BeatDuration,
+				Data.SongPos.Timestamp.Bar, Data.SongPos.Timestamp.Beat,
+				Data.SongPos.CurrentSongSection.Name.IsEmpty() ? TEXT("None") : *Data.SongPos.CurrentSongSection.Name);
+			Archive->Serialize(TCHAR_TO_ANSI(*Entry), Entry.Len());
+		}
+		Archive->Flush();
+		Archive->Close();
+	}
+}
+
diff --git a/Engine/Plugins/Runtime/Harmonix/Source/HarmonixMetasoundTests/Private/HarmonixMetasoundFunctionalTestAction.h b/Engine/Plugins/Runtime/Harmonix/Source/HarmonixMetasoundTests/Private/HarmonixMetasoundFunctionalTestAction.h
--- a/Engine/Plugins/Runtime/Harmonix/Source/HarmonixMetasoundTests/Private/HarmonixMetasoundFunctionalTestAction.h
+++ b/Engine/Plugins/Runtime/Harmonix/Source/HarmonixMetasoundTests/Private/HarmonixMetasoundFunctionalTestAction.h
@@ -4,6 +4,8 @@
 
 #include "HarmonixFunctionalTestAction.h"
 #include "AudioParameter.h"
+#include "HarmonixMetasound/Components/MusicClockComponent.h"
+#include "HarmonixMidi/MidiSongPos.h"
 
 #include "HarmonixMetasoundFunctionalTestAction.generated.h"
 
@@ -27,4 +29,39 @@ public:
 protected:
 	UFUNCTION()
 	void OnAudioFinished();
+};
+
+UCLASS(NotBlueprintable, meta=(DisplayName="Record Clock Output"))
+class UHarmonixMetasoundFunctonalTestActionRecordClockOutput : public UHarmonixFunctionalTestAction
+{
+	GENERATED_BODY()
+public:
+	
+	virtual void Tick_Implementation(AFunctionalTest* Test, float DeltaSeconds) override;
+
+	virtual void OnFinished_Implementation() override;
+
+	UPROPERTY(EditAnywhere, Category = "Functional Testing", meta=(ClampMin="0.0", UIMin="0.0"))
+	float RecordTime = 1.0f;
+
+	// optionally write the output to a file
+	UPROPERTY(EditAnywhere, Category = "Functional Testing")
+	FString OutputFilename = FString("Clock.csv");
+
+private:
+
+	float TotalTime = 0.0f;
+	float StartTime = 0.0f;
+	bool bClockStarted = false;
+
+	struct FClockFrameData
+	{
+		float Time = 0.0f;
+		FMidiSongPos SongPos;
+		float BarDuration = 0.0f;
+		float BeatDuration = 0.0f;
+		EMusicClockState ClockState;
+	};
+
+	TArray<FClockFrameData> ClockFrameData;
 };
\ No newline at end of file
