clean up some dead code and fix up some logic. Also, comment out some code that will need to change to avoid merge conflicts down the line

#rnx

#virtualized

diff --git a/Engine/Plugins/Runtime/Harmonix/Source/HarmonixMetasoundTests/Private/HarmonixMetasoundFunctionalTestAction.cpp b/Engine/Plugins/Runtime/Harmonix/Source/HarmonixMetasoundTests/Private/HarmonixMetasoundFunctionalTestAction.cpp
--- a/Engine/Plugins/Runtime/Harmonix/Source/HarmonixMetasoundTests/Private/HarmonixMetasoundFunctionalTestAction.cpp
+++ b/Engine/Plugins/Runtime/Harmonix/Source/HarmonixMetasoundTests/Private/HarmonixMetasoundFunctionalTestAction.cpp
@@ -40,35 +40,19 @@ void UHarmonixMetasoundFunctionalTestActionSetWaitForAudioFinished::OnAudioFinis
 
 void UHarmonixMetasoundFunctonalTestActionRecordClockOutput::Tick_Implementation(AFunctionalTest* Test, float DeltaSeconds)
 {
+	if (IsFinished())
+	{
+		return;
+	}
+	
 	if (RecordTime <= 0.0f)
 	{
 		Finish(true);
+		return;
 	}
 	
 	if (UMusicClockComponent* MusicClock = Test->GetComponentByClass<UMusicClockComponent>())
 	{
-
-		FMusicTimestamp Timestamp(4, 1.0f);
-		int32 NumChanges = MusicClock->GetSongMaps().GetNumTimeSignatureChanges();
-
-		float TotalBeats = 0.0f;
-		int32 BarIndex = Timestamp.Bar - 1;
-		int32 BeatIndex = Timestamp.Beat - 1.0f;
-
-		int32 LastBarIndex = 0;
-		int32 LastBeatIndex = 0;
-		int32 LastTimeSigNum = 0;
-		int32 LastTimeSigDenom = 0;
-		for (int32 Idx = 0; Idx < NumChanges; ++Idx)
-		{
-			const FTimeSignaturePoint* TimeSigPoint = MusicClock->GetSongMaps().GetTimeSignaturePoint(Idx);
-
-			int32 NumBars = TimeSigPoint->BarIndex - LastBarIndex;
-
-			TotalBeats += NumBars * LastTimeSigNum;
-
-			LastTimeSigNum = TimeSigPoint->TimeSignature.Numerator;
-		}
 		if (MusicClock->GetState() == EMusicClockState::Running)
 		{
 			if (!bClockStarted)
@@ -77,8 +61,9 @@ void UHarmonixMetasoundFunctonalTestActionRecordClockOutput::Tick_Implementation
 				StartTime = TotalTime;
 			}
 			FMidiSongPos SongPos = MusicClock->GetSongPos();
-			float BarDuration = MusicClock->GetCurrentBarDurationSeconds();
-			float BeatDuration = MusicClock->GetCurrentBeatDurationSeconds();
+			// TODO: uncomment these
+			float BarDuration = -1.0f; // MusicClock->GetCurrentSecondsPerBar();
+			float BeatDuration = -1.0f; //MusicClock->GetCurrentSecondsPerBeat();
 			ClockFrameData.Add({ TotalTime - StartTime, SongPos , BarDuration, BeatDuration, MusicClock->GetState()});
 		}
 	}
@@ -109,13 +94,13 @@ void UHarmonixMetasoundFunctonalTestActionRecordClockOutput::OnFinished_Implemen
 		for (const FClockFrameData& Data : ClockFrameData)
 		{
 			const FString Entry = FString::Printf(
-				TEXT("%.4f,%s,%.4f,%.4f,%d/%d,%.4f,%.4f,%.4f,%s,%.2f,%.2f,%d:%.4f,%s\n"),
+				TEXT("%.4f,%s,%.4f,%.4f,%d/%d,%.4f,%.4f,%.4f,%s,%s,%s,%d:%.4f,%s\n"),
 				Data.Time, *StaticEnum<EMusicClockState>()->GetNameByValue((int64)Data.ClockState).ToString(),
 				Data.SongPos.SecondsFromBarOne, Data.SongPos.SecondsIncludingCountIn,
 				Data.SongPos.TimeSigNumerator, Data.SongPos.TimeSigDenominator, Data.SongPos.Tempo,
 				Data.SongPos.BarsIncludingCountIn, Data.SongPos.BeatsIncludingCountIn,
 				*StaticEnum<EMusicalBeatType>()->GetNameByValue((int64)Data.SongPos.BeatType).ToString(),
-				Data.BarDuration, Data.BeatDuration,
+				TEXT("TODO: FIX ME"), TEXT("TODO: FIX ME"),
 				Data.SongPos.Timestamp.Bar, Data.SongPos.Timestamp.Beat,
 				Data.SongPos.CurrentSongSection.Name.IsEmpty() ? TEXT("None") : *Data.SongPos.CurrentSongSection.Name);
 			Archive->Serialize(TCHAR_TO_ANSI(*Entry), Entry.Len());
