Support Custom Clock Drivers on the Music Clock Component

#rnx

diff --git a/Engine/Plugins/Runtime/Harmonix/Source/HarmonixMetasound/Private/HarmonixMetasound/Components/MetasoundMusicClockDriver.h b/Engine/Plugins/Runtime/Harmonix/Source/HarmonixMetasound/Private/HarmonixMetasound/Components/MetasoundMusicClockDriver.h
--- a/Engine/Plugins/Runtime/Harmonix/Source/HarmonixMetasound/Private/HarmonixMetasound/Components/MetasoundMusicClockDriver.h
+++ b/Engine/Plugins/Runtime/Harmonix/Source/HarmonixMetasound/Private/HarmonixMetasound/Components/MetasoundMusicClockDriver.h
@@ -11,11 +11,11 @@ namespace HarmonixMetasound::Analysis
 	struct FSongMapChain;
 }
 
-struct FMetasoundMusicClockDriver : public FMusicClockDriverBase
+class FMetasoundMusicClockDriver : public FMusicClockDriverBase
 {
 public:
-	FMetasoundMusicClockDriver(UObject* WorldContextObj, float DefaultTempo, int32 DefaultTimeSigNum, int32 DefaultTimeSigDenom)
-		: FMusicClockDriverBase(WorldContextObj, DefaultTempo, DefaultTimeSigNum, DefaultTimeSigDenom) {}
+	FMetasoundMusicClockDriver(UObject* WorldContextObj, const FMetasoundMusicClockSettings& Settings)
+		: FMusicClockDriverBase(WorldContextObj, Settings) {}
 
 	virtual bool CalculateSongPosWithOffset(float MsOffset, ECalibratedMusicTimebase Timebase, FMidiSongPos& OutResult) const override;
 
diff --git a/Engine/Plugins/Runtime/Harmonix/Source/HarmonixMetasound/Private/HarmonixMetasound/Components/MusicClockComponent.cpp b/Engine/Plugins/Runtime/Harmonix/Source/HarmonixMetasound/Private/HarmonixMetasound/Components/MusicClockComponent.cpp
--- a/Engine/Plugins/Runtime/Harmonix/Source/HarmonixMetasound/Private/HarmonixMetasound/Components/MusicClockComponent.cpp
+++ b/Engine/Plugins/Runtime/Harmonix/Source/HarmonixMetasound/Private/HarmonixMetasound/Components/MusicClockComponent.cpp
@@ -55,6 +55,47 @@ UMusicClockComponent* UMusicClockComponent::CreateWallClockDrivenMusicClock(UObj
 	return NewClock;
 }
 
+UMusicClockComponent* UMusicClockComponent::CreateMusicClockWithSettings(UObject* WorldContextObject, TInstancedStruct<FMusicClockSettingsBase> MusicClockSettings, bool Start)
+{
+	UMusicClockComponent* NewClock = NewObject<UMusicClockComponent>(WorldContextObject);
+	if (!NewClock->GetOwner())
+	{
+		UE_LOG(LogMusicClock, Warning, TEXT("Can't create a music clock in a non-actor context!"));
+		return nullptr;
+	}
+	NewClock->RegisterComponent();
+	
+	if (const FMetasoundMusicClockSettings* MetasoundClockSettings = MusicClockSettings.GetPtr<FMetasoundMusicClockSettings>())
+	{
+		NewClock->DriveMethod = EMusicClockDriveMethod::MetaSound;
+		NewClock->DefaultTempo = MetasoundClockSettings->DefaultTempo;
+		NewClock->DefaultTimeSignatureNum = MetasoundClockSettings->DefaultTimeSigNumerator;
+		NewClock->DefaultTimeSignatureDenom = MetasoundClockSettings->DefaultTimeSigDenominator;
+		NewClock->MetasoundOutputName = MetasoundClockSettings->MetasoundOutputName;
+		NewClock->ConnectToMetasoundOnAudioComponent(MetasoundClockSettings->AudioComponent);
+	}
+	else if (const FWallClockMusicClockSettings* WallClockSettings = MusicClockSettings.GetPtr<FWallClockMusicClockSettings>())
+	{
+		NewClock->DriveMethod = EMusicClockDriveMethod::WallClock;
+		NewClock->DefaultTempo = WallClockSettings->DefaultTempo;
+		NewClock->DefaultTimeSignatureNum = WallClockSettings->DefaultTimeSigNumerator;
+		NewClock->DefaultTimeSignatureDenom = WallClockSettings->DefaultTimeSigDenominator;
+		NewClock->ConnectToWallClockForMidi(WallClockSettings->TempoMap);
+	}
+	else
+	{
+		NewClock->MusicClockSettings = MusicClockSettings;
+		NewClock->ConnectToCustomClock();	
+	}
+	
+	if (Start)
+	{
+		NewClock->Start();
+	}
+	return NewClock;
+}
+
+
 bool UMusicClockComponent::ConnectToMetasoundOnAudioComponent(UAudioComponent* InAudioComponent)
 {
 	DriveMethod = EMusicClockDriveMethod::MetaSound;
@@ -109,14 +150,23 @@ void UMusicClockComponent::EnsureClockIsValidForGameFrameFromSubsystem()
 
 void UMusicClockComponent::CreateClockDriver()
 {
-	if (DriveMethod == EMusicClockDriveMethod::WallClock || !::IsValid(MetasoundsAudioComponent))
+	if (DriveMethod == EMusicClockDriveMethod::WallClock)
 	{
 		ConnectToWallClock();
 	}
-	else
+	else if (DriveMethod == EMusicClockDriveMethod::MetaSound)
 	{
 		ConnectToMetasound();
 	}
+	else if (DriveMethod == EMusicClockDriveMethod::Custom)
+	{
+		ConnectToCustomClock();
+	}
+
+	if (!ClockDriver)
+	{
+		ConnectToWallClock();
+	}
 }
 
 bool UMusicClockComponent::ConnectToMetasound()
@@ -127,7 +177,13 @@ bool UMusicClockComponent::ConnectToMetasound()
 		return false;
 	}
 	DisconnectFromClockDriver();
-	TSharedPtr<FMetasoundMusicClockDriver> MetasoundClockDriver = MakeShared<FMetasoundMusicClockDriver>(this, DefaultTempo, DefaultTimeSignatureNum, DefaultTimeSignatureDenom);
+	FMetasoundMusicClockSettings MetasoundClockSettings;
+	MetasoundClockSettings.AudioComponent = MetasoundsAudioComponent;
+	MetasoundClockSettings.MetasoundOutputName = MetasoundOutputName;
+	MetasoundClockSettings.DefaultTempo = DefaultTempo;
+	MetasoundClockSettings.DefaultTimeSigNumerator = DefaultTimeSignatureNum;
+	MetasoundClockSettings.DefaultTimeSigDenominator = DefaultTimeSignatureDenom;
+	TSharedPtr<FMetasoundMusicClockDriver> MetasoundClockDriver = MakeShared<FMetasoundMusicClockDriver>(this, MetasoundClockSettings);
 	bool Connected = MetasoundClockDriver->ConnectToAudioComponentsMetasound(
 		MetasoundsAudioComponent,
 		MetasoundOutputName,
@@ -147,7 +203,24 @@ void UMusicClockComponent::ConnectToWallClock()
 	// we don't 'check' that the driver mode is wall clock here because if the driver mode is metasound and we can't 
 	// connect for some reason we will fall back to this clock driver!
 	DisconnectFromClockDriver();
-	ClockDriver = MakeShared<FWallClockMusicClockDriver>(this, TempoMap, DefaultTempo, DefaultTimeSignatureNum, DefaultTimeSignatureDenom);
+	FWallClockMusicClockSettings WallClockSettings;
+	WallClockSettings.TempoMap = TempoMap;
+	WallClockSettings.DefaultTempo = DefaultTempo;
+	WallClockSettings.DefaultTimeSigNumerator = DefaultTimeSignatureNum;
+	WallClockSettings.DefaultTimeSigDenominator = DefaultTimeSignatureDenom;
+	ClockDriver = MakeShared<FWallClockMusicClockDriver>(this, WallClockSettings);
+	if (IsActive())
+	{
+		ClockDriver->Start();
+	}
+}
+
+void UMusicClockComponent::ConnectToCustomClock()
+{
+	check(DriveMethod == EMusicClockDriveMethod::Custom);
+	DisconnectFromClockDriver();
+	const FMusicClockSettingsBase* Settings = MusicClockSettings.GetPtr<FMusicClockSettingsBase>();
+	ClockDriver = Settings->MakeInstance(this);
 	if (IsActive())
 	{
 		ClockDriver->Start();
diff --git a/Engine/Plugins/Runtime/Harmonix/Source/HarmonixMetasound/Private/HarmonixMetasound/Components/MusicClockDriverBase.cpp b/Engine/Plugins/Runtime/Harmonix/Source/HarmonixMetasound/Private/HarmonixMetasound/Components/MusicClockDriverBase.cpp
--- a/Engine/Plugins/Runtime/Harmonix/Source/HarmonixMetasound/Private/HarmonixMetasound/Components/MusicClockDriverBase.cpp
+++ b/Engine/Plugins/Runtime/Harmonix/Source/HarmonixMetasound/Private/HarmonixMetasound/Components/MusicClockDriverBase.cpp
@@ -11,9 +11,9 @@ void FMusicClockDriverBase::ResetDefaultSongMaps(float DefaultTempo, int32 Defau
 	DefaultMaps.GetBarMap().AddTimeSignatureAtBarIncludingCountIn(0, DefaultTimeSigNum, DefaultTimeSigDenom);
 }
 
-FMusicClockDriverBase::FMusicClockDriverBase(UObject* WorldContextObject, float DefaultTempo, int32 DefaultTimeSigNum, int32 DefaultTimeSigDenom)
+FMusicClockDriverBase::FMusicClockDriverBase(UObject* WorldContextObject, const FMusicClockSettingsBase& Settings)
 {
-	ResetDefaultSongMaps(DefaultTempo, DefaultTimeSigNum, DefaultTimeSigDenom);
+	ResetDefaultSongMaps(Settings.DefaultTempo, Settings.DefaultTimeSigNumerator, Settings.DefaultTimeSigDenominator);
 	ContextObj = TWeakObjectPtr<UObject>(WorldContextObject);
 }
 
diff --git a/Engine/Plugins/Runtime/Harmonix/Source/HarmonixMetasound/Private/HarmonixMetasound/Components/MusicClockTypes.cpp b/Engine/Plugins/Runtime/Harmonix/Source/HarmonixMetasound/Private/HarmonixMetasound/Components/MusicClockTypes.cpp
new file mode 100644
--- /dev/null
+++ b/Engine/Plugins/Runtime/Harmonix/Source/HarmonixMetasound/Private/HarmonixMetasound/Components/MusicClockTypes.cpp
@@ -0,0 +1,15 @@
+﻿// Copyright Epic Games, Inc. All Rights Reserved.
+
+#include "HarmonixMetasound/Components/MusicClockTypes.h"
+#include "HarmonixMetasound/Components/WallClockMusicClockDriver.h"
+#include "HarmonixMetasound/Components/MetasoundMusicClockDriver.h"
+
+TSharedPtr<FMusicClockDriverBase> FWallClockMusicClockSettings::MakeInstance(UObject* WorldContextObj) const
+{
+	return MakeShared<FWallClockMusicClockDriver>(WorldContextObj, *this);
+}
+
+TSharedPtr<FMusicClockDriverBase> FMetasoundMusicClockSettings::MakeInstance(UObject* WorldContextObj) const
+{
+	return MakeShared<FMetasoundMusicClockDriver>(WorldContextObj, *this);
+}
\ No newline at end of file
diff --git a/Engine/Plugins/Runtime/Harmonix/Source/HarmonixMetasound/Private/HarmonixMetasound/Components/WallClockMusicClockDriver.h b/Engine/Plugins/Runtime/Harmonix/Source/HarmonixMetasound/Private/HarmonixMetasound/Components/WallClockMusicClockDriver.h
--- a/Engine/Plugins/Runtime/Harmonix/Source/HarmonixMetasound/Private/HarmonixMetasound/Components/WallClockMusicClockDriver.h
+++ b/Engine/Plugins/Runtime/Harmonix/Source/HarmonixMetasound/Private/HarmonixMetasound/Components/WallClockMusicClockDriver.h
@@ -4,12 +4,12 @@
 #include "HarmonixMidi/MidiFile.h"
 #include "UObject/StrongObjectPtr.h"
 
-struct FWallClockMusicClockDriver : public FMusicClockDriverBase
+class FWallClockMusicClockDriver : public FMusicClockDriverBase
 {
 public:
-	FWallClockMusicClockDriver(UObject* WorldContextObj, UMidiFile* InTempoMap, float DefaultTempo, int32 DefaultTimeSigNum, int32 DefaultTimeSigDenom)
-		: FMusicClockDriverBase(WorldContextObj, DefaultTempo, DefaultTimeSigNum, DefaultTimeSigDenom)
-		, TempoMapMidi(InTempoMap)
+	FWallClockMusicClockDriver(UObject* WorldContextObj, const FWallClockMusicClockSettings& Settings)
+		: FMusicClockDriverBase(WorldContextObj, Settings)
+		, TempoMapMidi(Settings.TempoMap)
 	{}
 
 	virtual bool CalculateSongPosWithOffset(float MsOffset, ECalibratedMusicTimebase Timebase, FMidiSongPos& OutResult) const override;
diff --git a/Engine/Plugins/Runtime/Harmonix/Source/HarmonixMetasound/Public/HarmonixMetasound/Components/MusicClockComponent.h b/Engine/Plugins/Runtime/Harmonix/Source/HarmonixMetasound/Public/HarmonixMetasound/Components/MusicClockComponent.h
--- a/Engine/Plugins/Runtime/Harmonix/Source/HarmonixMetasound/Public/HarmonixMetasound/Components/MusicClockComponent.h
+++ b/Engine/Plugins/Runtime/Harmonix/Source/HarmonixMetasound/Public/HarmonixMetasound/Components/MusicClockComponent.h
@@ -8,7 +8,9 @@
 #include "HarmonixMidi/MidiSongPos.h"
 #include "HarmonixMidi/SongMaps.h"
 #include "MusicEnvironmentClockSource.h"
+#include "MusicClockTypes.h"
 #include "Templates/UniquePtr.h"
+#include "StructUtils/InstancedStruct.h"
 
 #include "MusicClockComponent.generated.h"
 
@@ -19,7 +21,7 @@ DECLARE_LOG_CATEGORY_EXTERN(LogMusicClock, Log, All)
 namespace Metasound { class FMetasoundGenerator; }
 struct FMusicalTimeSpan;
 class UAudioComponent;
-struct FMusicClockDriverBase;
+class FMusicClockDriverBase;
 
 UENUM(BlueprintType)
 enum class EMusicClockState : uint8
@@ -34,6 +36,7 @@ enum class EMusicClockDriveMethod : uint8
 {
 	WallClock,
 	MetaSound,
+	Custom
 };
 
 UENUM(BlueprintType)
@@ -82,15 +85,18 @@ public:
 	UPROPERTY(EditAnywhere, Category = "MusicClock")
 	EMusicClockDriveMethod DriveMethod = EMusicClockDriveMethod::MetaSound;
 
-	UPROPERTY(EditAnywhere, Category = "MusicClock", meta = (EditCondition = "DriveMethod == EMusicClockDriveMethod::MetaSound"))
+	UPROPERTY(EditAnywhere, Category = "MusicClock", meta = (EditCondition = "DriveMethod == EMusicClockDriveMethod::MetaSound", EditConditionHides))
 	FName MetasoundOutputName = "MIDI Clock";
 
-	UPROPERTY(VisibleAnywhere, BlueprintReadWrite, Category = "MusicClock", meta = (EditCondition = "DriveMethod == EMusicClockDriveMethod::MetaSound"))
+	UPROPERTY(VisibleAnywhere, BlueprintReadWrite, Category = "MusicClock", meta = (EditCondition = "DriveMethod == EMusicClockDriveMethod::MetaSound", EditConditionHides))
 	TObjectPtr<UAudioComponent> MetasoundsAudioComponent;
 
-	UPROPERTY(EditAnywhere, Category = "MusicClock", meta = (EditCondition = "DriveMethod == EMusicClockDriveMethod::WallClock"))
+	UPROPERTY(EditAnywhere, Category = "MusicClock", meta = (EditCondition = "DriveMethod == EMusicClockDriveMethod::WallClock", EditConditionHides))
 	TObjectPtr<UMidiFile> TempoMap;
 
+	UPROPERTY(EditAnywhere, Category = "MusicClock", NoClear, meta = (EditCondition = "DriveMethod == EMusicClockDriveMethod::Custom"))
+	TInstancedStruct<FMusicClockSettingsBase> MusicClockSettings = TInstancedStruct<FMetasoundMusicClockSettings>::Make();
+
 	/**
 	 * Start the clock ticking and tracking music time
 	 * Same as "Activate"
@@ -455,12 +461,20 @@ public:
 	UFUNCTION(BlueprintCallable, Category = "Audio|MusicClock", meta = (WorldContext = "WorldContextObject", DefaultToSelf = "WorldContextObject"))
 	static UE_API UMusicClockComponent* CreateWallClockDrivenMusicClock(UObject* WorldContextObject, UMidiFile* WithTempoMap, bool Start = true);
 
+	static UE_API UMusicClockComponent* CreateMusicClockWithSettings(UObject* WorldContextObject, TInstancedStruct<FMusicClockSettingsBase> MusicClockSettings, bool Start = true);
+
+	template<typename MusicClockSettingsType = FMusicClockSettingsBase>
+	static UMusicClockComponent* CreateMusicClock(UObject* WorldContextObject, const MusicClockSettingsType& InMusicClockSettings, bool Start = true)
+	{
+		return CreateMusicClockWithSettings(WorldContextObject, TInstancedStruct<FMusicClockSettingsBase>::Make(InMusicClockSettings), Start);
+	}
+
 	UFUNCTION(BlueprintCallable, Category = "MusicClock")
 	UE_API bool ConnectToMetasoundOnAudioComponent(UAudioComponent* InAudioComponent);
 
 	UFUNCTION(BlueprintCallable, Category = "MusicClock")
 	UE_API void ConnectToWallClockForMidi(UMidiFile* InTempoMap);
-
+	
 	UE_API FMidiSongPos CalculateSongPosWithOffset(float MsOffset, ECalibratedMusicTimebase Timebase = ECalibratedMusicTimebase::VideoRenderTime) const;
 
 	UE_API const FMidiSongPos& GetRawUnsmoothedAudioRenderPos() const;
@@ -497,6 +511,7 @@ private:
 	UE_API void MakeDefaultSongMap();
 	UE_API bool ConnectToMetasound();
 	UE_API void ConnectToWallClock();
+	UE_API void ConnectToCustomClock();
 	UE_API void DisconnectFromClockDriver();
 	
 	UE_API void EnsureClockIsValidForGameFrameFromSubsystem();
diff --git a/Engine/Plugins/Runtime/Harmonix/Source/HarmonixMetasound/Public/HarmonixMetasound/Components/MusicClockDriverBase.h b/Engine/Plugins/Runtime/Harmonix/Source/HarmonixMetasound/Public/HarmonixMetasound/Components/MusicClockDriverBase.h
--- a/Engine/Plugins/Runtime/Harmonix/Source/HarmonixMetasound/Public/HarmonixMetasound/Components/MusicClockDriverBase.h
+++ b/Engine/Plugins/Runtime/Harmonix/Source/HarmonixMetasound/Public/HarmonixMetasound/Components/MusicClockDriverBase.h
@@ -2,16 +2,20 @@
 #pragma once
 
 #include "HarmonixMetasound/Components/MusicClockComponent.h"
+#include "HarmonixMetasound/Components/MusicClockTypes.h"
 #include "Harmonix/MusicalTimebase.h"
 #include "HarmonixMidi/MidiSongPos.h"
 #include "HarmonixMidi/SongMaps.h"
 
+struct FMusicClockSettingsBase;
 struct ISongMapEvaluator;
 
-struct FMusicClockDriverBase : public TSharedFromThis<FMusicClockDriverBase>
+#define UE_API HARMONIXMETASOUND_API
+
+class UE_API FMusicClockDriverBase : public TSharedFromThis<FMusicClockDriverBase>
 {
 public:
-	FMusicClockDriverBase(UObject* WorldContextObject, float DefaultTempo, int32 DefaultTimeSigNum, int32 DefaultTimeSigDenom);
+	FMusicClockDriverBase(UObject* WorldContextObject, const FMusicClockSettingsBase& Settings);
 	virtual ~FMusicClockDriverBase();
 
 	void EnsureClockIsValidForGameFrame();
@@ -86,4 +90,6 @@ private:
 	EMusicClockState MusicClockState = EMusicClockState::Stopped;
 
 	TWeakObjectPtr<UObject> ContextObj;
-};
\ No newline at end of file
+};
+
+#undef UE_API
\ No newline at end of file
diff --git a/Engine/Plugins/Runtime/Harmonix/Source/HarmonixMetasound/Public/HarmonixMetasound/Components/MusicClockTypes.h b/Engine/Plugins/Runtime/Harmonix/Source/HarmonixMetasound/Public/HarmonixMetasound/Components/MusicClockTypes.h
new file mode 100644
--- /dev/null
+++ b/Engine/Plugins/Runtime/Harmonix/Source/HarmonixMetasound/Public/HarmonixMetasound/Components/MusicClockTypes.h
@@ -0,0 +1,58 @@
+﻿// Copyright Epic Games, Inc. All Rights Reserved.
+#pragma once
+
+#include "UObject/NameTypes.h"
+
+#include "MusicClockTypes.generated.h"
+
+class UAudioComponent;
+class UMidiFile;
+class FMusicClockDriverBase;
+
+USTRUCT()
+struct FMusicClockSettingsBase
+{
+	GENERATED_BODY()
+
+	UPROPERTY()
+	float DefaultTempo = 120.f;
+
+	UPROPERTY()
+	int32 DefaultTimeSigNumerator = 4;
+
+	UPROPERTY()
+	int32 DefaultTimeSigDenominator = 4;
+
+	virtual TSharedPtr<FMusicClockDriverBase> MakeInstance(UObject* WorldContextObj) const
+	{
+		unimplemented();
+		return nullptr;
+	}
+
+	virtual ~FMusicClockSettingsBase() {}
+};
+
+USTRUCT(DisplayName = "Wall Clock")
+struct FWallClockMusicClockSettings final : public FMusicClockSettingsBase
+{
+	GENERATED_BODY()
+	
+	UPROPERTY(EditAnywhere, Category = "MusicClock")
+	TObjectPtr<UMidiFile> TempoMap;
+
+	virtual TSharedPtr<FMusicClockDriverBase> MakeInstance(UObject* WorldContextObj) const override;
+};
+
+USTRUCT(DisplayName = "MetaSound Clock")
+struct FMetasoundMusicClockSettings final : public FMusicClockSettingsBase
+{
+	GENERATED_BODY()
+	
+	UPROPERTY(EditAnywhere, Category = "MusicClock")
+	FName MetasoundOutputName = "MIDI Clock";
+
+	UPROPERTY(VisibleAnywhere, Category = "MusicClock")
+	TObjectPtr<UAudioComponent> AudioComponent;
+
+	virtual TSharedPtr<FMusicClockDriverBase> MakeInstance(UObject* WorldContextObj) const override;
+};
\ No newline at end of file
