[Sparks] Converting Beat/BarDurationInSeconds to SecondsPerBeat/Bar to make the units more obvious.  Also a duration implies some fixed amount, this is an instantaneous first derivative.  Also add more useful and safe Beats/BarsPerSecond, which have the advantage of not going to infinity (with a zero at infinity, ouch) when the clock stops.

#tests PIE mainstage and jam in lobby

diff --git a/Engine/Plugins/Runtime/Harmonix/Source/HarmonixMetasound/Private/HarmonixMetasound/Components/MusicClockComponent.cpp b/Engine/Plugins/Runtime/Harmonix/Source/HarmonixMetasound/Private/HarmonixMetasound/Components/MusicClockComponent.cpp
--- a/Engine/Plugins/Runtime/Harmonix/Source/HarmonixMetasound/Private/HarmonixMetasound/Components/MusicClockComponent.cpp
+++ b/Engine/Plugins/Runtime/Harmonix/Source/HarmonixMetasound/Private/HarmonixMetasound/Components/MusicClockComponent.cpp
@@ -129,17 +129,6 @@ void UMusicClockComponent::SetDefaultTimeSignatureDenom(int Denom)
 	DefaultTimeSignatureDenom = Denom;
 }
 
-void UMusicClockComponent::CalculateBarAndBeatDurations(float Tempo, float AdvanceRate, int32 TimeSigNum, int32 TimeSigDenom, float& OutBarDuration, float& OutBeatDuration)
-{
-	OutBeatDuration = 0.0f;
-	OutBarDuration = 0.0f;
-	if (Tempo > 0.0f && AdvanceRate > 0.0f && TimeSigNum > 0 && TimeSigDenom)
-	{
-		OutBeatDuration = (60.0f / Tempo) / AdvanceRate;
-		OutBarDuration = ((TimeSigNum * OutBeatDuration) / (TimeSigDenom / 4.0f)) / AdvanceRate;
-	}
-}
-
 void UMusicClockComponent::EnsureClockIsValidForGameFrameFromSubsystem()
 {
 	if (ClockDriver)
@@ -401,6 +390,15 @@ float UMusicClockComponent::GetCurrentTempo() const
 	return 0.0f;
 }
 
+float UMusicClockComponent::GetCurrentBeatsPerMinute() const
+{
+	if (ClockDriver)
+	{
+		return ClockDriver->GetBeatsPerMinute();
+	}
+	return 0.0f;
+}
+
 void UMusicClockComponent::GetCurrentTimeSignature(int& OutNumerator, int& OutDenominator) const
 {
 	OutNumerator = 0;
@@ -412,20 +410,38 @@ void UMusicClockComponent::GetCurrentTimeSignature(int& OutNumerator, int& OutDe
 	}
 }
 
-float UMusicClockComponent::GetCurrentBarDurationSeconds() const
+float UMusicClockComponent::GetCurrentBarsPerSecond() const
+{
+	if (ClockDriver)
+	{
+		return ClockDriver->CurrentBarsPerSecond;
+	}
+	return 0;
+}
+
+float UMusicClockComponent::GetCurrentSecondsPerBar() const
+{
+	if (ClockDriver)
+	{
+		return (ClockDriver->CurrentBarsPerSecond != 0) ? (1.0f / ClockDriver->CurrentBarsPerSecond) : 0;
+	}
+	return 0;
+}
+
+float UMusicClockComponent::GetCurrentBeatsPerSecond() const
 {
 	if (ClockDriver)
 	{
-		return ClockDriver->CurrentBarDurationSec;
+		return ClockDriver->CurrentBeatsPerSecond;
 	}
 	return 0;
 }
 
-float UMusicClockComponent::GetCurrentBeatDurationSeconds() const
+float UMusicClockComponent::GetCurrentSecondsPerBeat() const
 {
 	if (ClockDriver)
 	{
-		return ClockDriver->CurrentBeatDurationSec;
+		return (ClockDriver->CurrentBeatsPerSecond != 0) ? (1.0f / ClockDriver->CurrentBeatsPerSecond) : 0;
 	}
 	return 0;
 }
diff --git a/Engine/Plugins/Runtime/Harmonix/Source/HarmonixMetasound/Private/HarmonixMetasound/Components/MusicClockDriverBase.cpp b/Engine/Plugins/Runtime/Harmonix/Source/HarmonixMetasound/Private/HarmonixMetasound/Components/MusicClockDriverBase.cpp
--- a/Engine/Plugins/Runtime/Harmonix/Source/HarmonixMetasound/Private/HarmonixMetasound/Components/MusicClockDriverBase.cpp
+++ b/Engine/Plugins/Runtime/Harmonix/Source/HarmonixMetasound/Private/HarmonixMetasound/Components/MusicClockDriverBase.cpp
@@ -11,15 +11,22 @@ void FMusicClockDriverBase::ResetDefaultSongMaps(float DefaultTempo, int32 Defau
 	DefaultMaps.GetBarMap().AddTimeSignatureAtBarIncludingCountIn(0, DefaultTimeSigNum, DefaultTimeSigDenom);
 }
 
-FMusicClockDriverBase::FMusicClockDriverBase(UObject* WorldContextObject, const FMusicClockSettingsBase& Settings)
+void FMusicClockDriverBase::UpdateBeatsAndBarsPerSecond()
 {
-	ResetDefaultSongMaps(Settings.DefaultTempo, Settings.DefaultTimeSigNumerator, Settings.DefaultTimeSigDenominator);
+	// Midi quarter note to beat conversion.
+	const float BPM = GetBeatsPerMinute();
+	CurrentBeatsPerSecond = (BPM / 60.0) * CurrentClockAdvanceRate;
+	CurrentBarsPerSecond = (TimeSignatureNum > 0) ? (CurrentBeatsPerSecond / TimeSignatureNum) : 0;
+}
+
+FMusicClockDriverBase::FMusicClockDriverBase(UObject* WorldContextObject, float DefaultTempo, int32 DefaultTimeSigNum, int32 DefaultTimeSigDenom)
+{
+	ResetDefaultSongMaps(DefaultTempo, DefaultTimeSigNum, DefaultTimeSigDenom);
 	ContextObj = TWeakObjectPtr<UObject>(WorldContextObject);
 }
 
 FMusicClockDriverBase::~FMusicClockDriverBase()
 {
-	
 }
 
 double FMusicClockDriverBase::GetWallClockTime() const
@@ -147,10 +154,18 @@ void FMusicClockDriverBase::UpdateMusicPlaybackRate(float InTempo, float InClock
 	// only update bar/beat durations if there have been any other timing changes
 	if (bHasAnyTimingChanges)
 	{
-		UMusicClockComponent::CalculateBarAndBeatDurations(Tempo, CurrentClockAdvanceRate, TimeSignatureNum, TimeSignatureDenom, CurrentBarDurationSec, CurrentBeatDurationSec);
+		UpdateBeatsAndBarsPerSecond();
 	}
 }
 
+float FMusicClockDriverBase::GetBeatsPerMinute() const
+{
+	// if we have a tempo set, we should have a time signature.
+	ensure(TimeSignatureDenom != 0 || Tempo == 0);
+    // Tempo is MidiTempo, which is in quarter notes per minute.
+	return Tempo * TimeSignatureDenom / 4.0f; 
+}
+
 void FMusicClockDriverBase::EnsureClockIsValidForGameFrame()
 {
 	//	Not for use outside the game thread.
diff --git a/Engine/Plugins/Runtime/Harmonix/Source/HarmonixMetasound/Public/HarmonixMetasound/Components/MusicClockComponent.h b/Engine/Plugins/Runtime/Harmonix/Source/HarmonixMetasound/Public/HarmonixMetasound/Components/MusicClockComponent.h
--- a/Engine/Plugins/Runtime/Harmonix/Source/HarmonixMetasound/Public/HarmonixMetasound/Components/MusicClockComponent.h
+++ b/Engine/Plugins/Runtime/Harmonix/Source/HarmonixMetasound/Public/HarmonixMetasound/Components/MusicClockComponent.h
@@ -124,12 +124,25 @@ public:
 	UE_API void Stop();
 
 	/**
-	 * Get The Current Tempo (BPM) of the runnig clock
+	 * Get The Current Tempo of the running clock
 	 * 0 if the clock is stopped
+	 * This is Midi Tempo, which is QuarterNotesPerMinute. 
+	 * This does not consider the advance rate,
+	 * if you want that, use GetBeats|BarsPerSecond.
 	 */
 	UFUNCTION(BlueprintPure, Category = "MusicClock")
 	UE_API float GetCurrentTempo() const;
 
+	/**
+	 * Get The Current BPM of the running clock
+	 * 0 if the clock is stopped
+	 * This is in units of our Beats, not Quarter notes.
+	 * This does not consider the advance rate,
+	 * if you want that, use GetBeats|BarsPerSecond.
+	 */
+	UFUNCTION(BlueprintPure, Category = "MusicClock")
+	UE_API float GetCurrentBeatsPerMinute() const;
+
 	/**
 	 * Get The Current Time Signature
 	 * 
@@ -139,22 +152,44 @@ public:
 	UE_API void GetCurrentTimeSignature(int& OutNumerator, int& OutDenominator) const;
 
 	/**
-	 * Get The Current Bar Duration in Seconds
+	 * Get The Current Bars per Second.
 	 * Based on the current tempo, advance rate and time signature
 	 * 
-	 * 0 if the clock is stopped
+	 * 0 if the clock is stopped.
 	 */
 	UFUNCTION(BlueprintPure, Category = "MusicClock")
-	UE_API float GetCurrentBarDurationSeconds() const;
+	UE_API float GetCurrentBarsPerSecond() const;
 
 	/**
-	 * Get The Current Beat Duration in Seconds
+	 * Get The Current Seconds per Bar
 	 * Based on the current tempo, advance rate and time signature
 	 * 
 	 * 0 if the clock is stopped
+	 * Don't use this method, use GetCurrentBarsPerSecond instead if possible because this method
+	 * returns zero when the clock is stopped rather than inf, which is what it should be.
 	 */
 	UFUNCTION(BlueprintPure, Category = "MusicClock")
-	UE_API float GetCurrentBeatDurationSeconds() const;
+	UE_API float GetCurrentSecondsPerBar() const;
+	
+	/**
+	 * Get The Current Beats per Second.
+	 * Based on the current tempo, advance rate and time signature
+	 * 
+	 * 0 if the clock is stopped.
+	 */
+	UFUNCTION(BlueprintPure, Category = "MusicClock")
+	UE_API float GetCurrentBeatsPerSecond() const;
+
+	/**
+	 * Get The Current Seconds per Beat
+	 * Based on the current tempo, advance rate and time signature
+	 * 
+	 * 0 if the clock is stopped.
+	 * Don't use this method, use GetCurrentBeatsPerSecond instead if possible because this method
+	 * returns zero when the clock is stopped rather than inf, which is what it should be.
+	 */
+	UFUNCTION(BlueprintPure, Category = "MusicClock")
+	UE_API float GetCurrentSecondsPerBeat() const;
 	
 	/**
      * Get The Current Clock advance rate
@@ -366,8 +401,6 @@ public:
 	// Set the default time signature denominator, only call when there's no tempo map.
 	UE_API void SetDefaultTimeSignatureDenom(int Denom);
 
-	static UE_API void CalculateBarAndBeatDurations(float Tempo, float AdvanceRate, int32 TimeSigNum, int32 TimeSigDenom, float& OutBarDuration, float& OutBeatDuration);
-
 protected:
 
 	/**
@@ -407,28 +440,9 @@ protected:
 	UPROPERTY(EditAnywhere, BlueprintReadWrite, Category = "MusicClock|Defaults")
 	int DefaultTimeSignatureDenom = 4;
 
-	UPROPERTY(Transient, BlueprintGetter=GetCurrentBeatDurationSeconds, Category = "MusicClock")
-	float CurrentBeatDurationSec = 0.0f;
-
-	UPROPERTY(Transient, BlueprintGetter=GetCurrentBarDurationSeconds, Category = "MusicClock")
-	float CurrentBarDurationSec = 0.0f;
-
 	UPROPERTY(EditAnywhere, BlueprintReadWrite, Category = "MusicClock")
 	bool RunPastMusicEnd = false;
 
-private:
-	// Don't let C++ access these directly! They are a blueprint convenience and only work because 
-	// they specify getter functions!
-	// NOTE These only have getters because a future refactor will get rid of them. -jwf + jake burga.
-	UPROPERTY(BlueprintGetter = GetCurrentSmoothedAudioRenderSongPos, Category = "MusicClock")
-	FMidiSongPos CurrentSmoothedAudioRenderSongPos;
-	UPROPERTY(BlueprintGetter = GetCurrentVideoRenderSongPos, Category = "MusicClock")
-	FMidiSongPos CurrentVideoRenderSongPos;
-	UPROPERTY(BlueprintGetter = GetCurrentPlayerExperiencedSongPos, Category = "MusicClock")
-	FMidiSongPos CurrentPlayerExperiencedSongPos;
-	UPROPERTY(BlueprintGetter = GetCurrentRawAudioRenderSongPos, Category = "MusicClock")
-	FMidiSongPos CurrentRawAudioRenderSongPos;
-	
 public:
 	// Getter functions for the Blueprint properties exposed above...
 	UFUNCTION(BlueprintCallable, BlueprintPure, Category = "MusicClock")
diff --git a/Engine/Plugins/Runtime/Harmonix/Source/HarmonixMetasound/Public/HarmonixMetasound/Components/MusicClockDriverBase.h b/Engine/Plugins/Runtime/Harmonix/Source/HarmonixMetasound/Public/HarmonixMetasound/Components/MusicClockDriverBase.h
--- a/Engine/Plugins/Runtime/Harmonix/Source/HarmonixMetasound/Public/HarmonixMetasound/Components/MusicClockDriverBase.h
+++ b/Engine/Plugins/Runtime/Harmonix/Source/HarmonixMetasound/Public/HarmonixMetasound/Components/MusicClockDriverBase.h
@@ -41,15 +41,19 @@ public:
 	virtual void Disconnect() = 0;
 
 	void UpdateMusicPlaybackRate(float InTempo, float InClockAdvanceRate, int32 InTimeSigNum, int32 InTimeSigDenom);
+
+	// This returns our BeatsPerMinute which is our true beats, not quarter notes.
+	float GetBeatsPerMinute() const;
 	
 	FSongMaps DefaultMaps;
 	
+	// Note that tempo is MidiTempo, quarter notes per minute.
 	float Tempo = 0.0f;
 	int TimeSignatureNum = 0;
 	int TimeSignatureDenom = 0;
 
-	float CurrentBeatDurationSec = 0.0f;
-	float CurrentBarDurationSec = 0.0f;
+	float CurrentBeatsPerSecond = 0.0f;
+	float CurrentBarsPerSecond = 0.0f;
 
 	FMidiSongPos CurrentSmoothedAudioRenderSongPos;
 	FMidiSongPos CurrentVideoRenderSongPos;
@@ -84,6 +88,7 @@ private:
 
 	void ResetDefaultSongMaps(float DefaultTempo, int32 DefaultTimeSigNum, int32 DefaultTimeSigDenom);
 	virtual bool RefreshCurrentSongPos() = 0;
+	void UpdateBeatsAndBarsPerSecond();
 
 	uint64 LastUpdateFrame = 0;
 
diff --git a/Engine/Plugins/Runtime/Harmonix/Source/HarmonixMidi/Private/HarmonixMidi/MidiSongPos.cpp b/Engine/Plugins/Runtime/Harmonix/Source/HarmonixMidi/Private/HarmonixMidi/MidiSongPos.cpp
--- a/Engine/Plugins/Runtime/Harmonix/Source/HarmonixMidi/Private/HarmonixMidi/MidiSongPos.cpp
+++ b/Engine/Plugins/Runtime/Harmonix/Source/HarmonixMidi/Private/HarmonixMidi/MidiSongPos.cpp
@@ -92,6 +92,14 @@ bool FMidiSongPos::operator==(const FMidiSongPos& rhs) const
 	return SongPosCmp(*this, rhs) == 0;
 }
 
+float FMidiSongPos::GetBeatsPerMinute() const
+{
+	// if we have a tempo set, we should have a time signature.
+	ensure(TimeSigDenominator != 0 || Tempo == 0);
+	return Tempo * TimeSigDenominator / 4.0f;
+}
+
+
 void FMidiSongPos::SetByTime(float InElapsedMs, float InBpm, int32 InTimeSigNumerator, int32 InTimeSigDenominator, int32 StartBar)
 {
 	// first some numbers we will need...
diff --git a/Engine/Plugins/Runtime/Harmonix/Source/HarmonixMidi/Public/HarmonixMidi/MidiSongPos.h b/Engine/Plugins/Runtime/Harmonix/Source/HarmonixMidi/Public/HarmonixMidi/MidiSongPos.h
--- a/Engine/Plugins/Runtime/Harmonix/Source/HarmonixMidi/Public/HarmonixMidi/MidiSongPos.h
+++ b/Engine/Plugins/Runtime/Harmonix/Source/HarmonixMidi/Public/HarmonixMidi/MidiSongPos.h
@@ -13,7 +13,7 @@
 struct ISongMapEvaluator;
 
 /////////////////////////////////////////////////////////////////////////////
-// Position within a song (midi info)
+ // Position within a song (midi info)
 //
 USTRUCT(BlueprintType, Meta = (DisplayName = "MIDI Song Position"))
 struct FMidiSongPos
@@ -40,7 +40,6 @@ public:
 	/** Tempo in Quarter Notes Per Minute (BPM) */
 	UPROPERTY(BlueprintReadOnly, Category = "MidiSongPosition")
 	float Tempo = 0.0f;
-
 	/** Total bars from the beginning of the song. */
 	UPROPERTY(BlueprintReadOnly, Category = "MidiSongPosition")
 	float BarsIncludingCountIn = 0.0f;
@@ -99,7 +98,10 @@ public:
 		IsSet = false;
 	}
 
-	float GetBeatsPerSecond() const { return Tempo / 60.0f; }
+	/** Tempo converted from Quarter Notes to our Beats, only valid if TimeSigDenominator is set. */
+	UE_API float GetBeatsPerMinute() const;
+
+	UE_API float GetBeatsPerSecond() const { return GetBeatsPerMinute() / 60.0f; }
 
 	// Has this been set by a system yet, or is it in the default state? 
 	bool IsValid() const { return IsSet; }
