[Harmonix] Fix merge conflict error

diff --git a/Engine/Plugins/Runtime/Harmonix/Source/HarmonixMetasound/Private/HarmonixMetasound/Components/MusicClockDriverBase.cpp b/Engine/Plugins/Runtime/Harmonix/Source/HarmonixMetasound/Private/HarmonixMetasound/Components/MusicClockDriverBase.cpp
--- a/Engine/Plugins/Runtime/Harmonix/Source/HarmonixMetasound/Private/HarmonixMetasound/Components/MusicClockDriverBase.cpp
+++ b/Engine/Plugins/Runtime/Harmonix/Source/HarmonixMetasound/Private/HarmonixMetasound/Components/MusicClockDriverBase.cpp
@@ -19,9 +19,9 @@ void FMusicClockDriverBase::UpdateBeatsAndBarsPerSecond()
 	CurrentBarsPerSecond = (TimeSignatureNum > 0) ? (CurrentBeatsPerSecond / TimeSignatureNum) : 0;
 }
 
-FMusicClockDriverBase::FMusicClockDriverBase(UObject* WorldContextObject, float DefaultTempo, int32 DefaultTimeSigNum, int32 DefaultTimeSigDenom)
+FMusicClockDriverBase::FMusicClockDriverBase(UObject* WorldContextObject, const FMusicClockSettingsBase& Settings)
 {
-	ResetDefaultSongMaps(DefaultTempo, DefaultTimeSigNum, DefaultTimeSigDenom);
+	ResetDefaultSongMaps(Settings.DefaultTempo, Settings.DefaultTimeSigNumerator, Settings.DefaultTimeSigDenominator);
 	ContextObj = TWeakObjectPtr<UObject>(WorldContextObject);
 }
 
