Fixup API macros

#rnx

diff --git a/Engine/Plugins/Runtime/Harmonix/Source/HarmonixMidi/Public/HarmonixMidi/MidiSongPos.h b/Engine/Plugins/Runtime/Harmonix/Source/HarmonixMidi/Public/HarmonixMidi/MidiSongPos.h
--- a/Engine/Plugins/Runtime/Harmonix/Source/HarmonixMidi/Public/HarmonixMidi/MidiSongPos.h
+++ b/Engine/Plugins/Runtime/Harmonix/Source/HarmonixMidi/Public/HarmonixMidi/MidiSongPos.h
@@ -101,7 +101,7 @@ public:
 	/** Tempo converted from Quarter Notes to our Beats, only valid if TimeSigDenominator is set. */
 	UE_API float GetBeatsPerMinute() const;
 
-	UE_API float GetBeatsPerSecond() const { return GetBeatsPerMinute() / 60.0f; }
+	float GetBeatsPerSecond() const { return GetBeatsPerMinute() / 60.0f; }
 
 	// Has this been set by a system yet, or is it in the default state? 
 	bool IsValid() const { return IsSet; }
