Harmonix Metasound "Midi Quantiz Trigger Node"

#rnx
#tests Audio MetaSounds Tests

diff --git a/Engine/Plugins/Runtime/Harmonix/Source/HarmonixMetasound/Private/HarmonixMetasound/Nodes/MidiQuantizeTriggerNode.cpp b/Engine/Plugins/Runtime/Harmonix/Source/HarmonixMetasound/Private/HarmonixMetasound/Nodes/MidiQuantizeTriggerNode.cpp
new file mode 100644
--- /dev/null
+++ b/Engine/Plugins/Runtime/Harmonix/Source/HarmonixMetasound/Private/HarmonixMetasound/Nodes/MidiQuantizeTriggerNode.cpp
@@ -0,0 +1,248 @@
+// Copyright Epic Games, Inc. All Rights Reserved.
+
+#include "MetasoundExecutableOperator.h"
+#include "MetasoundFacade.h"
+#include "MetasoundNodeInterface.h"
+#include "MetasoundParamHelper.h"
+#include "MetasoundStandardNodesCategories.h"
+
+#include "HarmonixMetasound/Common.h"
+#include "HarmonixMetasound/DataTypes/MidiClock.h"
+#include "HarmonixMetasound/MidiOps/PulseGenerator.h"
+#include "HarmonixMetasound/Nodes/MidiQuantizeTriggerNode.h"
+
+#define LOCTEXT_NAMESPACE "HarmonixMetaSound_MidiQuantizeTriggerNode"
+
+DEFINE_LOG_CATEGORY_STATIC(LogHarmonixMetasoundMidiQuantizeTriggerNode, Log, All);
+
+namespace HarmonixMetasound::Nodes::MidiQuantizeTriggerNode
+{
+	using namespace Metasound;
+
+	namespace Inputs
+	{
+		DEFINE_METASOUND_PARAM_ALIAS(MidiClock, CommonPinNames::Inputs::MidiClock);
+		DEFINE_OUTPUT_METASOUND_PARAM(TriggerInput, "Trigger Input", "The input trigger to quantize");
+		DEFINE_OUTPUT_METASOUND_PARAM(Quantization, "Subdivision", "The subdivision to quantize the trigger to")
+		DEFINE_OUTPUT_METASOUND_PARAM(Tolerance, "Tolerance", "The tolerance that will allow late triggers to process immediately (as percentage of the subdivision)")
+	}
+
+	namespace Outputs
+	{
+		DEFINE_OUTPUT_METASOUND_PARAM(TriggerOutput, "Trigger Out", "The output trigger to quantize")
+	}
+
+	const FNodeClassName& GetClassName()
+	{
+		static const FNodeClassName ClassName { HarmonixNodeNamespace, "MidiQuantizeTriggerNode", "" };
+		return ClassName;
+	}
+	
+	int32 GetCurrentMajorVersion()
+    {
+    	return 0;
+    }
+
+	struct FInputs
+	{
+		FMidiClockReadRef MidiClock;
+		FTriggerReadRef Trigger;
+		FEnumMidiClockSubdivisionQuantizationReadRef Quantization;
+		FFloatReadRef Tolerance;
+	};
+
+	struct FOutputs
+	{
+		FTriggerWriteRef Trigger;	
+	};
+	
+	class FMidiQuantizeTriggerOperator : public TExecutableOperator<FMidiQuantizeTriggerOperator>
+	{
+	public:
+
+		FMidiQuantizeTriggerOperator(const FBuildOperatorParams& InParams, FInputs&& InInputs, FOutputs&& InOutputs)
+			: Inputs(MoveTemp(InInputs)), Outputs(MoveTemp(InOutputs))
+		{
+			PulseGenerator.Enable(true);
+		}
+
+		static const FNodeClassMetadata& GetNodeInfo()
+		{
+			auto InitNodeInfo = []() -> FNodeClassMetadata
+			{
+				FNodeClassMetadata Info;
+				Info.ClassName        = GetClassName();
+				Info.MajorVersion     = GetCurrentMajorVersion();
+				Info.MinorVersion     = 1;
+				Info.DisplayName      = METASOUND_LOCTEXT("MidiQuantizeTriggerNode_DisplayName", "MIDI Quantize/Delay Trigger Node");
+				Info.Description      = METASOUND_LOCTEXT("MidiQuantizeTriggernode_Description", "Quantize an input trigger to a musical subdivision. Will delay the trigger to the next best subdivision or immediately play if it's close enough to the last subdivision based on the given tolerance.");
+				Info.Author           = PluginAuthor;
+				Info.PromptIfMissing  = PluginNodeMissingPrompt;
+				Info.DefaultInterface = GetVertexInterface();
+				Info.Keywords         = { MetasoundNodeCategories::Harmonix, NodeCategories::Music };
+				return Info;
+			};
+
+			static const FNodeClassMetadata Info = InitNodeInfo();
+			return Info;
+		}
+		
+		static const FVertexInterface& GetVertexInterface()
+		{
+			using namespace Metasound;
+
+			auto CreateVertexInterface = []() -> FVertexInterface
+			{
+				FInputVertexInterface InputVertexInterface
+				{
+					TInputDataVertex<FMidiClock>(METASOUND_GET_PARAM_NAME_AND_METADATA(Inputs::MidiClock)),
+					TInputDataVertex<FTrigger>(METASOUND_GET_PARAM_NAME_AND_METADATA(Inputs::TriggerInput)),
+					TInputDataVertex<FEnumMidiClockSubdivisionQuantizationType>(METASOUND_GET_PARAM_NAME_AND_METADATA(Inputs::Quantization), (int32)EMidiClockSubdivisionQuantization::Beat),
+					TInputDataVertex<float>(METASOUND_GET_PARAM_NAME_AND_METADATA(Inputs::Tolerance), 0.1f)
+				};
+
+				FOutputVertexInterface OutputVertexInterface
+				{
+					TOutputDataVertex<FTrigger>(METASOUND_GET_PARAM_NAME_AND_METADATA(Outputs::TriggerOutput))
+				};
+
+				return FVertexInterface(InputVertexInterface, OutputVertexInterface);
+			};
+
+			static FVertexInterface VertexInterface = CreateVertexInterface();
+			return VertexInterface;
+		}
+		
+		static TUniquePtr<IOperator> CreateOperator(const FBuildOperatorParams& InParams, FBuildResults& OutResults)
+		{
+			const FInputVertexInterfaceData& InputData = InParams.InputData;
+			
+			FInputs Inputs
+			{
+				InputData.GetOrCreateDefaultDataReadReference<FMidiClock>(Inputs::MidiClockName, InParams.OperatorSettings),
+				InputData.GetOrCreateDefaultDataReadReference<FTrigger>(Inputs::TriggerInputName, InParams.OperatorSettings),
+				InputData.GetOrCreateDefaultDataReadReference<FEnumMidiClockSubdivisionQuantizationType>(Inputs::QuantizationName, InParams.OperatorSettings),
+				InputData.GetOrCreateDefaultDataReadReference<float>(Inputs::ToleranceName, InParams.OperatorSettings),
+			};
+
+			FOutputs Outputs
+			{
+				FTriggerWriteRef::CreateNew(InParams.OperatorSettings)
+			};
+			
+			return MakeUnique<FMidiQuantizeTriggerOperator>(InParams, MoveTemp(Inputs), MoveTemp(Outputs));
+		}
+
+		void BindInputs(FInputVertexInterfaceData& InVertexData) override
+		{
+			InVertexData.BindReadVertex(Inputs::MidiClockName, Inputs.MidiClock);
+			InVertexData.BindReadVertex(Inputs::TriggerInputName, Inputs.Trigger);
+			InVertexData.BindReadVertex(Inputs::QuantizationName, Inputs.Quantization);
+			InVertexData.BindReadVertex(Inputs::ToleranceName, Inputs.Tolerance);
+		}
+
+		void BindOutputs(FOutputVertexInterfaceData& InVertexData) override
+		{
+			InVertexData.BindReadVertex(Outputs::TriggerOutputName, Outputs.Trigger);
+		}
+
+		void Reset(const FResetParams& Params)
+		{
+			Outputs.Trigger->Reset();
+
+			PulseGenerator.Reset();
+		}
+
+		void Execute()
+		{
+			using namespace HarmonixMetasound;
+			using namespace HarmonixMetasound::MidiClockMessageTypes;
+			using namespace Harmonix::Midi::Ops;
+
+			Outputs.Trigger->AdvanceBlock();
+			PulseGenerator.SetInterval({
+				*Inputs.Quantization,
+				*Inputs.Quantization,
+				1,
+				0
+			});
+
+			const float Tolerance = FMath::Clamp(*Inputs.Tolerance, 0.0f, 1.0f);
+			int32 ImmediateTriggerFrame = INDEX_NONE;
+			Inputs.Trigger->ExecuteBlock(
+				[](int32, int32) {},
+				[this, &ImmediateTriggerFrame, &Tolerance](int32 StartFrame, int32 EndFrame)
+				{
+					int32 TickAtFrame = Inputs.MidiClock->GetNextTickToProcessAtBlockFrame(StartFrame);
+					const ISongMapEvaluator& SongMaps = Inputs.MidiClock->GetSongMapEvaluator();
+
+					int32 LowerTick = 0;
+					int32 UpperTick = 0;
+					SongMaps.GetTicksForNearestSubdivision(TickAtFrame, *Inputs.Quantization, LowerTick, UpperTick);
+					int32 NearestTick =  ((TickAtFrame - LowerTick) < (UpperTick - TickAtFrame)) ? LowerTick : UpperTick;
+
+					// invalid tick range
+					if (LowerTick >= UpperTick)
+					{
+						return;
+					}
+
+					// the nearest tick is in the future, so we can cue it up to trigger with the subdivision trigger node
+					if (NearestTick >= TickAtFrame)
+					{
+						++TriggerCount;
+						ImmediateTriggerFrame = INDEX_NONE;
+					}
+					// the nearest tick is in the past, and the upper tick is in the future...
+					// see if the tick is within our tolerance
+					else if (UpperTick > TickAtFrame)
+					{
+						// keep track of the first trigger frame for future immediate triggers within this threshold.
+						// They should all happen at the same time
+						if (ImmediateTriggerFrame == INDEX_NONE)
+						{
+							ImmediateTriggerFrame = StartFrame;
+						}
+						
+						float PctInInterval = 2.0f * static_cast<float>(TickAtFrame - LowerTick) / static_cast<float>(UpperTick - LowerTick);
+						if (PctInInterval <= Tolerance)
+						{
+							Outputs.Trigger->TriggerFrame(ImmediateTriggerFrame);
+						}
+						else
+						{
+							// otherwise, cue it up for the future
+							++TriggerCount;
+						}
+					}
+				});
+
+			PulseGenerator.Process(*Inputs.MidiClock,
+				[this](const FPulseGenerator::FPulseInfo& Pulse)
+				{
+					if (TriggerCount > 0)
+					{
+						Outputs.Trigger->TriggerFrame(Pulse.BlockFrameIndex);
+					}
+					TriggerCount = 0;
+				},
+				[](const FPulseGenerator::FPulseInfo&) {});
+		}
+
+	private:
+
+		FInputs Inputs;
+		FOutputs Outputs;
+
+		int32 TriggerCount = 0;
+
+		//** DATA (current state)
+		Harmonix::Midi::Ops::FPulseGenerator PulseGenerator;
+
+	};
+
+	using FMidiQuantizeTriggerNode = Metasound::TNodeFacade<FMidiQuantizeTriggerOperator>;
+	METASOUND_REGISTER_NODE(FMidiQuantizeTriggerNode);
+}
+
+#undef LOCTEXT_NAMESPACE
\ No newline at end of file
diff --git a/Engine/Plugins/Runtime/Harmonix/Source/HarmonixMetasound/Public/HarmonixMetasound/Nodes/MidiQuantizeTriggerNode.h b/Engine/Plugins/Runtime/Harmonix/Source/HarmonixMetasound/Public/HarmonixMetasound/Nodes/MidiQuantizeTriggerNode.h
new file mode 100644
--- /dev/null
+++ b/Engine/Plugins/Runtime/Harmonix/Source/HarmonixMetasound/Public/HarmonixMetasound/Nodes/MidiQuantizeTriggerNode.h
@@ -0,0 +1,30 @@
+// Copyright Epic Games, Inc. All Rights Reserved.
+
+#pragma once
+
+#include "MetasoundNodeInterface.h"
+
+#include "HarmonixMetasound/Common.h"
+
+#define UE_API HARMONIXMETASOUND_API
+
+namespace HarmonixMetasound::Nodes::MidiQuantizeTriggerNode
+{
+	const UE_API Metasound::FNodeClassName& GetClassName();
+	UE_API int32 GetCurrentMajorVersion();
+	
+	namespace Inputs
+	{
+		DECLARE_METASOUND_PARAM_ALIAS(MidiClock);
+		DECLARE_METASOUND_PARAM_EXTERN(TriggerInput);
+		DECLARE_METASOUND_PARAM_EXTERN(Quantization)
+		DECLARE_METASOUND_PARAM_EXTERN(Tolerance)
+	}
+
+	namespace  Outputs
+	{
+		DECLARE_METASOUND_PARAM_EXTERN(TriggerOutput);
+	}
+}
+
+#undef UE_API
\ No newline at end of file
