Fix copy paste error.

Use the correct macros for metasound input definitions

#rnx

diff --git a/Engine/Plugins/Runtime/Harmonix/Source/HarmonixMetasound/Private/HarmonixMetasound/Nodes/MidiQuantizeTriggerNode.cpp b/Engine/Plugins/Runtime/Harmonix/Source/HarmonixMetasound/Private/HarmonixMetasound/Nodes/MidiQuantizeTriggerNode.cpp
--- a/Engine/Plugins/Runtime/Harmonix/Source/HarmonixMetasound/Private/HarmonixMetasound/Nodes/MidiQuantizeTriggerNode.cpp
+++ b/Engine/Plugins/Runtime/Harmonix/Source/HarmonixMetasound/Private/HarmonixMetasound/Nodes/MidiQuantizeTriggerNode.cpp
@@ -22,9 +22,9 @@ namespace HarmonixMetasound::Nodes::MidiQuantizeTriggerNode
 	namespace Inputs
 	{
 		DEFINE_METASOUND_PARAM_ALIAS(MidiClock, CommonPinNames::Inputs::MidiClock);
-		DEFINE_OUTPUT_METASOUND_PARAM(TriggerInput, "Trigger Input", "The input trigger to quantize");
-		DEFINE_OUTPUT_METASOUND_PARAM(Quantization, "Subdivision", "The subdivision to quantize the trigger to")
-		DEFINE_OUTPUT_METASOUND_PARAM(Tolerance, "Tolerance", "The tolerance that will allow late triggers to process immediately (as percentage of the subdivision)")
+		DEFINE_INPUT_METASOUND_PARAM(TriggerInput, "Trigger Input", "The input trigger to quantize");
+		DEFINE_INPUT_METASOUND_PARAM(Quantization, "Subdivision", "The subdivision to quantize the trigger to")
+		DEFINE_INPUT_METASOUND_PARAM(Tolerance, "Tolerance", "The tolerance that will allow late triggers to process immediately (as percentage of the subdivision)")
 	}
 
 	namespace Outputs
