Music Map Song Length Data

#rnx

diff --git a/Engine/Plugins/Runtime/Harmonix/Source/HarmonixMetasound/Private/HarmonixMetasound/DataTypes/HarmonixWaveMusicAsset.cpp b/Engine/Plugins/Runtime/Harmonix/Source/HarmonixMetasound/Private/HarmonixMetasound/DataTypes/HarmonixWaveMusicAsset.cpp
--- a/Engine/Plugins/Runtime/Harmonix/Source/HarmonixMetasound/Private/HarmonixMetasound/DataTypes/HarmonixWaveMusicAsset.cpp
+++ b/Engine/Plugins/Runtime/Harmonix/Source/HarmonixMetasound/Private/HarmonixMetasound/DataTypes/HarmonixWaveMusicAsset.cpp
@@ -26,7 +26,7 @@ namespace HarmonixWaveMusicAsset
 	EMetaSoundOutputAudioFormat GetOutputFormatFromChannelCount(int32 NumChannels)
 	{
 		switch (NumChannels)
-		{
+		{ 
 			case 1: return EMetaSoundOutputAudioFormat::Mono;
 			case 2: return EMetaSoundOutputAudioFormat::Stereo;
 			case 4: return EMetaSoundOutputAudioFormat::Quad;
@@ -143,6 +143,11 @@ namespace HarmonixWaveMusicAsset
 	}
 }
 
+float UHarmonixWaveMusicAsset::GetSongLengthSeconds() const
+{
+	return WaveFile->GetDuration();
+}
+
 #if WITH_EDITOR
 void UHarmonixWaveMusicAsset::PostEditChangeProperty(struct FPropertyChangedEvent& PropertyChangedEvent)
 {
diff --git a/Engine/Plugins/Runtime/Harmonix/Source/HarmonixMetasound/Public/HarmonixMetasound/DataTypes/HarmonixWaveMusicAsset.h b/Engine/Plugins/Runtime/Harmonix/Source/HarmonixMetasound/Public/HarmonixMetasound/DataTypes/HarmonixWaveMusicAsset.h
--- a/Engine/Plugins/Runtime/Harmonix/Source/HarmonixMetasound/Public/HarmonixMetasound/DataTypes/HarmonixWaveMusicAsset.h
+++ b/Engine/Plugins/Runtime/Harmonix/Source/HarmonixMetasound/Public/HarmonixMetasound/DataTypes/HarmonixWaveMusicAsset.h
@@ -30,6 +30,8 @@ public:
 	UPROPERTY(EditAnywhere, Category = "Music")
 	TObjectPtr<USoundWave> WaveFile;
 
+	UE_API virtual float GetSongLengthSeconds() const override;
+
 #if WITH_EDITOR
 	UE_API virtual void PostEditChangeProperty(struct FPropertyChangedEvent& PropertyChangedEvent) override;
 #endif
diff --git a/Engine/Plugins/Runtime/MusicEnvironment/Source/MusicEnvironment/Public/MusicTypes/MusicMapSource.h b/Engine/Plugins/Runtime/MusicEnvironment/Source/MusicEnvironment/Public/MusicTypes/MusicMapSource.h
--- a/Engine/Plugins/Runtime/MusicEnvironment/Source/MusicEnvironment/Public/MusicTypes/MusicMapSource.h
+++ b/Engine/Plugins/Runtime/MusicEnvironment/Source/MusicEnvironment/Public/MusicTypes/MusicMapSource.h
@@ -58,6 +58,9 @@ public:
 
 	UFUNCTION()
 	virtual FMarkerProviderResults GatherMarkers(const UFrameBasedMusicMap* Map) const { return FMarkerProviderResults(); }
+
+	UFUNCTION()
+	virtual float GetSongLengthSeconds() const { return 0.0f; }
 };
 
 #undef UE_API // MUSICENVIRONMENT_API
\ No newline at end of file
