Add logic to make MIDI Quantize Trigger node only trigger one trigger per subdivision

#rnx

diff --git a/Engine/Plugins/Runtime/Harmonix/Source/HarmonixMetasound/Private/HarmonixMetasound/Nodes/MidiQuantizeTriggerNode.cpp b/Engine/Plugins/Runtime/Harmonix/Source/HarmonixMetasound/Private/HarmonixMetasound/Nodes/MidiQuantizeTriggerNode.cpp
--- a/Engine/Plugins/Runtime/Harmonix/Source/HarmonixMetasound/Private/HarmonixMetasound/Nodes/MidiQuantizeTriggerNode.cpp
+++ b/Engine/Plugins/Runtime/Harmonix/Source/HarmonixMetasound/Private/HarmonixMetasound/Nodes/MidiQuantizeTriggerNode.cpp
@@ -149,7 +149,9 @@ namespace HarmonixMetasound::Nodes::MidiQuantizeTriggerNode
 		void Reset(const FResetParams& Params)
 		{
 			Outputs.Trigger->Reset();
-
+			PulseGenerator.Reset();
+			bCanTriggerImmediate = true;
+			bTriggerOnSubdivision = false;
 			PulseGenerator.Reset();
 		}
 
@@ -168,11 +170,17 @@ namespace HarmonixMetasound::Nodes::MidiQuantizeTriggerNode
 			});
 
 			const float Tolerance = FMath::Clamp(*Inputs.Tolerance, 0.0f, 1.0f);
-			int32 ImmediateTriggerFrame = INDEX_NONE;
 			Inputs.Trigger->ExecuteBlock(
 				[](int32, int32) {},
-				[this, &ImmediateTriggerFrame, &Tolerance](int32 StartFrame, int32 EndFrame)
+				[this, &Tolerance](int32 StartFrame, int32 EndFrame)
 				{
+					// we can ignore triggers if we have already cued up a trigger for the next subdivision
+					// AND we can't trigger anything immediately
+					if (bTriggerOnSubdivision && !bCanTriggerImmediate)
+					{
+						return;
+					}
+					
 					int32 TickAtFrame = Inputs.MidiClock->GetNextTickToProcessAtBlockFrame(StartFrame);
 					const ISongMapEvaluator& SongMaps = Inputs.MidiClock->GetSongMapEvaluator();
 
@@ -186,33 +194,36 @@ namespace HarmonixMetasound::Nodes::MidiQuantizeTriggerNode
 					{
 						return;
 					}
-
-					// the nearest tick is in the future, so we can cue it up to trigger with the subdivision trigger node
+					
 					if (NearestTick >= TickAtFrame)
 					{
-						++TriggerCount;
-						ImmediateTriggerFrame = INDEX_NONE;
+						// the nearest tick is in the future, so we can cue it up to trigger with the subdivision trigger node
+						// we also know at this point we can't trigger anything immediately anymore
+						bTriggerOnSubdivision = true;
+						bCanTriggerImmediate = false;
 					}
-					// the nearest tick is in the past, and the upper tick is in the future...
-					// see if the tick is within our tolerance
 					else if (UpperTick > TickAtFrame)
 					{
-						// keep track of the first trigger frame for future immediate triggers within this threshold.
-						// They should all happen at the same time
-						if (ImmediateTriggerFrame == INDEX_NONE)
-						{
-							ImmediateTriggerFrame = StartFrame;
-						}
-						
-						float PctInInterval = 2.0f * static_cast<float>(TickAtFrame - LowerTick) / static_cast<float>(UpperTick - LowerTick);
-						if (PctInInterval <= Tolerance)
+						if (!bCanTriggerImmediate)
 						{
-							Outputs.Trigger->TriggerFrame(ImmediateTriggerFrame);
+							bTriggerOnSubdivision = true;
 						}
 						else
 						{
-							// otherwise, cue it up for the future
-							++TriggerCount;
+							// the nearest tick is in the past, and the upper tick is in the future...
+							// see if the tick is within our tolerance
+							float PctInInterval = 2.0f * static_cast<float>(TickAtFrame - LowerTick) / static_cast<float>(UpperTick - LowerTick);
+							if (PctInInterval <= Tolerance)
+							{
+								Outputs.Trigger->TriggerFrame(StartFrame);
+							}
+							else
+							{
+								// otherwise, cue it up for the future
+								bTriggerOnSubdivision = true;
+							}
+							// prevent other immediate triggers from happening until the next subdivision
+							bCanTriggerImmediate = false;
 						}
 					}
 				});
@@ -220,11 +231,15 @@ namespace HarmonixMetasound::Nodes::MidiQuantizeTriggerNode
 			PulseGenerator.Process(*Inputs.MidiClock,
 				[this](const FPulseGenerator::FPulseInfo& Pulse)
 				{
-					if (TriggerCount > 0)
+					if (bTriggerOnSubdivision)
 					{
 						Outputs.Trigger->TriggerFrame(Pulse.BlockFrameIndex);
 					}
-					TriggerCount = 0;
+
+					// if we triggered this subdivision, don't allow immediate triggers to happen
+					// since it would be for the same subdivision
+					bCanTriggerImmediate = !bTriggerOnSubdivision;
+					bTriggerOnSubdivision = false;
 				},
 				[](const FPulseGenerator::FPulseInfo&) {});
 		}
@@ -234,7 +249,8 @@ namespace HarmonixMetasound::Nodes::MidiQuantizeTriggerNode
 		FInputs Inputs;
 		FOutputs Outputs;
 
-		int32 TriggerCount = 0;
+		bool bTriggerOnSubdivision = false;
+		bool bCanTriggerImmediate = true;
 
 		//** DATA (current state)
 		Harmonix::Midi::Ops::FPulseGenerator PulseGenerator;
