Fix Mac Nightly

#rnx

diff --git a/Engine/Plugins/Runtime/Harmonix/Source/HarmonixMetasound/Private/HarmonixMetasound/Nodes/SimpleSamplerNode.cpp b/Engine/Plugins/Runtime/Harmonix/Source/HarmonixMetasound/Private/HarmonixMetasound/Nodes/SimpleSamplerNode.cpp
--- a/Engine/Plugins/Runtime/Harmonix/Source/HarmonixMetasound/Private/HarmonixMetasound/Nodes/SimpleSamplerNode.cpp
+++ b/Engine/Plugins/Runtime/Harmonix/Source/HarmonixMetasound/Private/HarmonixMetasound/Nodes/SimpleSamplerNode.cpp
@@ -18,6 +18,11 @@ namespace HarmonixMetasound::SimpleSamplerNode
 {
 	using namespace Metasound;
 
+	/**
+	 * Maximum duration in seconds for capture and playback
+	 */
+	static constexpr float kMaxSampleSeconds = 10.0f;
+
 	namespace Inputs
 	{
 		DEFINE_METASOUND_PARAM_ALIAS(Enabled, CommonPinNames::Inputs::Enable);
@@ -88,7 +93,6 @@ namespace HarmonixMetasound::SimpleSamplerNode
 		check (InputInterface.Num() == Inputs::kNumNonAudioInputs);
 		check (NumChannels <= kMaxChannels);
 
-		TStringBuilder<32> NameStr;
 		for (int32 ChannelIndex = 0; ChannelIndex < NumChannels; ++ChannelIndex)
 		{
 #if WITH_EDITOR
@@ -147,7 +151,7 @@ namespace HarmonixMetasound::SimpleSamplerNode
 			, Outputs(MoveTemp(Outputs))
 		{
 			TSharedPtr<const FSimpleSamplerNodeOperatorData> InOperatorData = StaticCastSharedPtr<const FSimpleSamplerNodeOperatorData>(Params.Node.GetOperatorData());
-			MaxCaptureSeconds = InOperatorData ? InOperatorData->MaxCaptureSeconds : 5.0f;
+			MaxCaptureSeconds = InOperatorData ? FMath::Min(InOperatorData->MaxCaptureSeconds, kMaxSampleSeconds) : 5.0f;
 			Reset(Params);
 		}
 		
@@ -253,11 +257,16 @@ namespace HarmonixMetasound::SimpleSamplerNode
 
 		void BulkOutput(EInputOrSilence Mode, int32 StartFrameIndexInBlock = 0, int32 NumFrames = -1)
 		{
+			StartFrameIndexInBlock = FMath::Max(StartFrameIndexInBlock, 0);
+
 			if (NumFrames == -1)
 			{
 				NumFrames = Outputs.AudioOutputs[0]->Num() - StartFrameIndexInBlock;
 			}
-			check (StartFrameIndexInBlock + NumFrames <= Outputs.AudioOutputs[0]->Num());
+			else
+			{
+				NumFrames = FMath::Min(NumFrames, Outputs.AudioOutputs[0]->Num() - StartFrameIndexInBlock);
+			}
 
 			if (NumFrames <= 0)
 			{
@@ -279,7 +288,6 @@ namespace HarmonixMetasound::SimpleSamplerNode
 					FMemory::Memset(Outputs.AudioOutputs[ChannelIndex]->GetData() + StartFrameIndexInBlock, 0, sizeof(float) * NumFrames);
 				}
 			}
-
 		}
 
 		void Execute()
@@ -291,7 +299,8 @@ namespace HarmonixMetasound::SimpleSamplerNode
 				return;
 			}
 
-			CaptureLengthFrames = FMath::Min(FMath::CeilToInt32(*Inputs.CaptureDuration * CaptureBuffer.GetMaxConfig().GetSampleRate()), CaptureBuffer.GetMaxNumFrames());
+			float CaptureDurationSec = FMath::Clamp(*Inputs.CaptureDuration, 0.0f, kMaxSampleSeconds);
+			CaptureLengthFrames = FMath::Clamp(FMath::CeilToInt32(CaptureDurationSec * CaptureBuffer.GetMaxConfig().GetSampleRate()), 0, CaptureBuffer.GetMaxNumFrames());
 			if (CaptureLengthFrames <= CaptureBuffer.GetNumValidFrames())
 			{
 				// crop our capture buffer...
@@ -340,7 +349,8 @@ namespace HarmonixMetasound::SimpleSamplerNode
 					}
 
 					// Since we are starting... capture any new length at our input...
-					PlayLengthFrames = FMath::Min(FMath::CeilToInt32(*Inputs.PlayDuration * CaptureBuffer.GetMaxConfig().GetSampleRate()), CaptureLengthFrames);
+					float PlayDurationSec = FMath::Clamp(*Inputs.PlayDuration, 0.0f, kMaxSampleSeconds);
+					PlayLengthFrames = FMath::Clamp(FMath::CeilToInt32(PlayDurationSec * CaptureBuffer.GetMaxConfig().GetSampleRate()), 0, CaptureLengthFrames);
 					// ... and whether we should be playing backward...
 					if (*Inputs.Reverse)
 					{
@@ -479,7 +489,7 @@ namespace HarmonixMetasound::SimpleSamplerNode
 			int32 FramesAvailableToPlay = PlayLengthFrames - NextPlayFrame;
 			int32 FramesToPlay = FMath::Min(FramesWriteableInOutput, FramesAvailableToPlay);
 
-			if (FramesToPlay)
+			if (FramesToPlay > 0)
 			{
 				bCurrentlyPlaying = true;
 				if (!bPlayingReverse)
@@ -548,9 +558,11 @@ namespace HarmonixMetasound::SimpleSamplerNode
 			// we may have some silence to write at the end of playback...
 			if (!bCurrentlyPlaying)
 			{
+				int32 SilentFrameIndex = DestinationFrame + FramesToPlay;
 				for (int32 ChannelIndex = 0; ChannelIndex < CaptureBuffer.GetNumValidChannels(); ++ChannelIndex)
 				{
-					FMemory::Memset(Outputs.AudioOutputs[ChannelIndex]->GetData() + DestinationFrame + FramesToPlay, 0, (Outputs.AudioOutputs[0]->Num() - (DestinationFrame + FramesToPlay)) * sizeof(float));
+					int32 NumSilentFrames = Outputs.AudioOutputs[ChannelIndex]->Num() - SilentFrameIndex;
+					FMemory::Memset(Outputs.AudioOutputs[ChannelIndex]->GetData() + SilentFrameIndex, 0, NumSilentFrames * sizeof(float));
 				}
 			}
 
