Get rid of Metasound DisplayName/Tooltip in non-Editor builds, don't just set them to FText::GetEmpty()

Each of these generates a global ctor and dtor.

diff --git a/Engine/Plugins/Runtime/Harmonix/Source/HarmonixMetasound/Private/HarmonixMetasound/Nodes/DistortionNode.cpp b/Engine/Plugins/Runtime/Harmonix/Source/HarmonixMetasound/Private/HarmonixMetasound/Nodes/DistortionNode.cpp
--- a/Engine/Plugins/Runtime/Harmonix/Source/HarmonixMetasound/Private/HarmonixMetasound/Nodes/DistortionNode.cpp
+++ b/Engine/Plugins/Runtime/Harmonix/Source/HarmonixMetasound/Private/HarmonixMetasound/Nodes/DistortionNode.cpp
@@ -186,7 +186,11 @@ namespace HarmonixMetasound
 		static const FVertexInterface& GetVertexInterface()
 		{
 
+#if WITH_EDITOR
 #define METASOUND_GET_PARAM_NAME_WITH_INDEX_AND_METADATA_ADVANCED_DISPLAY(NAME, INDEX) METASOUND_GET_PARAM_NAME_WITH_INDEX(NAME, INDEX), FDataVertexMetadata { NAME##Tooltip, FText::Format(NAME##DisplayName, INDEX), true }
+#else
+#define METASOUND_GET_PARAM_NAME_WITH_INDEX_AND_METADATA_ADVANCED_DISPLAY(NAME, INDEX) METASOUND_GET_PARAM_NAME_WITH_INDEX(NAME, INDEX), FDataVertexMetadata { FText::GetEmpty(), FText::GetEmpty(), true }
+#endif
 
 			auto InitVertexInterface = []() -> FVertexInterface
 			{
diff --git a/Engine/Plugins/Runtime/Harmonix/Source/HarmonixMetasound/Private/HarmonixMetasound/Nodes/FusionSamplerNode.cpp b/Engine/Plugins/Runtime/Harmonix/Source/HarmonixMetasound/Private/HarmonixMetasound/Nodes/FusionSamplerNode.cpp
--- a/Engine/Plugins/Runtime/Harmonix/Source/HarmonixMetasound/Private/HarmonixMetasound/Nodes/FusionSamplerNode.cpp
+++ b/Engine/Plugins/Runtime/Harmonix/Source/HarmonixMetasound/Private/HarmonixMetasound/Nodes/FusionSamplerNode.cpp
@@ -450,8 +450,7 @@ namespace HarmonixMetasound
 			FFusionSamplerOperatorBase::GetInputVertexInterface(),
 			FOutputVertexInterface(
 				TOutputDataVertex<FAudioBuffer>(METASOUND_GET_PARAM_NAME_AND_METADATA(Outputs::AudioMono)),
-				TOutputDataVertex<FFusionSyncLink>(FusionSamplerNodePinNames::RenderSyncName,
-					{ FusionSamplerNodePinNames::RenderSyncTooltip, FusionSamplerNodePinNames::RenderSyncDisplayName, true })
+				TOutputDataVertex<FFusionSyncLink>(METASOUND_GET_PARAM_NAME_AND_METADATA_ADVANCED(FusionSamplerNodePinNames::RenderSync))
 			)
 		);
 		return Interface;
@@ -491,8 +490,7 @@ namespace HarmonixMetasound
 			FOutputVertexInterface(
 				TOutputDataVertex<FAudioBuffer>(METASOUND_GET_PARAM_NAME_AND_METADATA(Outputs::AudioLeft)),
 				TOutputDataVertex<FAudioBuffer>(METASOUND_GET_PARAM_NAME_AND_METADATA(Outputs::AudioRight)),
-				TOutputDataVertex<FFusionSyncLink>(FusionSamplerNodePinNames::RenderSyncName,
-					{ FusionSamplerNodePinNames::RenderSyncTooltip, FusionSamplerNodePinNames::RenderSyncDisplayName, true })
+				TOutputDataVertex<FFusionSyncLink>(METASOUND_GET_PARAM_NAME_AND_METADATA_ADVANCED(FusionSamplerNodePinNames::RenderSync))
 			)
 		);
 		return Interface;
diff --git a/Engine/Plugins/Runtime/Harmonix/Source/HarmonixMetasound/Public/HarmonixMetasound/Common.h b/Engine/Plugins/Runtime/Harmonix/Source/HarmonixMetasound/Public/HarmonixMetasound/Common.h
--- a/Engine/Plugins/Runtime/Harmonix/Source/HarmonixMetasound/Public/HarmonixMetasound/Common.h
+++ b/Engine/Plugins/Runtime/Harmonix/Source/HarmonixMetasound/Public/HarmonixMetasound/Common.h
@@ -19,6 +19,8 @@ namespace HarmonixMetasound
 		HARMONIXMETASOUND_API extern const FText Analysis;
 	}
 
+#if WITH_EDITOR
+
 #define DECLARE_METASOUND_PARAM_EXTERN(NAME)      \
 	HARMONIXMETASOUND_API extern const TCHAR* NAME##Name;       \
 	HARMONIXMETASOUND_API extern const FText NAME##Tooltip;     \
@@ -29,7 +31,6 @@ namespace HarmonixMetasound
 	HARMONIXMETASOUND_API extern const FText& NAME##Tooltip;     \
 	HARMONIXMETASOUND_API extern const FText& NAME##DisplayName;
 
-#if WITH_EDITOR
 #define DEFINE_INPUT_METASOUND_PARAM(NAME, NAME_TEXT, TOOLTIP_TEXT)						\
 	const TCHAR* NAME##Name = TEXT(NAME_TEXT);											\
 	const FText NAME##Tooltip = LOCTEXT(#NAME "InputTooltip", TOOLTIP_TEXT);			\
@@ -38,16 +39,6 @@ namespace HarmonixMetasound
 	const TCHAR* NAME##Name = TEXT(NAME_TEXT);											\
 	const FText NAME##Tooltip = LOCTEXT(#NAME "OutputTooltip", TOOLTIP_TEXT);			\
     const FText NAME##DisplayName = LOCTEXT(#NAME "OutputDisplayName", NAME_TEXT);
-#else 
-#define DEFINE_INPUT_METASOUND_PARAM(NAME, NAME_TEXT, TOOLTIP_TEXT)		\
-	const TCHAR* NAME##Name = TEXT(NAME_TEXT);							\
-	const FText NAME##Tooltip = FText::GetEmpty();						\
-	const FText NAME##DisplayName = FText::GetEmpty();
-#define DEFINE_OUTPUT_METASOUND_PARAM(NAME, NAME_TEXT, TOOLTIP_TEXT)	\
-	const TCHAR* NAME##Name = TEXT(NAME_TEXT);							\
-	const FText NAME##Tooltip = FText::GetEmpty();						\
-	const FText NAME##DisplayName = FText::GetEmpty();
-#endif // WITH_EDITOR
 
 // Alias an existing param definition. Helpful for re-using definitions while providing a public node API.
 #define DEFINE_METASOUND_PARAM_ALIAS(ALIAS_NAME, TARGET_NAME) \
@@ -55,6 +46,25 @@ namespace HarmonixMetasound
 	const FText& ALIAS_NAME##Tooltip = TARGET_NAME##Tooltip; \
 	const FText& ALIAS_NAME##DisplayName = TARGET_NAME##DisplayName; 
 
+#else // WITH_EDITOR
+
+#define DECLARE_METASOUND_PARAM_EXTERN(NAME)      \
+	HARMONIXMETASOUND_API extern const TCHAR* NAME##Name;
+
+#define DECLARE_METASOUND_PARAM_ALIAS(NAME)      \
+	HARMONIXMETASOUND_API extern const TCHAR* NAME##Name;
+
+#define DEFINE_INPUT_METASOUND_PARAM(NAME, NAME_TEXT, TOOLTIP_TEXT)		\
+	const TCHAR* NAME##Name = TEXT(NAME_TEXT);
+#define DEFINE_OUTPUT_METASOUND_PARAM(NAME, NAME_TEXT, TOOLTIP_TEXT)	\
+	const TCHAR* NAME##Name = TEXT(NAME_TEXT);
+
+// Alias an existing param definition. Helpful for re-using definitions while providing a public node API.
+#define DEFINE_METASOUND_PARAM_ALIAS(ALIAS_NAME, TARGET_NAME) \
+	const TCHAR* ALIAS_NAME##Name = TARGET_NAME##Name;
+
+#endif // WITH_EDITOR
+
 	namespace CommonPinNames
 	{
 		namespace Inputs
diff --git a/Engine/Plugins/Runtime/Metasound/Source/MetasoundGraphCore/Public/MetasoundParamHelper.h b/Engine/Plugins/Runtime/Metasound/Source/MetasoundGraphCore/Public/MetasoundParamHelper.h
--- a/Engine/Plugins/Runtime/Metasound/Source/MetasoundGraphCore/Public/MetasoundParamHelper.h
+++ b/Engine/Plugins/Runtime/Metasound/Source/MetasoundGraphCore/Public/MetasoundParamHelper.h
@@ -64,9 +64,7 @@
 	static const FText NAME##DisplayName = LOCTEXT(#NAME "DisplayName", NAME_TEXT);
 #else 
 #define DEFINE_METASOUND_PARAM(NAME, NAME_TEXT, TOOLTIP_TEXT) \
-	const FLazyName NAME##Name = TEXT(NAME_TEXT); \
-	static const FText NAME##Tooltip = FText::GetEmpty(); \
-	static const FText NAME##DisplayName = FText::GetEmpty();
+	const FLazyName NAME##Name = TEXT(NAME_TEXT);
 #endif // WITH_EDITOR
 
 #if WITH_EDITOR
@@ -76,19 +74,26 @@
 	static const FText NAME##DisplayName = LOCTEXT(#NAME "DisplayName", NAME_TEXT); 
 #else 
 #define METASOUND_PARAM(NAME, NAME_TEXT, TOOLTIP_TEXT) \
-	static const FLazyName NAME##Name = TEXT(NAME_TEXT); \
-	static const FText NAME##Tooltip = FText::GetEmpty(); \
-	static const FText NAME##DisplayName = FText::GetEmpty();
+	static const FLazyName NAME##Name = TEXT(NAME_TEXT);
 #endif // WITH_EDITOR
 
 #define METASOUND_GET_PARAM_NAME(NAME) NAME##Name
+
+#if WITH_EDITOR
 #define METASOUND_GET_PARAM_TT(NAME) NAME##Tooltip
 #define METASOUND_GET_PARAM_NAME_AND_TT(NAME) NAME##Name, NAME##Tooltip
 #define METASOUND_GET_PARAM_METADATA(NAME) FDataVertexMetadata { NAME##Tooltip, NAME##DisplayName }
+#define METASOUND_GET_PARAM_METADATA_ADVANCED(NAME) FDataVertexMetadata { NAME##Tooltip, NAME##DisplayName, true }
 #define METASOUND_GET_PARAM_DISPLAYNAME(NAME) NAME##DisplayName
+#else
+#define METASOUND_GET_PARAM_TT(NAME) FText::GetEmpty()
+#define METASOUND_GET_PARAM_NAME_AND_TT(NAME) NAME##Name, FText::GetEmpty()
+#define METASOUND_GET_PARAM_METADATA(NAME) FDataVertexMetadata { FText::GetEmpty(), FText::GetEmpty() }
+#define METASOUND_GET_PARAM_METADATA_ADVANCED(NAME) FDataVertexMetadata { FText::GetEmpty(), FText::GetEmpty(), true }
+#define METASOUND_GET_PARAM_DISPLAYNAME(NAME) FText::GetEmpty()
+#endif // WITH_EDITOR
 #define METASOUND_GET_PARAM_NAME_AND_METADATA(NAME) METASOUND_GET_PARAM_NAME(NAME), METASOUND_GET_PARAM_METADATA(NAME)
 
-#define METASOUND_GET_PARAM_METADATA_ADVANCED(NAME) FDataVertexMetadata { NAME##Tooltip, NAME##DisplayName, true }
 #define METASOUND_GET_PARAM_NAME_AND_METADATA_ADVANCED(NAME) METASOUND_GET_PARAM_NAME(NAME), METASOUND_GET_PARAM_METADATA_ADVANCED(NAME)
 
 #define METASOUND_GET_PARAM_NAME_WITH_INDEX(NAME, INDEX) *FString::Format(*static_cast<FName>(NAME##Name).ToString(), {INDEX})
