Run UnrealCodeFixup to add #include UE_INLINE_GENERATED_CPP_BY_NAME to files where possible

diff --git a/Engine/Plugins/AI/MassCrowd/Source/MassCrowd/Private/MassCrowdBubble.cpp b/Engine/Plugins/AI/MassCrowd/Source/MassCrowd/Private/MassCrowdBubble.cpp
--- a/Engine/Plugins/AI/MassCrowd/Source/MassCrowd/Private/MassCrowdBubble.cpp
+++ b/Engine/Plugins/AI/MassCrowd/Source/MassCrowd/Private/MassCrowdBubble.cpp
@@ -5,6 +5,8 @@
 #include "Net/UnrealNetwork.h"
 #include "MassExecutionContext.h"
 
+#include UE_INLINE_GENERATED_CPP_BY_NAME(MassCrowdBubble)
+
 namespace UE::Mass::Crowd
 {
 	bool bDebugReplicationPositions = false;
diff --git a/Engine/Plugins/AI/MassCrowd/Source/MassCrowd/Private/MassCrowdLaneClosingTest.cpp b/Engine/Plugins/AI/MassCrowd/Source/MassCrowd/Private/MassCrowdLaneClosingTest.cpp
--- a/Engine/Plugins/AI/MassCrowd/Source/MassCrowd/Private/MassCrowdLaneClosingTest.cpp
+++ b/Engine/Plugins/AI/MassCrowd/Source/MassCrowd/Private/MassCrowdLaneClosingTest.cpp
@@ -5,6 +5,8 @@
 #include "MassCrowdSubsystem.h"
 #include "ZoneGraphRenderingUtilities.h"
 
+#include UE_INLINE_GENERATED_CPP_BY_NAME(MassCrowdLaneClosingTest)
+
 void UZoneGraphCloseCrowdLaneTest::OnLaneLocationUpdated(const FZoneGraphLaneLocation& PrevLaneLocation, const FZoneGraphLaneLocation& NextLaneLocation)
 {
 	if (CrowdSubsystem == nullptr)
diff --git a/Engine/Plugins/AI/MassCrowd/Source/MassCrowd/Private/MassCrowdLaneDataRenderingComponent.cpp b/Engine/Plugins/AI/MassCrowd/Source/MassCrowd/Private/MassCrowdLaneDataRenderingComponent.cpp
--- a/Engine/Plugins/AI/MassCrowd/Source/MassCrowd/Private/MassCrowdLaneDataRenderingComponent.cpp
+++ b/Engine/Plugins/AI/MassCrowd/Source/MassCrowd/Private/MassCrowdLaneDataRenderingComponent.cpp
@@ -12,6 +12,8 @@
 #include "Debug/DebugDrawService.h"
 #include "Engine/Canvas.h"
 
+#include UE_INLINE_GENERATED_CPP_BY_NAME(MassCrowdLaneDataRenderingComponent)
+
 #if !UE_BUILD_SHIPPING && !UE_BUILD_TEST
 void UMassCrowdLaneDataRenderingComponent::OnRegister()
 {
diff --git a/Engine/Plugins/AI/MassCrowd/Source/MassCrowd/Private/MassCrowdMemberTrait.cpp b/Engine/Plugins/AI/MassCrowd/Source/MassCrowd/Private/MassCrowdMemberTrait.cpp
--- a/Engine/Plugins/AI/MassCrowd/Source/MassCrowd/Private/MassCrowdMemberTrait.cpp
+++ b/Engine/Plugins/AI/MassCrowd/Source/MassCrowd/Private/MassCrowdMemberTrait.cpp
@@ -4,6 +4,8 @@
 #include "MassCrowdFragments.h"
 #include "MassEntityTemplateRegistry.h"
 
+#include UE_INLINE_GENERATED_CPP_BY_NAME(MassCrowdMemberTrait)
+
 void UMassCrowdMemberTrait::BuildTemplate(FMassEntityTemplateBuildContext& BuildContext, const UWorld& World) const
 {
 	BuildContext.AddTag<FMassCrowdTag>();
diff --git a/Engine/Plugins/AI/MassCrowd/Source/MassCrowd/Private/MassCrowdNavigationProcessor.cpp b/Engine/Plugins/AI/MassCrowd/Source/MassCrowd/Private/MassCrowdNavigationProcessor.cpp
--- a/Engine/Plugins/AI/MassCrowd/Source/MassCrowd/Private/MassCrowdNavigationProcessor.cpp
+++ b/Engine/Plugins/AI/MassCrowd/Source/MassCrowd/Private/MassCrowdNavigationProcessor.cpp
@@ -18,6 +18,8 @@
 //----------------------------------------------------------------------//
 // UMassCrowdLaneTrackingSignalProcessor
 //----------------------------------------------------------------------//
+
+#include UE_INLINE_GENERATED_CPP_BY_NAME(MassCrowdNavigationProcessor)
 UMassCrowdLaneTrackingSignalProcessor::UMassCrowdLaneTrackingSignalProcessor()
 {
 	ExecutionOrder.ExecuteBefore.Add(UE::Mass::ProcessorGroupNames::Behavior);
diff --git a/Engine/Plugins/AI/MassCrowd/Source/MassCrowd/Private/MassCrowdReplicator.cpp b/Engine/Plugins/AI/MassCrowd/Source/MassCrowd/Private/MassCrowdReplicator.cpp
--- a/Engine/Plugins/AI/MassCrowd/Source/MassCrowd/Private/MassCrowdReplicator.cpp
+++ b/Engine/Plugins/AI/MassCrowd/Source/MassCrowd/Private/MassCrowdReplicator.cpp
@@ -7,6 +7,8 @@
 //----------------------------------------------------------------------//
 //  UMassCrowdReplicator
 //----------------------------------------------------------------------//
+
+#include UE_INLINE_GENERATED_CPP_BY_NAME(MassCrowdReplicator)
 void UMassCrowdReplicator::AddRequirements(FMassEntityQuery& EntityQuery)
 {
 	FMassReplicationProcessorPositionYawHandler::AddRequirements(EntityQuery);
diff --git a/Engine/Plugins/AI/MassCrowd/Source/MassCrowd/Private/MassCrowdRepresentationActorManagement.cpp b/Engine/Plugins/AI/MassCrowd/Source/MassCrowd/Private/MassCrowdRepresentationActorManagement.cpp
--- a/Engine/Plugins/AI/MassCrowd/Source/MassCrowd/Private/MassCrowdRepresentationActorManagement.cpp
+++ b/Engine/Plugins/AI/MassCrowd/Source/MassCrowd/Private/MassCrowdRepresentationActorManagement.cpp
@@ -9,6 +9,8 @@
 #include "MassAgentComponent.h"
 #include "MassRepresentationSubsystem.h"
 
+#include UE_INLINE_GENERATED_CPP_BY_NAME(MassCrowdRepresentationActorManagement)
+
 void UMassCrowdRepresentationActorManagement::SetActorEnabled(const EMassActorEnabledType EnabledType, AActor& Actor, const int32 EntityIdx, FMassCommandBuffer& CommandBuffer) const
 {
 	Super::SetActorEnabled(EnabledType, Actor, EntityIdx, CommandBuffer);
diff --git a/Engine/Plugins/AI/MassCrowd/Source/MassCrowd/Private/MassCrowdRepresentationSubsystem.cpp b/Engine/Plugins/AI/MassCrowd/Source/MassCrowd/Private/MassCrowdRepresentationSubsystem.cpp
--- a/Engine/Plugins/AI/MassCrowd/Source/MassCrowd/Private/MassCrowdRepresentationSubsystem.cpp
+++ b/Engine/Plugins/AI/MassCrowd/Source/MassCrowd/Private/MassCrowdRepresentationSubsystem.cpp
@@ -4,6 +4,8 @@
 #include "MassCrowdSpawnerSubsystem.h"
 #include "Engine/World.h"
 
+#include UE_INLINE_GENERATED_CPP_BY_NAME(MassCrowdRepresentationSubsystem)
+
 void UMassCrowdRepresentationSubsystem::Initialize(FSubsystemCollectionBase& Collection)
 {
 	Collection.InitializeDependency(UMassCrowdSpawnerSubsystem::StaticClass());
diff --git a/Engine/Plugins/AI/MassCrowd/Source/MassCrowd/Private/MassCrowdServerRepresentationLODProcessor.cpp b/Engine/Plugins/AI/MassCrowd/Source/MassCrowd/Private/MassCrowdServerRepresentationLODProcessor.cpp
--- a/Engine/Plugins/AI/MassCrowd/Source/MassCrowd/Private/MassCrowdServerRepresentationLODProcessor.cpp
+++ b/Engine/Plugins/AI/MassCrowd/Source/MassCrowd/Private/MassCrowdServerRepresentationLODProcessor.cpp
@@ -8,6 +8,8 @@
 #include "MassRepresentationFragments.h"
 #include "MassEntityManager.h"
 
+#include UE_INLINE_GENERATED_CPP_BY_NAME(MassCrowdServerRepresentationLODProcessor)
+
 namespace UE::MassCrowd
 {
 	int32 bDebugCrowdServerRepresentationLOD = 0;
diff --git a/Engine/Plugins/AI/MassCrowd/Source/MassCrowd/Private/MassCrowdServerRepresentationProcessor.cpp b/Engine/Plugins/AI/MassCrowd/Source/MassCrowd/Private/MassCrowdServerRepresentationProcessor.cpp
--- a/Engine/Plugins/AI/MassCrowd/Source/MassCrowd/Private/MassCrowdServerRepresentationProcessor.cpp
+++ b/Engine/Plugins/AI/MassCrowd/Source/MassCrowd/Private/MassCrowdServerRepresentationProcessor.cpp
@@ -4,6 +4,8 @@
 #include "MassCrowdFragments.h"
 #include "MassLODTypes.h"
 
+#include UE_INLINE_GENERATED_CPP_BY_NAME(MassCrowdServerRepresentationProcessor)
+
 UMassCrowdServerRepresentationProcessor::UMassCrowdServerRepresentationProcessor()
 {
 	ExecutionFlags = (int32)EProcessorExecutionFlags::Server;
diff --git a/Engine/Plugins/AI/MassCrowd/Source/MassCrowd/Private/MassCrowdServerRepresentationTrait.cpp b/Engine/Plugins/AI/MassCrowd/Source/MassCrowd/Private/MassCrowdServerRepresentationTrait.cpp
--- a/Engine/Plugins/AI/MassCrowd/Source/MassCrowd/Private/MassCrowdServerRepresentationTrait.cpp
+++ b/Engine/Plugins/AI/MassCrowd/Source/MassCrowd/Private/MassCrowdServerRepresentationTrait.cpp
@@ -8,6 +8,8 @@
 #include "MassActorSubsystem.h"
 #include "MassEntityUtils.h"
 
+#include UE_INLINE_GENERATED_CPP_BY_NAME(MassCrowdServerRepresentationTrait)
+
 
 UMassCrowdServerRepresentationTrait::UMassCrowdServerRepresentationTrait()
 {
diff --git a/Engine/Plugins/AI/MassCrowd/Source/MassCrowd/Private/MassCrowdSettings.cpp b/Engine/Plugins/AI/MassCrowd/Source/MassCrowd/Private/MassCrowdSettings.cpp
--- a/Engine/Plugins/AI/MassCrowd/Source/MassCrowd/Private/MassCrowdSettings.cpp
+++ b/Engine/Plugins/AI/MassCrowd/Source/MassCrowd/Private/MassCrowdSettings.cpp
@@ -2,6 +2,8 @@
 
 #include "MassCrowdSettings.h"
 
+#include UE_INLINE_GENERATED_CPP_BY_NAME(MassCrowdSettings)
+
 #if WITH_EDITOR
 void UMassCrowdSettings::PostEditChangeChainProperty(FPropertyChangedChainEvent& PropertyChangedEvent)
 {
diff --git a/Engine/Plugins/AI/MassCrowd/Source/MassCrowd/Private/MassCrowdSubsystem.cpp b/Engine/Plugins/AI/MassCrowd/Source/MassCrowd/Private/MassCrowdSubsystem.cpp
--- a/Engine/Plugins/AI/MassCrowd/Source/MassCrowd/Private/MassCrowdSubsystem.cpp
+++ b/Engine/Plugins/AI/MassCrowd/Source/MassCrowd/Private/MassCrowdSubsystem.cpp
@@ -12,6 +12,8 @@
 #include "ZoneGraphSubsystem.h"
 #include "VisualLogger/VisualLogger.h"
 
+#include UE_INLINE_GENERATED_CPP_BY_NAME(MassCrowdSubsystem)
+
 #if WITH_MASSGAMEPLAY_DEBUG
 namespace UE::MassCrowdDebug
 {
diff --git a/Engine/Plugins/AI/MassCrowd/Source/MassCrowd/Private/MassCrowdVisualizationLODProcessor.cpp b/Engine/Plugins/AI/MassCrowd/Source/MassCrowd/Private/MassCrowdVisualizationLODProcessor.cpp
--- a/Engine/Plugins/AI/MassCrowd/Source/MassCrowd/Private/MassCrowdVisualizationLODProcessor.cpp
+++ b/Engine/Plugins/AI/MassCrowd/Source/MassCrowd/Private/MassCrowdVisualizationLODProcessor.cpp
@@ -8,6 +8,8 @@
 #include "MassEntityManager.h"
 #include "MassRepresentationFragments.h"
 
+#include UE_INLINE_GENERATED_CPP_BY_NAME(MassCrowdVisualizationLODProcessor)
+
 namespace UE::MassCrowd
 {
 
diff --git a/Engine/Plugins/AI/MassCrowd/Source/MassCrowd/Private/MassCrowdVisualizationProcessor.cpp b/Engine/Plugins/AI/MassCrowd/Source/MassCrowd/Private/MassCrowdVisualizationProcessor.cpp
--- a/Engine/Plugins/AI/MassCrowd/Source/MassCrowd/Private/MassCrowdVisualizationProcessor.cpp
+++ b/Engine/Plugins/AI/MassCrowd/Source/MassCrowd/Private/MassCrowdVisualizationProcessor.cpp
@@ -6,6 +6,8 @@
 #include "MassActorSubsystem.h"
 #include "MassRepresentationFragments.h"
 
+#include UE_INLINE_GENERATED_CPP_BY_NAME(MassCrowdVisualizationProcessor)
+
 namespace UE::MassCrowd
 {
 	int32 bDebugCrowdVisualType = 0;
diff --git a/Engine/Plugins/AI/MassCrowd/Source/MassCrowd/Private/MassCrowdVisualizationTrait.cpp b/Engine/Plugins/AI/MassCrowd/Source/MassCrowd/Private/MassCrowdVisualizationTrait.cpp
--- a/Engine/Plugins/AI/MassCrowd/Source/MassCrowd/Private/MassCrowdVisualizationTrait.cpp
+++ b/Engine/Plugins/AI/MassCrowd/Source/MassCrowd/Private/MassCrowdVisualizationTrait.cpp
@@ -6,6 +6,8 @@
 #include "MassCrowdFragments.h"
 #include "MassEntityTemplateRegistry.h"
 
+#include UE_INLINE_GENERATED_CPP_BY_NAME(MassCrowdVisualizationTrait)
+
 UMassCrowdVisualizationTrait::UMassCrowdVisualizationTrait()
 {
 	// Override the subsystem to support parallelization of the crowd
diff --git a/Engine/Plugins/AI/MassCrowd/Source/MassCrowd/Private/Tasks/MassCrowdClaimWaitSlotTask.cpp b/Engine/Plugins/AI/MassCrowd/Source/MassCrowd/Private/Tasks/MassCrowdClaimWaitSlotTask.cpp
--- a/Engine/Plugins/AI/MassCrowd/Source/MassCrowd/Private/Tasks/MassCrowdClaimWaitSlotTask.cpp
+++ b/Engine/Plugins/AI/MassCrowd/Source/MassCrowd/Private/Tasks/MassCrowdClaimWaitSlotTask.cpp
@@ -8,6 +8,8 @@
 #include "StateTreeLinker.h"
 #include "MassStateTreeDependency.h"
 
+#include UE_INLINE_GENERATED_CPP_BY_NAME(MassCrowdClaimWaitSlotTask)
+
 FMassCrowdClaimWaitSlotTask::FMassCrowdClaimWaitSlotTask()
 {
 	// This task should not react to Enter/ExitState when the state is reselected.
diff --git a/Engine/Plugins/AI/MassCrowd/Source/MassCrowd/Private/Tasks/MassZoneGraphFindWanderTarget.cpp b/Engine/Plugins/AI/MassCrowd/Source/MassCrowd/Private/Tasks/MassZoneGraphFindWanderTarget.cpp
--- a/Engine/Plugins/AI/MassCrowd/Source/MassCrowd/Private/Tasks/MassZoneGraphFindWanderTarget.cpp
+++ b/Engine/Plugins/AI/MassCrowd/Source/MassCrowd/Private/Tasks/MassZoneGraphFindWanderTarget.cpp
@@ -13,6 +13,8 @@
 #include "MassDebugger.h"
 #include "MassStateTreeDependency.h"
 
+#include UE_INLINE_GENERATED_CPP_BY_NAME(MassZoneGraphFindWanderTarget)
+
 
 FMassZoneGraphFindWanderTarget::FMassZoneGraphFindWanderTarget()
 {
diff --git a/Engine/Plugins/AI/MassCrowd/Source/MassCrowd/Private/ZoneGraphCrowdLaneAnnotations.cpp b/Engine/Plugins/AI/MassCrowd/Source/MassCrowd/Private/ZoneGraphCrowdLaneAnnotations.cpp
--- a/Engine/Plugins/AI/MassCrowd/Source/MassCrowd/Private/ZoneGraphCrowdLaneAnnotations.cpp
+++ b/Engine/Plugins/AI/MassCrowd/Source/MassCrowd/Private/ZoneGraphCrowdLaneAnnotations.cpp
@@ -15,6 +15,8 @@
 #include "Engine/Canvas.h"
 #endif // UE_ENABLE_DEBUG_DRAWING
 
+#include UE_INLINE_GENERATED_CPP_BY_NAME(ZoneGraphCrowdLaneAnnotations)
+
 void UZoneGraphCrowdLaneAnnotations::PostSubsystemsInitialized()
 {
 	Super::PostSubsystemsInitialized();
diff --git a/Engine/Plugins/Experimental/LearningAgents/Source/Learning/Private/LearningNeuralNetwork.cpp b/Engine/Plugins/Experimental/LearningAgents/Source/Learning/Private/LearningNeuralNetwork.cpp
--- a/Engine/Plugins/Experimental/LearningAgents/Source/Learning/Private/LearningNeuralNetwork.cpp
+++ b/Engine/Plugins/Experimental/LearningAgents/Source/Learning/Private/LearningNeuralNetwork.cpp
@@ -11,6 +11,8 @@
 #include "NNEModelData.h"
 #include "NNERuntimeCPU.h"
 
+#include UE_INLINE_GENERATED_CPP_BY_NAME(LearningNeuralNetwork)
+
 void ULearningNeuralNetworkData::PostLoad()
 {
 	Super::PostLoad();
diff --git a/Engine/Plugins/Experimental/LearningAgents/Source/LearningAgents/Private/LearningAgentsActions.cpp b/Engine/Plugins/Experimental/LearningAgents/Source/LearningAgents/Private/LearningAgentsActions.cpp
--- a/Engine/Plugins/Experimental/LearningAgents/Source/LearningAgents/Private/LearningAgentsActions.cpp
+++ b/Engine/Plugins/Experimental/LearningAgents/Source/LearningAgents/Private/LearningAgentsActions.cpp
@@ -10,6 +10,8 @@
 
 #include "Containers/StaticArray.h"
 
+#include UE_INLINE_GENERATED_CPP_BY_NAME(LearningAgentsActions)
+
 bool operator==(const FLearningAgentsActionSchemaElement& Lhs, const FLearningAgentsActionSchemaElement& Rhs)
 {
 	return Lhs.SchemaElement.Index == Rhs.SchemaElement.Index;
diff --git a/Engine/Plugins/Experimental/LearningAgents/Source/LearningAgents/Private/LearningAgentsController.cpp b/Engine/Plugins/Experimental/LearningAgents/Source/LearningAgents/Private/LearningAgentsController.cpp
--- a/Engine/Plugins/Experimental/LearningAgents/Source/LearningAgents/Private/LearningAgentsController.cpp
+++ b/Engine/Plugins/Experimental/LearningAgents/Source/LearningAgents/Private/LearningAgentsController.cpp
@@ -7,6 +7,8 @@
 #include "LearningAgentsActions.h"
 #include "LearningLog.h"
 
+#include UE_INLINE_GENERATED_CPP_BY_NAME(LearningAgentsController)
+
 ULearningAgentsController::ULearningAgentsController() : Super(FObjectInitializer::Get()) {}
 ULearningAgentsController::ULearningAgentsController(FVTableHelper& Helper) : Super(Helper) {}
 ULearningAgentsController::~ULearningAgentsController() = default;
diff --git a/Engine/Plugins/Experimental/LearningAgents/Source/LearningAgents/Private/LearningAgentsCritic.cpp b/Engine/Plugins/Experimental/LearningAgents/Source/LearningAgents/Private/LearningAgentsCritic.cpp
--- a/Engine/Plugins/Experimental/LearningAgents/Source/LearningAgents/Private/LearningAgentsCritic.cpp
+++ b/Engine/Plugins/Experimental/LearningAgents/Source/LearningAgents/Private/LearningAgentsCritic.cpp
@@ -17,6 +17,8 @@
 
 #include "NNERuntimeBasicCpuBuilder.h"
 
+#include UE_INLINE_GENERATED_CPP_BY_NAME(LearningAgentsCritic)
+
 namespace UE::Learning::Agents::Critic::Private
 {
 	static inline UE::NNE::RuntimeBasic::FModelBuilder::EActivationFunction GetBuilderActivationFunction(const ELearningAgentsActivationFunction ActivationFunction)
diff --git a/Engine/Plugins/Experimental/LearningAgents/Source/LearningAgents/Private/LearningAgentsInteractor.cpp b/Engine/Plugins/Experimental/LearningAgents/Source/LearningAgents/Private/LearningAgentsInteractor.cpp
--- a/Engine/Plugins/Experimental/LearningAgents/Source/LearningAgents/Private/LearningAgentsInteractor.cpp
+++ b/Engine/Plugins/Experimental/LearningAgents/Source/LearningAgents/Private/LearningAgentsInteractor.cpp
@@ -13,6 +13,8 @@
 
 #include "EngineDefines.h"
 
+#include UE_INLINE_GENERATED_CPP_BY_NAME(LearningAgentsInteractor)
+
 ULearningAgentsInteractor::ULearningAgentsInteractor() : Super(FObjectInitializer::Get()) {}
 ULearningAgentsInteractor::ULearningAgentsInteractor(FVTableHelper& Helper) : Super(Helper) {}
 ULearningAgentsInteractor::~ULearningAgentsInteractor() = default;
diff --git a/Engine/Plugins/Experimental/LearningAgents/Source/LearningAgents/Private/LearningAgentsManager.cpp b/Engine/Plugins/Experimental/LearningAgents/Source/LearningAgents/Private/LearningAgentsManager.cpp
--- a/Engine/Plugins/Experimental/LearningAgents/Source/LearningAgents/Private/LearningAgentsManager.cpp
+++ b/Engine/Plugins/Experimental/LearningAgents/Source/LearningAgents/Private/LearningAgentsManager.cpp
@@ -10,6 +10,8 @@
 #include "EngineDefines.h"
 #include "Async/Async.h"
 
+#include UE_INLINE_GENERATED_CPP_BY_NAME(LearningAgentsManager)
+
 ULearningAgentsManager::ULearningAgentsManager() : Super(FObjectInitializer::Get()) { PrimaryComponentTick.bCanEverTick = true; }
 ULearningAgentsManager::ULearningAgentsManager(FVTableHelper& Helper) : Super(Helper) {}
 ULearningAgentsManager::~ULearningAgentsManager() = default;
@@ -497,4 +499,4 @@ void ULearningAgentsManager::UpdateAgentSets()
 	VacantAgentIds.Sort([](const int32& Lhs, const int32& Rhs) { return Lhs > Rhs; });
 	VacantAgentSet = VacantAgentIds;
 	VacantAgentSet.TryMakeSlice();
-}
\ No newline at end of file
+}
diff --git a/Engine/Plugins/Experimental/LearningAgents/Source/LearningAgents/Private/LearningAgentsManagerListener.cpp b/Engine/Plugins/Experimental/LearningAgents/Source/LearningAgents/Private/LearningAgentsManagerListener.cpp
--- a/Engine/Plugins/Experimental/LearningAgents/Source/LearningAgents/Private/LearningAgentsManagerListener.cpp
+++ b/Engine/Plugins/Experimental/LearningAgents/Source/LearningAgents/Private/LearningAgentsManagerListener.cpp
@@ -5,6 +5,8 @@
 #include "LearningAgentsManager.h"
 #include "LearningLog.h"
 
+#include UE_INLINE_GENERATED_CPP_BY_NAME(LearningAgentsManagerListener)
+
 ULearningAgentsManagerListener::ULearningAgentsManagerListener(const FObjectInitializer& ObjectInitializer) : Super(ObjectInitializer) {}
 ULearningAgentsManagerListener::ULearningAgentsManagerListener(FVTableHelper& Helper) : Super(Helper) {}
 ULearningAgentsManagerListener::~ULearningAgentsManagerListener() = default;
diff --git a/Engine/Plugins/Experimental/LearningAgents/Source/LearningAgents/Private/LearningAgentsNeuralNetwork.cpp b/Engine/Plugins/Experimental/LearningAgents/Source/LearningAgents/Private/LearningAgentsNeuralNetwork.cpp
--- a/Engine/Plugins/Experimental/LearningAgents/Source/LearningAgents/Private/LearningAgentsNeuralNetwork.cpp
+++ b/Engine/Plugins/Experimental/LearningAgents/Source/LearningAgents/Private/LearningAgentsNeuralNetwork.cpp
@@ -9,6 +9,8 @@
 #include "UObject/Package.h"
 #include "Misc/FileHelper.h"
 
+#include UE_INLINE_GENERATED_CPP_BY_NAME(LearningAgentsNeuralNetwork)
+
 ULearningAgentsNeuralNetwork::ULearningAgentsNeuralNetwork() = default;
 ULearningAgentsNeuralNetwork::ULearningAgentsNeuralNetwork(FVTableHelper& Helper) : Super(Helper) {}
 ULearningAgentsNeuralNetwork::~ULearningAgentsNeuralNetwork() = default;
diff --git a/Engine/Plugins/Experimental/LearningAgents/Source/LearningAgents/Private/LearningAgentsObservations.cpp b/Engine/Plugins/Experimental/LearningAgents/Source/LearningAgents/Private/LearningAgentsObservations.cpp
--- a/Engine/Plugins/Experimental/LearningAgents/Source/LearningAgents/Private/LearningAgentsObservations.cpp
+++ b/Engine/Plugins/Experimental/LearningAgents/Source/LearningAgents/Private/LearningAgentsObservations.cpp
@@ -14,6 +14,8 @@
 #include "Engine/World.h"
 #include "Components/SplineComponent.h"
 
+#include UE_INLINE_GENERATED_CPP_BY_NAME(LearningAgentsObservations)
+
 bool operator==(const FLearningAgentsObservationSchemaElement& Lhs, const FLearningAgentsObservationSchemaElement& Rhs)
 {
 	return Lhs.SchemaElement.Index == Rhs.SchemaElement.Index;
diff --git a/Engine/Plugins/Experimental/LearningAgents/Source/LearningAgents/Private/LearningAgentsPolicy.cpp b/Engine/Plugins/Experimental/LearningAgents/Source/LearningAgents/Private/LearningAgentsPolicy.cpp
--- a/Engine/Plugins/Experimental/LearningAgents/Source/LearningAgents/Private/LearningAgentsPolicy.cpp
+++ b/Engine/Plugins/Experimental/LearningAgents/Source/LearningAgents/Private/LearningAgentsPolicy.cpp
@@ -16,6 +16,8 @@
 
 #include "NNERuntimeBasicCpuBuilder.h"
 
+#include UE_INLINE_GENERATED_CPP_BY_NAME(LearningAgentsPolicy)
+
 namespace UE::Learning::Agents::Policy::Private
 {
 	static inline UE::NNE::RuntimeBasic::FModelBuilder::EActivationFunction GetBuilderActivationFunction(const ELearningAgentsActivationFunction ActivationFunction)
diff --git a/Engine/Plugins/Experimental/LearningAgents/Source/LearningAgentsTraining/Private/LearningAgentsCommunicator.cpp b/Engine/Plugins/Experimental/LearningAgents/Source/LearningAgentsTraining/Private/LearningAgentsCommunicator.cpp
--- a/Engine/Plugins/Experimental/LearningAgents/Source/LearningAgentsTraining/Private/LearningAgentsCommunicator.cpp
+++ b/Engine/Plugins/Experimental/LearningAgents/Source/LearningAgentsTraining/Private/LearningAgentsCommunicator.cpp
@@ -9,6 +9,8 @@
 #include "Misc/Guid.h"
 #include "Misc/Paths.h"
 
+#include UE_INLINE_GENERATED_CPP_BY_NAME(LearningAgentsCommunicator)
+
 FLearningAgentsCommunicator ULearningAgentsCommunicatorLibrary::MakeSharedMemoryTrainingProcess(
 	const FLearningAgentsTrainerProcessSettings& TrainerProcessSettings,
 	const FLearningAgentsSharedMemoryCommunicatorSettings& SharedMemorySettings)
diff --git a/Engine/Plugins/Experimental/LearningAgents/Source/LearningAgentsTraining/Private/LearningAgentsCompletions.cpp b/Engine/Plugins/Experimental/LearningAgents/Source/LearningAgentsTraining/Private/LearningAgentsCompletions.cpp
--- a/Engine/Plugins/Experimental/LearningAgents/Source/LearningAgentsTraining/Private/LearningAgentsCompletions.cpp
+++ b/Engine/Plugins/Experimental/LearningAgents/Source/LearningAgentsTraining/Private/LearningAgentsCompletions.cpp
@@ -8,6 +8,8 @@
 #include "LearningCompletion.h"
 #include "LearningLog.h"
 
+#include UE_INLINE_GENERATED_CPP_BY_NAME(LearningAgentsCompletions)
+
 namespace UE::Learning::Agents
 {
 	ELearningAgentsCompletion GetLearningAgentsCompletion(const ECompletionMode CompletionMode)
diff --git a/Engine/Plugins/Experimental/LearningAgents/Source/LearningAgentsTraining/Private/LearningAgentsEntitiesManagerComponent.cpp b/Engine/Plugins/Experimental/LearningAgents/Source/LearningAgentsTraining/Private/LearningAgentsEntitiesManagerComponent.cpp
--- a/Engine/Plugins/Experimental/LearningAgents/Source/LearningAgentsTraining/Private/LearningAgentsEntitiesManagerComponent.cpp
+++ b/Engine/Plugins/Experimental/LearningAgents/Source/LearningAgentsTraining/Private/LearningAgentsEntitiesManagerComponent.cpp
@@ -6,6 +6,8 @@
 #include "LearningAgentsGym.h"
 #include "LearningAgentsEntityInterface.h"
 
+#include UE_INLINE_GENERATED_CPP_BY_NAME(LearningAgentsEntitiesManagerComponent)
+
 void ULearningAgentsEntitiesManagerComponent::InitializeLearningComponent()
 {
 	TObjectPtr<ALearningAgentsGymBase> Gym = Cast<ALearningAgentsGymBase>(GetOwner());
diff --git a/Engine/Plugins/Experimental/LearningAgents/Source/LearningAgentsTraining/Private/LearningAgentsGym.cpp b/Engine/Plugins/Experimental/LearningAgents/Source/LearningAgentsTraining/Private/LearningAgentsGym.cpp
--- a/Engine/Plugins/Experimental/LearningAgents/Source/LearningAgentsTraining/Private/LearningAgentsGym.cpp
+++ b/Engine/Plugins/Experimental/LearningAgents/Source/LearningAgentsTraining/Private/LearningAgentsGym.cpp
@@ -7,6 +7,8 @@
 #include "LearningAgentsLearningComponentInterface.h"
 #include "Kismet/KismetMathLibrary.h"
 
+#include UE_INLINE_GENERATED_CPP_BY_NAME(LearningAgentsGym)
+
 ALearningAgentsGymBase::ALearningAgentsGymBase()
 {
 	PrimaryActorTick.bCanEverTick = false;
diff --git a/Engine/Plugins/Experimental/LearningAgents/Source/LearningAgentsTraining/Private/LearningAgentsGymSimple.cpp b/Engine/Plugins/Experimental/LearningAgents/Source/LearningAgentsTraining/Private/LearningAgentsGymSimple.cpp
--- a/Engine/Plugins/Experimental/LearningAgents/Source/LearningAgentsTraining/Private/LearningAgentsGymSimple.cpp
+++ b/Engine/Plugins/Experimental/LearningAgents/Source/LearningAgentsTraining/Private/LearningAgentsGymSimple.cpp
@@ -10,6 +10,8 @@
 #include "Components/StaticMeshComponent.h"
 #include "Kismet/KismetMathLibrary.h"
 
+#include UE_INLINE_GENERATED_CPP_BY_NAME(LearningAgentsGymSimple)
+
 ALearningAgentsGymSimple::ALearningAgentsGymSimple()
 {
 	RootComponent = CreateDefaultSubobject<USceneComponent>(TEXT("RootComponent"));
diff --git a/Engine/Plugins/Experimental/LearningAgents/Source/LearningAgentsTraining/Private/LearningAgentsGymsManager.cpp b/Engine/Plugins/Experimental/LearningAgents/Source/LearningAgentsTraining/Private/LearningAgentsGymsManager.cpp
--- a/Engine/Plugins/Experimental/LearningAgents/Source/LearningAgentsTraining/Private/LearningAgentsGymsManager.cpp
+++ b/Engine/Plugins/Experimental/LearningAgents/Source/LearningAgentsTraining/Private/LearningAgentsGymsManager.cpp
@@ -11,6 +11,8 @@
 #include "Misc/DateTime.h"
 #include "Math/UnrealMathUtility.h"
 
+#include UE_INLINE_GENERATED_CPP_BY_NAME(LearningAgentsGymsManager)
+
 ALearningAgentsGymsManager::ALearningAgentsGymsManager()
 {
 	PrimaryActorTick.bCanEverTick = false;
diff --git a/Engine/Plugins/Experimental/LearningAgents/Source/LearningAgentsTraining/Private/LearningAgentsImitationTrainer.cpp b/Engine/Plugins/Experimental/LearningAgents/Source/LearningAgentsTraining/Private/LearningAgentsImitationTrainer.cpp
--- a/Engine/Plugins/Experimental/LearningAgents/Source/LearningAgentsTraining/Private/LearningAgentsImitationTrainer.cpp
+++ b/Engine/Plugins/Experimental/LearningAgents/Source/LearningAgentsTraining/Private/LearningAgentsImitationTrainer.cpp
@@ -18,6 +18,8 @@
 #include "HAL/PlatformMisc.h"
 #include "Misc/Paths.h"
 
+#include UE_INLINE_GENERATED_CPP_BY_NAME(LearningAgentsImitationTrainer)
+
 TSharedRef<FJsonObject> FLearningAgentsImitationTrainerTrainingSettings::AsJsonConfig() const
 {
 	TSharedRef<FJsonObject> ConfigObject = MakeShared<FJsonObject>();
diff --git a/Engine/Plugins/Experimental/LearningAgents/Source/LearningAgentsTraining/Private/LearningAgentsPPOTrainer.cpp b/Engine/Plugins/Experimental/LearningAgents/Source/LearningAgentsTraining/Private/LearningAgentsPPOTrainer.cpp
--- a/Engine/Plugins/Experimental/LearningAgents/Source/LearningAgentsTraining/Private/LearningAgentsPPOTrainer.cpp
+++ b/Engine/Plugins/Experimental/LearningAgents/Source/LearningAgentsTraining/Private/LearningAgentsPPOTrainer.cpp
@@ -18,6 +18,8 @@
 #include "Dom/JsonObject.h"
 #include "HAL/PlatformMisc.h"
 
+#include UE_INLINE_GENERATED_CPP_BY_NAME(LearningAgentsPPOTrainer)
+
 ULearningAgentsPPOTrainer::ULearningAgentsPPOTrainer() : Super(FObjectInitializer::Get()) {}
 ULearningAgentsPPOTrainer::ULearningAgentsPPOTrainer(FVTableHelper& Helper) : Super(Helper) {}
 ULearningAgentsPPOTrainer::~ULearningAgentsPPOTrainer() = default;
diff --git a/Engine/Plugins/Experimental/LearningAgents/Source/LearningAgentsTraining/Private/LearningAgentsRecorder.cpp b/Engine/Plugins/Experimental/LearningAgents/Source/LearningAgentsTraining/Private/LearningAgentsRecorder.cpp
--- a/Engine/Plugins/Experimental/LearningAgents/Source/LearningAgentsTraining/Private/LearningAgentsRecorder.cpp
+++ b/Engine/Plugins/Experimental/LearningAgents/Source/LearningAgentsTraining/Private/LearningAgentsRecorder.cpp
@@ -17,6 +17,8 @@
 #include "Serialization/JsonSerializer.h"
 #include "Serialization/JsonWriter.h"
 
+#include UE_INLINE_GENERATED_CPP_BY_NAME(LearningAgentsRecorder)
+
 FLearningAgentsRecorderPathSettings::FLearningAgentsRecorderPathSettings()
 {
 	IntermediateRelativePath.Path = FPaths::ProjectIntermediateDir();
diff --git a/Engine/Plugins/Experimental/LearningAgents/Source/LearningAgentsTraining/Private/LearningAgentsRecording.cpp b/Engine/Plugins/Experimental/LearningAgents/Source/LearningAgentsTraining/Private/LearningAgentsRecording.cpp
--- a/Engine/Plugins/Experimental/LearningAgents/Source/LearningAgentsTraining/Private/LearningAgentsRecording.cpp
+++ b/Engine/Plugins/Experimental/LearningAgents/Source/LearningAgentsTraining/Private/LearningAgentsRecording.cpp
@@ -11,6 +11,8 @@
 #include "Serialization/JsonWriter.h"
 #include "Dom/JsonObject.h"
 
+#include UE_INLINE_GENERATED_CPP_BY_NAME(LearningAgentsRecording)
+
 ULearningAgentsRecording::ULearningAgentsRecording() = default;
 ULearningAgentsRecording::ULearningAgentsRecording(FVTableHelper& Helper) : Super(Helper) {}
 ULearningAgentsRecording::~ULearningAgentsRecording() = default;
@@ -312,4 +314,4 @@ void ULearningAgentsRecording::ForceMarkDirty()
 
 		Package->PackageMarkedDirtyEvent.Broadcast(Package, bIsDirty);
 	}
-}
\ No newline at end of file
+}
diff --git a/Engine/Plugins/Experimental/LearningAgents/Source/LearningAgentsTraining/Private/LearningAgentsRewards.cpp b/Engine/Plugins/Experimental/LearningAgents/Source/LearningAgentsTraining/Private/LearningAgentsRewards.cpp
--- a/Engine/Plugins/Experimental/LearningAgents/Source/LearningAgentsTraining/Private/LearningAgentsRewards.cpp
+++ b/Engine/Plugins/Experimental/LearningAgents/Source/LearningAgentsTraining/Private/LearningAgentsRewards.cpp
@@ -9,6 +9,8 @@
 
 #include "Components/SplineComponent.h"
 
+#include UE_INLINE_GENERATED_CPP_BY_NAME(LearningAgentsRewards)
+
 float ULearningAgentsRewards::MakeReward(
 	const float RewardValue,
 	const float RewardScale,
diff --git a/Engine/Plugins/Experimental/LearningAgents/Source/LearningAgentsTraining/Private/LearningAgentsTrainer.cpp b/Engine/Plugins/Experimental/LearningAgents/Source/LearningAgentsTraining/Private/LearningAgentsTrainer.cpp
--- a/Engine/Plugins/Experimental/LearningAgents/Source/LearningAgentsTraining/Private/LearningAgentsTrainer.cpp
+++ b/Engine/Plugins/Experimental/LearningAgents/Source/LearningAgentsTraining/Private/LearningAgentsTrainer.cpp
@@ -13,6 +13,8 @@
 #include "Editor/EditorPerformanceSettings.h"
 #endif
 
+#include UE_INLINE_GENERATED_CPP_BY_NAME(LearningAgentsTrainer)
+
 namespace UE::Learning::Agents
 {
 	ELearningAgentsTrainingDevice GetLearningAgentsTrainingDevice(const ETrainerDevice Device)
diff --git a/Engine/Plugins/Experimental/LearningAgents/Source/LearningAgentsTraining/Private/LearningAgentsTrainingEnvironment.cpp b/Engine/Plugins/Experimental/LearningAgents/Source/LearningAgentsTraining/Private/LearningAgentsTrainingEnvironment.cpp
--- a/Engine/Plugins/Experimental/LearningAgents/Source/LearningAgentsTraining/Private/LearningAgentsTrainingEnvironment.cpp
+++ b/Engine/Plugins/Experimental/LearningAgents/Source/LearningAgentsTraining/Private/LearningAgentsTrainingEnvironment.cpp
@@ -6,6 +6,8 @@
 #include "LearningCompletion.h"
 #include "LearningLog.h"
 
+#include UE_INLINE_GENERATED_CPP_BY_NAME(LearningAgentsTrainingEnvironment)
+
 ULearningAgentsTrainingEnvironment::ULearningAgentsTrainingEnvironment() : Super(FObjectInitializer::Get()) {}
 ULearningAgentsTrainingEnvironment::ULearningAgentsTrainingEnvironment(FVTableHelper& Helper) : Super(Helper) {}
 ULearningAgentsTrainingEnvironment::~ULearningAgentsTrainingEnvironment() = default;
diff --git a/Engine/Plugins/Experimental/LearningAgents/Source/LearningTraining/Private/LearningPPOTrainer.cpp b/Engine/Plugins/Experimental/LearningAgents/Source/LearningTraining/Private/LearningPPOTrainer.cpp
--- a/Engine/Plugins/Experimental/LearningAgents/Source/LearningTraining/Private/LearningPPOTrainer.cpp
+++ b/Engine/Plugins/Experimental/LearningAgents/Source/LearningTraining/Private/LearningPPOTrainer.cpp
@@ -26,6 +26,8 @@
 #include "Common/TcpSocketBuilder.h"
 #include "SocketSubsystem.h"
 
+#include UE_INLINE_GENERATED_CPP_BY_NAME(LearningPPOTrainer)
+
 ULearningSocketPPOTrainerServerCommandlet::ULearningSocketPPOTrainerServerCommandlet(const FObjectInitializer& ObjectInitializer)
 	: Super(ObjectInitializer)
 {}
diff --git a/Engine/Plugins/Runtime/AudioSynesthesia/Source/AudioSynesthesiaCore/Private/LKFSNRTFactory.cpp b/Engine/Plugins/Runtime/AudioSynesthesia/Source/AudioSynesthesiaCore/Private/LKFSNRTFactory.cpp
--- a/Engine/Plugins/Runtime/AudioSynesthesia/Source/AudioSynesthesiaCore/Private/LKFSNRTFactory.cpp
+++ b/Engine/Plugins/Runtime/AudioSynesthesia/Source/AudioSynesthesiaCore/Private/LKFSNRTFactory.cpp
@@ -6,6 +6,8 @@
 #include "AudioSynesthesiaCoreLog.h"
 #include "DSP/SlidingWindow.h"
 
+#include UE_INLINE_GENERATED_CPP_BY_NAME(LKFSNRTFactory)
+
 namespace Audio
 {
 	namespace LKFSNRTFactoryPrivate
diff --git a/Engine/Plugins/Runtime/Harmonix/Source/HarmonixMetasound/Private/HarmonixMetasound/Components/MusicClockTypes.cpp b/Engine/Plugins/Runtime/Harmonix/Source/HarmonixMetasound/Private/HarmonixMetasound/Components/MusicClockTypes.cpp
--- a/Engine/Plugins/Runtime/Harmonix/Source/HarmonixMetasound/Private/HarmonixMetasound/Components/MusicClockTypes.cpp
+++ b/Engine/Plugins/Runtime/Harmonix/Source/HarmonixMetasound/Private/HarmonixMetasound/Components/MusicClockTypes.cpp
@@ -1,9 +1,11 @@
-﻿// Copyright Epic Games, Inc. All Rights Reserved.
+// Copyright Epic Games, Inc. All Rights Reserved.
 
 #include "HarmonixMetasound/Components/MusicClockTypes.h"
 #include "HarmonixMetasound/Components/WallClockMusicClockDriver.h"
 #include "HarmonixMetasound/Components/MetasoundMusicClockDriver.h"
 
+#include UE_INLINE_GENERATED_CPP_BY_NAME(MusicClockTypes)
+
 TSharedPtr<FMusicClockDriverBase> FWallClockMusicClockSettings::MakeInstance(UObject* WorldContextObj) const
 {
 	return MakeShared<FWallClockMusicClockDriver>(WorldContextObj, *this);
@@ -12,4 +14,4 @@ TSharedPtr<FMusicClockDriverBase> FWallClockMusicClockSettings::MakeInstance(UOb
 TSharedPtr<FMusicClockDriverBase> FMetasoundMusicClockSettings::MakeInstance(UObject* WorldContextObj) const
 {
 	return MakeShared<FMetasoundMusicClockDriver>(WorldContextObj, *this);
-}
\ No newline at end of file
+}
diff --git a/Engine/Plugins/VirtualProduction/DynamicMaterial/Source/DynamicMaterial/Private/Components/DMMaterialComponent.cpp b/Engine/Plugins/VirtualProduction/DynamicMaterial/Source/DynamicMaterial/Private/Components/DMMaterialComponent.cpp
--- a/Engine/Plugins/VirtualProduction/DynamicMaterial/Source/DynamicMaterial/Private/Components/DMMaterialComponent.cpp
+++ b/Engine/Plugins/VirtualProduction/DynamicMaterial/Source/DynamicMaterial/Private/Components/DMMaterialComponent.cpp
@@ -10,6 +10,8 @@
 #include "Textures/SlateIcon.h"
 #endif
 
+#include UE_INLINE_GENERATED_CPP_BY_NAME(DMMaterialComponent)
+
 #define LOCTEXT_NAMESPACE "DMMaterialComponent"
 
 double UDMMaterialComponent::MinCleanTime = FPlatformTime::Seconds();
diff --git a/Engine/Plugins/VirtualProduction/DynamicMaterial/Source/DynamicMaterial/Private/Components/DMMaterialComponentDynamic.cpp b/Engine/Plugins/VirtualProduction/DynamicMaterial/Source/DynamicMaterial/Private/Components/DMMaterialComponentDynamic.cpp
--- a/Engine/Plugins/VirtualProduction/DynamicMaterial/Source/DynamicMaterial/Private/Components/DMMaterialComponentDynamic.cpp
+++ b/Engine/Plugins/VirtualProduction/DynamicMaterial/Source/DynamicMaterial/Private/Components/DMMaterialComponentDynamic.cpp
@@ -7,6 +7,8 @@
 #include "Model/DynamicMaterialModel.h"
 #include "Model/DynamicMaterialModelDynamic.h"
 
+#include UE_INLINE_GENERATED_CPP_BY_NAME(DMMaterialComponentDynamic)
+
 const FString UDMMaterialComponentDynamic::ParentValuePathToken = TEXT("ParentValue");
 
 UDMMaterialComponentDynamic::UDMMaterialComponentDynamic()
diff --git a/Engine/Plugins/VirtualProduction/DynamicMaterial/Source/DynamicMaterial/Private/Components/DMMaterialLinkedComponent.cpp b/Engine/Plugins/VirtualProduction/DynamicMaterial/Source/DynamicMaterial/Private/Components/DMMaterialLinkedComponent.cpp
--- a/Engine/Plugins/VirtualProduction/DynamicMaterial/Source/DynamicMaterial/Private/Components/DMMaterialLinkedComponent.cpp
+++ b/Engine/Plugins/VirtualProduction/DynamicMaterial/Source/DynamicMaterial/Private/Components/DMMaterialLinkedComponent.cpp
@@ -2,6 +2,8 @@
 
 #include "Components/DMMaterialLinkedComponent.h"
 
+#include UE_INLINE_GENERATED_CPP_BY_NAME(DMMaterialLinkedComponent)
+
 #if WITH_EDITOR
 UDMMaterialComponent* UDMMaterialLinkedComponent::GetParentComponent() const
 {
diff --git a/Engine/Plugins/VirtualProduction/DynamicMaterial/Source/DynamicMaterial/Private/Components/DMMaterialParameter.cpp b/Engine/Plugins/VirtualProduction/DynamicMaterial/Source/DynamicMaterial/Private/Components/DMMaterialParameter.cpp
--- a/Engine/Plugins/VirtualProduction/DynamicMaterial/Source/DynamicMaterial/Private/Components/DMMaterialParameter.cpp
+++ b/Engine/Plugins/VirtualProduction/DynamicMaterial/Source/DynamicMaterial/Private/Components/DMMaterialParameter.cpp
@@ -7,6 +7,8 @@
 #include "DynamicMaterialModule.h"
 #endif
 
+#include UE_INLINE_GENERATED_CPP_BY_NAME(DMMaterialParameter)
+
 UDMMaterialParameter::UDMMaterialParameter()
 {
 	ParameterName = NAME_None;
diff --git a/Engine/Plugins/VirtualProduction/DynamicMaterial/Source/DynamicMaterial/Private/Components/DMMaterialValue.cpp b/Engine/Plugins/VirtualProduction/DynamicMaterial/Source/DynamicMaterial/Private/Components/DMMaterialValue.cpp
--- a/Engine/Plugins/VirtualProduction/DynamicMaterial/Source/DynamicMaterial/Private/Components/DMMaterialValue.cpp
+++ b/Engine/Plugins/VirtualProduction/DynamicMaterial/Source/DynamicMaterial/Private/Components/DMMaterialValue.cpp
@@ -15,6 +15,8 @@
 #include "PropertyHandle.h"
 #endif // WITH_EDITOR
 
+#include UE_INLINE_GENERATED_CPP_BY_NAME(DMMaterialValue)
+
 #define LOCTEXT_NAMESPACE "DMMaterialValue"
 
 TMap<EDMValueType, TStrongObjectPtr<UClass>> UDMMaterialValue::TypeClasses = {};
diff --git a/Engine/Plugins/VirtualProduction/DynamicMaterial/Source/DynamicMaterial/Private/Components/DMMaterialValueDynamic.cpp b/Engine/Plugins/VirtualProduction/DynamicMaterial/Source/DynamicMaterial/Private/Components/DMMaterialValueDynamic.cpp
--- a/Engine/Plugins/VirtualProduction/DynamicMaterial/Source/DynamicMaterial/Private/Components/DMMaterialValueDynamic.cpp
+++ b/Engine/Plugins/VirtualProduction/DynamicMaterial/Source/DynamicMaterial/Private/Components/DMMaterialValueDynamic.cpp
@@ -12,6 +12,8 @@
 #include "UObject/Package.h"
 #endif
 
+#include UE_INLINE_GENERATED_CPP_BY_NAME(DMMaterialValueDynamic)
+
 #if WITH_EDITOR
 UDMMaterialValueDynamic* UDMMaterialValueDynamic::CreateValueDynamic(TSubclassOf<UDMMaterialValueDynamic> InInstanceValueClass, 
 	UDynamicMaterialModelDynamic* InMaterialModelDynamic, UDMMaterialValue* InParentValue)
diff --git a/Engine/Plugins/VirtualProduction/DynamicMaterial/Source/DynamicMaterial/Private/Components/DMRenderTargetRenderer.cpp b/Engine/Plugins/VirtualProduction/DynamicMaterial/Source/DynamicMaterial/Private/Components/DMRenderTargetRenderer.cpp
--- a/Engine/Plugins/VirtualProduction/DynamicMaterial/Source/DynamicMaterial/Private/Components/DMRenderTargetRenderer.cpp
+++ b/Engine/Plugins/VirtualProduction/DynamicMaterial/Source/DynamicMaterial/Private/Components/DMRenderTargetRenderer.cpp
@@ -6,6 +6,8 @@
 #include "Misc/CoreDelegates.h"
 #include "Templates/SubclassOf.h"
 
+#include UE_INLINE_GENERATED_CPP_BY_NAME(DMRenderTargetRenderer)
+
 #define LOCTEXT_NAMESPACE "DMRenderTargetRenderer"
 
 UDMRenderTargetRenderer* UDMRenderTargetRenderer::CreateRenderTargetRenderer(TSubclassOf<UDMRenderTargetRenderer> InRendererClass, 
diff --git a/Engine/Plugins/VirtualProduction/DynamicMaterial/Source/DynamicMaterial/Private/Components/DMTextureUV.cpp b/Engine/Plugins/VirtualProduction/DynamicMaterial/Source/DynamicMaterial/Private/Components/DMTextureUV.cpp
--- a/Engine/Plugins/VirtualProduction/DynamicMaterial/Source/DynamicMaterial/Private/Components/DMTextureUV.cpp
+++ b/Engine/Plugins/VirtualProduction/DynamicMaterial/Source/DynamicMaterial/Private/Components/DMTextureUV.cpp
@@ -12,6 +12,8 @@
 #include "Model/IDMMaterialBuildUtilsInterface.h"
 #endif
 
+#include UE_INLINE_GENERATED_CPP_BY_NAME(DMTextureUV)
+
 namespace UE::DynamicMaterial::Private
 {
 	TMap<int32, FName> BaseParameterNames = {
diff --git a/Engine/Plugins/VirtualProduction/DynamicMaterial/Source/DynamicMaterial/Private/Components/DMTextureUVDynamic.cpp b/Engine/Plugins/VirtualProduction/DynamicMaterial/Source/DynamicMaterial/Private/Components/DMTextureUVDynamic.cpp
--- a/Engine/Plugins/VirtualProduction/DynamicMaterial/Source/DynamicMaterial/Private/Components/DMTextureUVDynamic.cpp
+++ b/Engine/Plugins/VirtualProduction/DynamicMaterial/Source/DynamicMaterial/Private/Components/DMTextureUVDynamic.cpp
@@ -8,6 +8,8 @@
 #include "Model/DynamicMaterialModel.h"
 #include "Model/DynamicMaterialModelDynamic.h"
 
+#include UE_INLINE_GENERATED_CPP_BY_NAME(DMTextureUVDynamic)
+
 #if WITH_EDITOR
 UDMTextureUVDynamic* UDMTextureUVDynamic::CreateTextureUVDynamic(UDynamicMaterialModelDynamic* InMaterialModelDynamic, 
 	UDMTextureUV* InParentTextureUV)
diff --git a/Engine/Plugins/VirtualProduction/DynamicMaterial/Source/DynamicMaterial/Private/Components/MaterialValues/DMMaterialValueBool.cpp b/Engine/Plugins/VirtualProduction/DynamicMaterial/Source/DynamicMaterial/Private/Components/MaterialValues/DMMaterialValueBool.cpp
--- a/Engine/Plugins/VirtualProduction/DynamicMaterial/Source/DynamicMaterial/Private/Components/MaterialValues/DMMaterialValueBool.cpp
+++ b/Engine/Plugins/VirtualProduction/DynamicMaterial/Source/DynamicMaterial/Private/Components/MaterialValues/DMMaterialValueBool.cpp
@@ -12,6 +12,8 @@
 #include "Model/IDMMaterialBuildUtilsInterface.h"
 #include "Utils/DMUtils.h"
 #endif
+
+#include UE_INLINE_GENERATED_CPP_BY_NAME(DMMaterialValueBool)
  
 #define LOCTEXT_NAMESPACE "DMMaterialValueBool"
 
diff --git a/Engine/Plugins/VirtualProduction/DynamicMaterial/Source/DynamicMaterial/Private/Components/MaterialValues/DMMaterialValueColorAtlas.cpp b/Engine/Plugins/VirtualProduction/DynamicMaterial/Source/DynamicMaterial/Private/Components/MaterialValues/DMMaterialValueColorAtlas.cpp
--- a/Engine/Plugins/VirtualProduction/DynamicMaterial/Source/DynamicMaterial/Private/Components/MaterialValues/DMMaterialValueColorAtlas.cpp
+++ b/Engine/Plugins/VirtualProduction/DynamicMaterial/Source/DynamicMaterial/Private/Components/MaterialValues/DMMaterialValueColorAtlas.cpp
@@ -17,6 +17,8 @@
 #include "Utils/DMUtils.h"
 #endif
 
+#include UE_INLINE_GENERATED_CPP_BY_NAME(DMMaterialValueColorAtlas)
+
 #define LOCTEXT_NAMESPACE "DMMaterialValueColorAtlas"
 
 UDMMaterialValueColorAtlas::UDMMaterialValueColorAtlas()
diff --git a/Engine/Plugins/VirtualProduction/DynamicMaterial/Source/DynamicMaterial/Private/Components/MaterialValues/DMMaterialValueFloat.cpp b/Engine/Plugins/VirtualProduction/DynamicMaterial/Source/DynamicMaterial/Private/Components/MaterialValues/DMMaterialValueFloat.cpp
--- a/Engine/Plugins/VirtualProduction/DynamicMaterial/Source/DynamicMaterial/Private/Components/MaterialValues/DMMaterialValueFloat.cpp
+++ b/Engine/Plugins/VirtualProduction/DynamicMaterial/Source/DynamicMaterial/Private/Components/MaterialValues/DMMaterialValueFloat.cpp
@@ -2,6 +2,8 @@
 
 #include "Components/MaterialValues/DMMaterialValueFloat.h"
 
+#include UE_INLINE_GENERATED_CPP_BY_NAME(DMMaterialValueFloat)
+
 UDMMaterialValueFloat::UDMMaterialValueFloat()
 	: UDMMaterialValueFloat(EDMValueType::VT_None)
 {
diff --git a/Engine/Plugins/VirtualProduction/DynamicMaterial/Source/DynamicMaterial/Private/Components/MaterialValues/DMMaterialValueFloat1.cpp b/Engine/Plugins/VirtualProduction/DynamicMaterial/Source/DynamicMaterial/Private/Components/MaterialValues/DMMaterialValueFloat1.cpp
--- a/Engine/Plugins/VirtualProduction/DynamicMaterial/Source/DynamicMaterial/Private/Components/MaterialValues/DMMaterialValueFloat1.cpp
+++ b/Engine/Plugins/VirtualProduction/DynamicMaterial/Source/DynamicMaterial/Private/Components/MaterialValues/DMMaterialValueFloat1.cpp
@@ -14,6 +14,8 @@
 #include "Model/IDMMaterialBuildUtilsInterface.h"
 #include "Utils/DMUtils.h"
 #endif
+
+#include UE_INLINE_GENERATED_CPP_BY_NAME(DMMaterialValueFloat1)
  
 #define LOCTEXT_NAMESPACE "DMMaterialValueFloat1"
 
diff --git a/Engine/Plugins/VirtualProduction/DynamicMaterial/Source/DynamicMaterial/Private/Components/MaterialValues/DMMaterialValueFloat2.cpp b/Engine/Plugins/VirtualProduction/DynamicMaterial/Source/DynamicMaterial/Private/Components/MaterialValues/DMMaterialValueFloat2.cpp
--- a/Engine/Plugins/VirtualProduction/DynamicMaterial/Source/DynamicMaterial/Private/Components/MaterialValues/DMMaterialValueFloat2.cpp
+++ b/Engine/Plugins/VirtualProduction/DynamicMaterial/Source/DynamicMaterial/Private/Components/MaterialValues/DMMaterialValueFloat2.cpp
@@ -16,6 +16,8 @@
 #include "Utils/DMUtils.h"
 #endif
 
+#include UE_INLINE_GENERATED_CPP_BY_NAME(DMMaterialValueFloat2)
+
 #define LOCTEXT_NAMESPACE "DMMaterialValueFloat2"
 
 UDMMaterialValueFloat2::UDMMaterialValueFloat2()
diff --git a/Engine/Plugins/VirtualProduction/DynamicMaterial/Source/DynamicMaterial/Private/Components/MaterialValues/DMMaterialValueFloat3RGB.cpp b/Engine/Plugins/VirtualProduction/DynamicMaterial/Source/DynamicMaterial/Private/Components/MaterialValues/DMMaterialValueFloat3RGB.cpp
--- a/Engine/Plugins/VirtualProduction/DynamicMaterial/Source/DynamicMaterial/Private/Components/MaterialValues/DMMaterialValueFloat3RGB.cpp
+++ b/Engine/Plugins/VirtualProduction/DynamicMaterial/Source/DynamicMaterial/Private/Components/MaterialValues/DMMaterialValueFloat3RGB.cpp
@@ -14,6 +14,8 @@
 #include "PropertyHandle.h"
 #include "Utils/DMUtils.h"
 #endif
+
+#include UE_INLINE_GENERATED_CPP_BY_NAME(DMMaterialValueFloat3RGB)
  
 #define LOCTEXT_NAMESPACE "DMMaterialValueFloat3RGB"
 
diff --git a/Engine/Plugins/VirtualProduction/DynamicMaterial/Source/DynamicMaterial/Private/Components/MaterialValues/DMMaterialValueFloat3RPY.cpp b/Engine/Plugins/VirtualProduction/DynamicMaterial/Source/DynamicMaterial/Private/Components/MaterialValues/DMMaterialValueFloat3RPY.cpp
--- a/Engine/Plugins/VirtualProduction/DynamicMaterial/Source/DynamicMaterial/Private/Components/MaterialValues/DMMaterialValueFloat3RPY.cpp
+++ b/Engine/Plugins/VirtualProduction/DynamicMaterial/Source/DynamicMaterial/Private/Components/MaterialValues/DMMaterialValueFloat3RPY.cpp
@@ -14,6 +14,8 @@
 #include "PropertyHandle.h"
 #include "Utils/DMUtils.h"
 #endif
+
+#include UE_INLINE_GENERATED_CPP_BY_NAME(DMMaterialValueFloat3RPY)
  
 #define LOCTEXT_NAMESPACE "DMMaterialValueFloat3RPY"
 
diff --git a/Engine/Plugins/VirtualProduction/DynamicMaterial/Source/DynamicMaterial/Private/Components/MaterialValues/DMMaterialValueFloat3XYZ.cpp b/Engine/Plugins/VirtualProduction/DynamicMaterial/Source/DynamicMaterial/Private/Components/MaterialValues/DMMaterialValueFloat3XYZ.cpp
--- a/Engine/Plugins/VirtualProduction/DynamicMaterial/Source/DynamicMaterial/Private/Components/MaterialValues/DMMaterialValueFloat3XYZ.cpp
+++ b/Engine/Plugins/VirtualProduction/DynamicMaterial/Source/DynamicMaterial/Private/Components/MaterialValues/DMMaterialValueFloat3XYZ.cpp
@@ -14,6 +14,8 @@
 #include "PropertyHandle.h"
 #include "Utils/DMUtils.h"
 #endif
+
+#include UE_INLINE_GENERATED_CPP_BY_NAME(DMMaterialValueFloat3XYZ)
  
 #define LOCTEXT_NAMESPACE "DMMaterialValueFloat3XYZ"
 
diff --git a/Engine/Plugins/VirtualProduction/DynamicMaterial/Source/DynamicMaterial/Private/Components/MaterialValues/DMMaterialValueFloat4.cpp b/Engine/Plugins/VirtualProduction/DynamicMaterial/Source/DynamicMaterial/Private/Components/MaterialValues/DMMaterialValueFloat4.cpp
--- a/Engine/Plugins/VirtualProduction/DynamicMaterial/Source/DynamicMaterial/Private/Components/MaterialValues/DMMaterialValueFloat4.cpp
+++ b/Engine/Plugins/VirtualProduction/DynamicMaterial/Source/DynamicMaterial/Private/Components/MaterialValues/DMMaterialValueFloat4.cpp
@@ -14,6 +14,8 @@
 #include "PropertyHandle.h"
 #include "Utils/DMUtils.h"
 #endif
+
+#include UE_INLINE_GENERATED_CPP_BY_NAME(DMMaterialValueFloat4)
  
 #define LOCTEXT_NAMESPACE "DMMaterialValueFloat4"
 
diff --git a/Engine/Plugins/VirtualProduction/DynamicMaterial/Source/DynamicMaterial/Private/Components/MaterialValues/DMMaterialValueRenderTarget.cpp b/Engine/Plugins/VirtualProduction/DynamicMaterial/Source/DynamicMaterial/Private/Components/MaterialValues/DMMaterialValueRenderTarget.cpp
--- a/Engine/Plugins/VirtualProduction/DynamicMaterial/Source/DynamicMaterial/Private/Components/MaterialValues/DMMaterialValueRenderTarget.cpp
+++ b/Engine/Plugins/VirtualProduction/DynamicMaterial/Source/DynamicMaterial/Private/Components/MaterialValues/DMMaterialValueRenderTarget.cpp
@@ -10,6 +10,8 @@
 #include "Components/MaterialValuesDynamic/DMMaterialValueRenderTargetDynamic.h"
 #endif
 
+#include UE_INLINE_GENERATED_CPP_BY_NAME(DMMaterialValueRenderTarget)
+
 #define LOCTEXT_NAMESPACE "DMMaterialValueRenderTarget"
 
 const FString UDMMaterialValueRenderTarget::RendererPathToken = "Renderer";
diff --git a/Engine/Plugins/VirtualProduction/DynamicMaterial/Source/DynamicMaterial/Private/Components/MaterialValues/DMMaterialValueTexture.cpp b/Engine/Plugins/VirtualProduction/DynamicMaterial/Source/DynamicMaterial/Private/Components/MaterialValues/DMMaterialValueTexture.cpp
--- a/Engine/Plugins/VirtualProduction/DynamicMaterial/Source/DynamicMaterial/Private/Components/MaterialValues/DMMaterialValueTexture.cpp
+++ b/Engine/Plugins/VirtualProduction/DynamicMaterial/Source/DynamicMaterial/Private/Components/MaterialValues/DMMaterialValueTexture.cpp
@@ -18,6 +18,8 @@
 #include "Utils/DMUtils.h"
 #endif
 
+#include UE_INLINE_GENERATED_CPP_BY_NAME(DMMaterialValueTexture)
+
 #define LOCTEXT_NAMESPACE "DMMaterialValueTexture"
 
 #if WITH_EDITOR
diff --git a/Engine/Plugins/VirtualProduction/DynamicMaterial/Source/DynamicMaterial/Private/Components/MaterialValuesDynamic/DMMaterialValueBoolDynamic.cpp b/Engine/Plugins/VirtualProduction/DynamicMaterial/Source/DynamicMaterial/Private/Components/MaterialValuesDynamic/DMMaterialValueBoolDynamic.cpp
--- a/Engine/Plugins/VirtualProduction/DynamicMaterial/Source/DynamicMaterial/Private/Components/MaterialValuesDynamic/DMMaterialValueBoolDynamic.cpp
+++ b/Engine/Plugins/VirtualProduction/DynamicMaterial/Source/DynamicMaterial/Private/Components/MaterialValuesDynamic/DMMaterialValueBoolDynamic.cpp
@@ -6,6 +6,8 @@
 #include "Components/MaterialValues/DMMaterialValueBool.h"
 #include "Materials/MaterialInstanceDynamic.h"
 
+#include UE_INLINE_GENERATED_CPP_BY_NAME(DMMaterialValueBoolDynamic)
+
 UDMMaterialValueBoolDynamic::UDMMaterialValueBoolDynamic()
 	: UDMMaterialValueDynamic()
 	, Value(false)
diff --git a/Engine/Plugins/VirtualProduction/DynamicMaterial/Source/DynamicMaterial/Private/Components/MaterialValuesDynamic/DMMaterialValueColorAtlasDynamic.cpp b/Engine/Plugins/VirtualProduction/DynamicMaterial/Source/DynamicMaterial/Private/Components/MaterialValuesDynamic/DMMaterialValueColorAtlasDynamic.cpp
--- a/Engine/Plugins/VirtualProduction/DynamicMaterial/Source/DynamicMaterial/Private/Components/MaterialValuesDynamic/DMMaterialValueColorAtlasDynamic.cpp
+++ b/Engine/Plugins/VirtualProduction/DynamicMaterial/Source/DynamicMaterial/Private/Components/MaterialValuesDynamic/DMMaterialValueColorAtlasDynamic.cpp
@@ -6,6 +6,8 @@
 #include "Components/MaterialValues/DMMaterialValueColorAtlas.h"
 #include "Materials/MaterialInstanceDynamic.h"
 
+#include UE_INLINE_GENERATED_CPP_BY_NAME(DMMaterialValueColorAtlasDynamic)
+
 UDMMaterialValueColorAtlasDynamic::UDMMaterialValueColorAtlasDynamic()
 	: UDMMaterialValueDynamic()
 	, Value(0.f)
diff --git a/Engine/Plugins/VirtualProduction/DynamicMaterial/Source/DynamicMaterial/Private/Components/MaterialValuesDynamic/DMMaterialValueFloat1Dynamic.cpp b/Engine/Plugins/VirtualProduction/DynamicMaterial/Source/DynamicMaterial/Private/Components/MaterialValuesDynamic/DMMaterialValueFloat1Dynamic.cpp
--- a/Engine/Plugins/VirtualProduction/DynamicMaterial/Source/DynamicMaterial/Private/Components/MaterialValuesDynamic/DMMaterialValueFloat1Dynamic.cpp
+++ b/Engine/Plugins/VirtualProduction/DynamicMaterial/Source/DynamicMaterial/Private/Components/MaterialValuesDynamic/DMMaterialValueFloat1Dynamic.cpp
@@ -6,6 +6,8 @@
 #include "Components/MaterialValues/DMMaterialValueFloat1.h"
 #include "Materials/MaterialInstanceDynamic.h"
 
+#include UE_INLINE_GENERATED_CPP_BY_NAME(DMMaterialValueFloat1Dynamic)
+
 UDMMaterialValueFloat1Dynamic::UDMMaterialValueFloat1Dynamic()
 	: UDMMaterialValueDynamic()
 	, Value(0.f)
diff --git a/Engine/Plugins/VirtualProduction/DynamicMaterial/Source/DynamicMaterial/Private/Components/MaterialValuesDynamic/DMMaterialValueFloat2Dynamic.cpp b/Engine/Plugins/VirtualProduction/DynamicMaterial/Source/DynamicMaterial/Private/Components/MaterialValuesDynamic/DMMaterialValueFloat2Dynamic.cpp
--- a/Engine/Plugins/VirtualProduction/DynamicMaterial/Source/DynamicMaterial/Private/Components/MaterialValuesDynamic/DMMaterialValueFloat2Dynamic.cpp
+++ b/Engine/Plugins/VirtualProduction/DynamicMaterial/Source/DynamicMaterial/Private/Components/MaterialValuesDynamic/DMMaterialValueFloat2Dynamic.cpp
@@ -10,6 +10,8 @@
 #include "PropertyHandle.h"
 #endif
 
+#include UE_INLINE_GENERATED_CPP_BY_NAME(DMMaterialValueFloat2Dynamic)
+
 UDMMaterialValueFloat2Dynamic::UDMMaterialValueFloat2Dynamic()
 	: UDMMaterialValueDynamic()
 	, Value(FVector2D::ZeroVector)
diff --git a/Engine/Plugins/VirtualProduction/DynamicMaterial/Source/DynamicMaterial/Private/Components/MaterialValuesDynamic/DMMaterialValueFloat3RGBDynamic.cpp b/Engine/Plugins/VirtualProduction/DynamicMaterial/Source/DynamicMaterial/Private/Components/MaterialValuesDynamic/DMMaterialValueFloat3RGBDynamic.cpp
--- a/Engine/Plugins/VirtualProduction/DynamicMaterial/Source/DynamicMaterial/Private/Components/MaterialValuesDynamic/DMMaterialValueFloat3RGBDynamic.cpp
+++ b/Engine/Plugins/VirtualProduction/DynamicMaterial/Source/DynamicMaterial/Private/Components/MaterialValuesDynamic/DMMaterialValueFloat3RGBDynamic.cpp
@@ -10,6 +10,8 @@
 #include "PropertyHandle.h"
 #endif
 
+#include UE_INLINE_GENERATED_CPP_BY_NAME(DMMaterialValueFloat3RGBDynamic)
+
 UDMMaterialValueFloat3RGBDynamic::UDMMaterialValueFloat3RGBDynamic()
 	: UDMMaterialValueDynamic()
 	, Value(FLinearColor::Black)
diff --git a/Engine/Plugins/VirtualProduction/DynamicMaterial/Source/DynamicMaterial/Private/Components/MaterialValuesDynamic/DMMaterialValueFloat3RPYDynamic.cpp b/Engine/Plugins/VirtualProduction/DynamicMaterial/Source/DynamicMaterial/Private/Components/MaterialValuesDynamic/DMMaterialValueFloat3RPYDynamic.cpp
--- a/Engine/Plugins/VirtualProduction/DynamicMaterial/Source/DynamicMaterial/Private/Components/MaterialValuesDynamic/DMMaterialValueFloat3RPYDynamic.cpp
+++ b/Engine/Plugins/VirtualProduction/DynamicMaterial/Source/DynamicMaterial/Private/Components/MaterialValuesDynamic/DMMaterialValueFloat3RPYDynamic.cpp
@@ -10,6 +10,8 @@
 #include "PropertyHandle.h"
 #endif
 
+#include UE_INLINE_GENERATED_CPP_BY_NAME(DMMaterialValueFloat3RPYDynamic)
+
 UDMMaterialValueFloat3RPYDynamic::UDMMaterialValueFloat3RPYDynamic()
 	: UDMMaterialValueDynamic()
 	, Value(FRotator::ZeroRotator)
diff --git a/Engine/Plugins/VirtualProduction/DynamicMaterial/Source/DynamicMaterial/Private/Components/MaterialValuesDynamic/DMMaterialValueFloat3XYZDynamic.cpp b/Engine/Plugins/VirtualProduction/DynamicMaterial/Source/DynamicMaterial/Private/Components/MaterialValuesDynamic/DMMaterialValueFloat3XYZDynamic.cpp
--- a/Engine/Plugins/VirtualProduction/DynamicMaterial/Source/DynamicMaterial/Private/Components/MaterialValuesDynamic/DMMaterialValueFloat3XYZDynamic.cpp
+++ b/Engine/Plugins/VirtualProduction/DynamicMaterial/Source/DynamicMaterial/Private/Components/MaterialValuesDynamic/DMMaterialValueFloat3XYZDynamic.cpp
@@ -10,6 +10,8 @@
 #include "PropertyHandle.h"
 #endif
 
+#include UE_INLINE_GENERATED_CPP_BY_NAME(DMMaterialValueFloat3XYZDynamic)
+
 UDMMaterialValueFloat3XYZDynamic::UDMMaterialValueFloat3XYZDynamic()
 	: UDMMaterialValueDynamic()
 	, Value(FVector::ZeroVector)
diff --git a/Engine/Plugins/VirtualProduction/DynamicMaterial/Source/DynamicMaterial/Private/Components/MaterialValuesDynamic/DMMaterialValueFloat4Dynamic.cpp b/Engine/Plugins/VirtualProduction/DynamicMaterial/Source/DynamicMaterial/Private/Components/MaterialValuesDynamic/DMMaterialValueFloat4Dynamic.cpp
--- a/Engine/Plugins/VirtualProduction/DynamicMaterial/Source/DynamicMaterial/Private/Components/MaterialValuesDynamic/DMMaterialValueFloat4Dynamic.cpp
+++ b/Engine/Plugins/VirtualProduction/DynamicMaterial/Source/DynamicMaterial/Private/Components/MaterialValuesDynamic/DMMaterialValueFloat4Dynamic.cpp
@@ -10,6 +10,8 @@
 #include "PropertyHandle.h"
 #endif
 
+#include UE_INLINE_GENERATED_CPP_BY_NAME(DMMaterialValueFloat4Dynamic)
+
 UDMMaterialValueFloat4Dynamic::UDMMaterialValueFloat4Dynamic()
 	: UDMMaterialValueDynamic()
 	, Value(FLinearColor::Black)
diff --git a/Engine/Plugins/VirtualProduction/DynamicMaterial/Source/DynamicMaterial/Private/Components/MaterialValuesDynamic/DMMaterialValueRenderTargetDynamic.cpp b/Engine/Plugins/VirtualProduction/DynamicMaterial/Source/DynamicMaterial/Private/Components/MaterialValuesDynamic/DMMaterialValueRenderTargetDynamic.cpp
--- a/Engine/Plugins/VirtualProduction/DynamicMaterial/Source/DynamicMaterial/Private/Components/MaterialValuesDynamic/DMMaterialValueRenderTargetDynamic.cpp
+++ b/Engine/Plugins/VirtualProduction/DynamicMaterial/Source/DynamicMaterial/Private/Components/MaterialValuesDynamic/DMMaterialValueRenderTargetDynamic.cpp
@@ -2,6 +2,8 @@
 
 #include "Components/MaterialValuesDynamic/DMMaterialValueRenderTargetDynamic.h"
 
+#include UE_INLINE_GENERATED_CPP_BY_NAME(DMMaterialValueRenderTargetDynamic)
+
 UDMMaterialValueRenderTargetDynamic::UDMMaterialValueRenderTargetDynamic()
 {
 }
diff --git a/Engine/Plugins/VirtualProduction/DynamicMaterial/Source/DynamicMaterial/Private/Components/MaterialValuesDynamic/DMMaterialValueTextureDynamic.cpp b/Engine/Plugins/VirtualProduction/DynamicMaterial/Source/DynamicMaterial/Private/Components/MaterialValuesDynamic/DMMaterialValueTextureDynamic.cpp
--- a/Engine/Plugins/VirtualProduction/DynamicMaterial/Source/DynamicMaterial/Private/Components/MaterialValuesDynamic/DMMaterialValueTextureDynamic.cpp
+++ b/Engine/Plugins/VirtualProduction/DynamicMaterial/Source/DynamicMaterial/Private/Components/MaterialValuesDynamic/DMMaterialValueTextureDynamic.cpp
@@ -10,6 +10,8 @@
 #include "Engine/Texture.h"
 #endif
 
+#include UE_INLINE_GENERATED_CPP_BY_NAME(DMMaterialValueTextureDynamic)
+
 UDMMaterialValueTextureDynamic::UDMMaterialValueTextureDynamic()
 	: UDMMaterialValueDynamic()
 	, Value(nullptr)
diff --git a/Engine/Plugins/VirtualProduction/DynamicMaterial/Source/DynamicMaterial/Private/Components/RenderTargetRenderers/DMRenderTargetTextRenderer.cpp b/Engine/Plugins/VirtualProduction/DynamicMaterial/Source/DynamicMaterial/Private/Components/RenderTargetRenderers/DMRenderTargetTextRenderer.cpp
--- a/Engine/Plugins/VirtualProduction/DynamicMaterial/Source/DynamicMaterial/Private/Components/RenderTargetRenderers/DMRenderTargetTextRenderer.cpp
+++ b/Engine/Plugins/VirtualProduction/DynamicMaterial/Source/DynamicMaterial/Private/Components/RenderTargetRenderers/DMRenderTargetTextRenderer.cpp
@@ -14,6 +14,8 @@
 #include "Dom/JsonValue.h"
 #endif
 
+#include UE_INLINE_GENERATED_CPP_BY_NAME(DMRenderTargetTextRenderer)
+
 #define LOCTEXT_NAMESPACE "DMRenderTargetTextRenderer"
 
 namespace UE::DynamicMaterial::Private
diff --git a/Engine/Plugins/VirtualProduction/DynamicMaterial/Source/DynamicMaterial/Private/Components/RenderTargetRenderers/DMRenderTargetUMGWidgetRenderer.cpp b/Engine/Plugins/VirtualProduction/DynamicMaterial/Source/DynamicMaterial/Private/Components/RenderTargetRenderers/DMRenderTargetUMGWidgetRenderer.cpp
--- a/Engine/Plugins/VirtualProduction/DynamicMaterial/Source/DynamicMaterial/Private/Components/RenderTargetRenderers/DMRenderTargetUMGWidgetRenderer.cpp
+++ b/Engine/Plugins/VirtualProduction/DynamicMaterial/Source/DynamicMaterial/Private/Components/RenderTargetRenderers/DMRenderTargetUMGWidgetRenderer.cpp
@@ -6,6 +6,8 @@
 #include "Dom/JsonValue.h"
 #include "Engine/Engine.h"
 
+#include UE_INLINE_GENERATED_CPP_BY_NAME(DMRenderTargetUMGWidgetRenderer)
+
 #define LOCTEXT_NAMESPACE "DMRenderTargetUWidgetRenderer"
 
 UDMRenderTargetUMGWidgetRenderer::UDMRenderTargetUMGWidgetRenderer()
diff --git a/Engine/Plugins/VirtualProduction/DynamicMaterial/Source/DynamicMaterial/Private/Components/RenderTargetRenderers/DMRenderTargetWidgetRendererBase.cpp b/Engine/Plugins/VirtualProduction/DynamicMaterial/Source/DynamicMaterial/Private/Components/RenderTargetRenderers/DMRenderTargetWidgetRendererBase.cpp
--- a/Engine/Plugins/VirtualProduction/DynamicMaterial/Source/DynamicMaterial/Private/Components/RenderTargetRenderers/DMRenderTargetWidgetRendererBase.cpp
+++ b/Engine/Plugins/VirtualProduction/DynamicMaterial/Source/DynamicMaterial/Private/Components/RenderTargetRenderers/DMRenderTargetWidgetRendererBase.cpp
@@ -7,6 +7,8 @@
 #include "Slate/WidgetRenderer.h"
 #include "Widgets/SWidget.h"
 
+#include UE_INLINE_GENERATED_CPP_BY_NAME(DMRenderTargetWidgetRendererBase)
+
 #define LOCTEXT_NAMESPACE "DMRenderTargetUMGWidgetRenderer"
 
 UDMRenderTargetWidgetRendererBase::UDMRenderTargetWidgetRendererBase()
diff --git a/Engine/Plugins/VirtualProduction/DynamicMaterial/Source/DynamicMaterial/Private/DMDefs.cpp b/Engine/Plugins/VirtualProduction/DynamicMaterial/Source/DynamicMaterial/Private/DMDefs.cpp
--- a/Engine/Plugins/VirtualProduction/DynamicMaterial/Source/DynamicMaterial/Private/DMDefs.cpp
+++ b/Engine/Plugins/VirtualProduction/DynamicMaterial/Source/DynamicMaterial/Private/DMDefs.cpp
@@ -2,6 +2,8 @@
 
 #include "DMDefs.h"
 
+#include UE_INLINE_GENERATED_CPP_BY_NAME(DMDefs)
+
 int32 FDMUpdateGuard::GuardCount = 0;
 uint32 FDMInitializationGuard::GuardCount = 0;
 
diff --git a/Engine/Plugins/VirtualProduction/DynamicMaterial/Source/DynamicMaterial/Private/DMValueDefinition.cpp b/Engine/Plugins/VirtualProduction/DynamicMaterial/Source/DynamicMaterial/Private/DMValueDefinition.cpp
--- a/Engine/Plugins/VirtualProduction/DynamicMaterial/Source/DynamicMaterial/Private/DMValueDefinition.cpp
+++ b/Engine/Plugins/VirtualProduction/DynamicMaterial/Source/DynamicMaterial/Private/DMValueDefinition.cpp
@@ -17,6 +17,8 @@
 #include "Textures/SlateIcon.h"
 #endif
 
+#include UE_INLINE_GENERATED_CPP_BY_NAME(DMValueDefinition)
+
 #define LOCTEXT_NAMESPACE "DMValueDefinition"
 
 namespace UE::MaterialDesigner::Private
diff --git a/Engine/Plugins/VirtualProduction/DynamicMaterial/Source/DynamicMaterial/Private/Material/DynamicMaterialInstance.cpp b/Engine/Plugins/VirtualProduction/DynamicMaterial/Source/DynamicMaterial/Private/Material/DynamicMaterialInstance.cpp
--- a/Engine/Plugins/VirtualProduction/DynamicMaterial/Source/DynamicMaterial/Private/Material/DynamicMaterialInstance.cpp
+++ b/Engine/Plugins/VirtualProduction/DynamicMaterial/Source/DynamicMaterial/Private/Material/DynamicMaterialInstance.cpp
@@ -13,6 +13,8 @@
 #include "Model/IDynamicMaterialModelEditorOnlyDataInterface.h"
 #endif
 
+#include UE_INLINE_GENERATED_CPP_BY_NAME(DynamicMaterialInstance)
+
 #define LOCTEXT_NAMESPACE "DynamicMaterialInstance"
 
 namespace UE::DynamicMaterial::Private
diff --git a/Engine/Plugins/VirtualProduction/DynamicMaterial/Source/DynamicMaterial/Private/Model/DynamicMaterialModel.cpp b/Engine/Plugins/VirtualProduction/DynamicMaterial/Source/DynamicMaterial/Private/Model/DynamicMaterialModel.cpp
--- a/Engine/Plugins/VirtualProduction/DynamicMaterial/Source/DynamicMaterial/Private/Model/DynamicMaterialModel.cpp
+++ b/Engine/Plugins/VirtualProduction/DynamicMaterial/Source/DynamicMaterial/Private/Model/DynamicMaterialModel.cpp
@@ -17,6 +17,8 @@
 #include "Materials/Material.h"
 #endif
 
+#include UE_INLINE_GENERATED_CPP_BY_NAME(DynamicMaterialModel)
+
 const FString UDynamicMaterialModel::ValuesPathToken = FString(TEXT("Values"));
 const FString UDynamicMaterialModel::ParametersPathToken = FString(TEXT("Parameters"));
 const FLazyName UDynamicMaterialModel::GlobalBaseColorValueName = FLazyName(TEXT("GlobalBaseColorValue"));
diff --git a/Engine/Plugins/VirtualProduction/DynamicMaterial/Source/DynamicMaterial/Private/Model/DynamicMaterialModelDynamic.cpp b/Engine/Plugins/VirtualProduction/DynamicMaterial/Source/DynamicMaterial/Private/Model/DynamicMaterialModelDynamic.cpp
--- a/Engine/Plugins/VirtualProduction/DynamicMaterial/Source/DynamicMaterial/Private/Model/DynamicMaterialModelDynamic.cpp
+++ b/Engine/Plugins/VirtualProduction/DynamicMaterial/Source/DynamicMaterial/Private/Model/DynamicMaterialModelDynamic.cpp
@@ -13,6 +13,8 @@
 #include "DynamicMaterialModule.h"
 #endif
 
+#include UE_INLINE_GENERATED_CPP_BY_NAME(DynamicMaterialModelDynamic)
+
 const FString UDynamicMaterialModelDynamic::ParentModelPathToken = TEXT("ParentModel");
 const FString UDynamicMaterialModelDynamic::DynamicComponentsPathToken = TEXT("DynamicComponents");
 
diff --git a/Engine/Plugins/VirtualProduction/DynamicMaterial/Source/DynamicMaterial/Private/Utils/DMJsonUtils.cpp b/Engine/Plugins/VirtualProduction/DynamicMaterial/Source/DynamicMaterial/Private/Utils/DMJsonUtils.cpp
--- a/Engine/Plugins/VirtualProduction/DynamicMaterial/Source/DynamicMaterial/Private/Utils/DMJsonUtils.cpp
+++ b/Engine/Plugins/VirtualProduction/DynamicMaterial/Source/DynamicMaterial/Private/Utils/DMJsonUtils.cpp
@@ -9,6 +9,8 @@
 #include "DynamicMaterialModule.h"
 #include "JsonObjectConverter.h"
 
+#include UE_INLINE_GENERATED_CPP_BY_NAME(DMJsonUtils)
+
 namespace UE::DynamicMaterial::Private
 {
 	const FString JsonKey_Class = TEXT("Class");
diff --git a/Engine/Plugins/VirtualProduction/DynamicMaterial/Source/DynamicMaterialEditor/Private/Components/DMMaterialEffect.cpp b/Engine/Plugins/VirtualProduction/DynamicMaterial/Source/DynamicMaterialEditor/Private/Components/DMMaterialEffect.cpp
--- a/Engine/Plugins/VirtualProduction/DynamicMaterial/Source/DynamicMaterialEditor/Private/Components/DMMaterialEffect.cpp
+++ b/Engine/Plugins/VirtualProduction/DynamicMaterial/Source/DynamicMaterialEditor/Private/Components/DMMaterialEffect.cpp
@@ -6,6 +6,8 @@
 #include "Components/DMMaterialStage.h"
 #include "DMComponentPath.h"
 
+#include UE_INLINE_GENERATED_CPP_BY_NAME(DMMaterialEffect)
+
 #define LOCTEXT_NAMESPACE "DMMaterialEffect"
 
 EDMMaterialEffectTarget UDMMaterialEffect::StageTypeToEffectType(EDMMaterialLayerStage InStageType)
diff --git a/Engine/Plugins/VirtualProduction/DynamicMaterial/Source/DynamicMaterialEditor/Private/Components/DMMaterialEffectFunction.cpp b/Engine/Plugins/VirtualProduction/DynamicMaterial/Source/DynamicMaterialEditor/Private/Components/DMMaterialEffectFunction.cpp
--- a/Engine/Plugins/VirtualProduction/DynamicMaterial/Source/DynamicMaterialEditor/Private/Components/DMMaterialEffectFunction.cpp
+++ b/Engine/Plugins/VirtualProduction/DynamicMaterial/Source/DynamicMaterialEditor/Private/Components/DMMaterialEffectFunction.cpp
@@ -26,6 +26,8 @@
 #include "Utils/DMPrivate.h"
 #include "Utils/DMUtils.h"
 
+#include UE_INLINE_GENERATED_CPP_BY_NAME(DMMaterialEffectFunction)
+
 #define LOCTEXT_NAMESPACE "DMMaterialEffectFunction"
 
 const FString UDMMaterialEffectFunction::InputsPathToken = TEXT("Inputs");
diff --git a/Engine/Plugins/VirtualProduction/DynamicMaterial/Source/DynamicMaterialEditor/Private/Components/DMMaterialEffectStack.cpp b/Engine/Plugins/VirtualProduction/DynamicMaterial/Source/DynamicMaterialEditor/Private/Components/DMMaterialEffectStack.cpp
--- a/Engine/Plugins/VirtualProduction/DynamicMaterial/Source/DynamicMaterialEditor/Private/Components/DMMaterialEffectStack.cpp
+++ b/Engine/Plugins/VirtualProduction/DynamicMaterial/Source/DynamicMaterialEditor/Private/Components/DMMaterialEffectStack.cpp
@@ -9,6 +9,8 @@
 #include "UObject/Package.h"
 #include "Utils/DMPrivate.h"
 
+#include UE_INLINE_GENERATED_CPP_BY_NAME(DMMaterialEffectStack)
+
 #define LOCTEXT_NAMESPACE "DMMaterialEffectStack"
 
 const FString UDMMaterialEffectStack::EffectsPathToken = FString(TEXT("Effect"));
diff --git a/Engine/Plugins/VirtualProduction/DynamicMaterial/Source/DynamicMaterialEditor/Private/Components/DMMaterialLayer.cpp b/Engine/Plugins/VirtualProduction/DynamicMaterial/Source/DynamicMaterialEditor/Private/Components/DMMaterialLayer.cpp
--- a/Engine/Plugins/VirtualProduction/DynamicMaterial/Source/DynamicMaterialEditor/Private/Components/DMMaterialLayer.cpp
+++ b/Engine/Plugins/VirtualProduction/DynamicMaterial/Source/DynamicMaterialEditor/Private/Components/DMMaterialLayer.cpp
@@ -18,6 +18,8 @@
 #include "Serialization/JsonWriter.h"
 #include "Utils/DMPrivate.h"
 
+#include UE_INLINE_GENERATED_CPP_BY_NAME(DMMaterialLayer)
+
 #define LOCTEXT_NAMESPACE "DMMaterialLayer"
 
 namespace UE::DynamicMaterialEditor::Private
diff --git a/Engine/Plugins/VirtualProduction/DynamicMaterial/Source/DynamicMaterialEditor/Private/Components/DMMaterialProperty.cpp b/Engine/Plugins/VirtualProduction/DynamicMaterial/Source/DynamicMaterialEditor/Private/Components/DMMaterialProperty.cpp
--- a/Engine/Plugins/VirtualProduction/DynamicMaterial/Source/DynamicMaterialEditor/Private/Components/DMMaterialProperty.cpp
+++ b/Engine/Plugins/VirtualProduction/DynamicMaterial/Source/DynamicMaterialEditor/Private/Components/DMMaterialProperty.cpp
@@ -38,6 +38,8 @@
 #include "Utils/DMPrivate.h"
 #include "Utils/DMUtils.h"
 
+#include UE_INLINE_GENERATED_CPP_BY_NAME(DMMaterialProperty)
+
 #define LOCTEXT_NAMESPACE "DMMaterialProperty"
 
 const FString UDMMaterialProperty::ComponentsPathToken = TEXT("Components");
diff --git a/Engine/Plugins/VirtualProduction/DynamicMaterial/Source/DynamicMaterialEditor/Private/Components/DMMaterialSlot.cpp b/Engine/Plugins/VirtualProduction/DynamicMaterial/Source/DynamicMaterialEditor/Private/Components/DMMaterialSlot.cpp
--- a/Engine/Plugins/VirtualProduction/DynamicMaterial/Source/DynamicMaterialEditor/Private/Components/DMMaterialSlot.cpp
+++ b/Engine/Plugins/VirtualProduction/DynamicMaterial/Source/DynamicMaterialEditor/Private/Components/DMMaterialSlot.cpp
@@ -20,6 +20,8 @@
 #include "UObject/UObjectGlobals.h"
 #include "Utils/DMPrivate.h"
 
+#include UE_INLINE_GENERATED_CPP_BY_NAME(DMMaterialSlot)
+
 #define LOCTEXT_NAMESPACE "DMMaterialStage"
 
 const FString UDMMaterialSlot::LayersPathToken = FString(TEXT("Layers"));
diff --git a/Engine/Plugins/VirtualProduction/DynamicMaterial/Source/DynamicMaterialEditor/Private/Components/DMMaterialStage.cpp b/Engine/Plugins/VirtualProduction/DynamicMaterial/Source/DynamicMaterialEditor/Private/Components/DMMaterialStage.cpp
--- a/Engine/Plugins/VirtualProduction/DynamicMaterial/Source/DynamicMaterialEditor/Private/Components/DMMaterialStage.cpp
+++ b/Engine/Plugins/VirtualProduction/DynamicMaterial/Source/DynamicMaterialEditor/Private/Components/DMMaterialStage.cpp
@@ -20,6 +20,8 @@
 #include "UObject/Package.h"
 #include "Utils/DMPrivate.h"
 
+#include UE_INLINE_GENERATED_CPP_BY_NAME(DMMaterialStage)
+
 #define LOCTEXT_NAMESPACE "DMMaterialStage"
 
 const FString UDMMaterialStage::SourcePathToken = FString(TEXT("Source"));
diff --git a/Engine/Plugins/VirtualProduction/DynamicMaterial/Source/DynamicMaterialEditor/Private/Components/DMMaterialStageBlend.cpp b/Engine/Plugins/VirtualProduction/DynamicMaterial/Source/DynamicMaterialEditor/Private/Components/DMMaterialStageBlend.cpp
--- a/Engine/Plugins/VirtualProduction/DynamicMaterial/Source/DynamicMaterialEditor/Private/Components/DMMaterialStageBlend.cpp
+++ b/Engine/Plugins/VirtualProduction/DynamicMaterial/Source/DynamicMaterialEditor/Private/Components/DMMaterialStageBlend.cpp
@@ -19,6 +19,8 @@
 #include "Model/DynamicMaterialModelEditorOnlyData.h"
 #include "Utils/DMInputNodeBuilder.h"
 #include "Utils/DMPrivate.h"
+
+#include UE_INLINE_GENERATED_CPP_BY_NAME(DMMaterialStageBlend)
  
 #define LOCTEXT_NAMESPACE "DMMaterialProperty"
  
diff --git a/Engine/Plugins/VirtualProduction/DynamicMaterial/Source/DynamicMaterialEditor/Private/Components/DMMaterialStageBlendFunction.cpp b/Engine/Plugins/VirtualProduction/DynamicMaterial/Source/DynamicMaterialEditor/Private/Components/DMMaterialStageBlendFunction.cpp
--- a/Engine/Plugins/VirtualProduction/DynamicMaterial/Source/DynamicMaterialEditor/Private/Components/DMMaterialStageBlendFunction.cpp
+++ b/Engine/Plugins/VirtualProduction/DynamicMaterial/Source/DynamicMaterialEditor/Private/Components/DMMaterialStageBlendFunction.cpp
@@ -10,6 +10,8 @@
 #include "Utils/DMMaterialFunctionLibrary.h"
 #include "Utils/DMUtils.h"
 
+#include UE_INLINE_GENERATED_CPP_BY_NAME(DMMaterialStageBlendFunction)
+
 #define LOCTEXT_NAMESPACE "DMMaterialStageBlendFunction"
 
 UDMMaterialStageBlendFunction::UDMMaterialStageBlendFunction()
diff --git a/Engine/Plugins/VirtualProduction/DynamicMaterial/Source/DynamicMaterialEditor/Private/Components/DMMaterialStageExpression.cpp b/Engine/Plugins/VirtualProduction/DynamicMaterial/Source/DynamicMaterialEditor/Private/Components/DMMaterialStageExpression.cpp
--- a/Engine/Plugins/VirtualProduction/DynamicMaterial/Source/DynamicMaterialEditor/Private/Components/DMMaterialStageExpression.cpp
+++ b/Engine/Plugins/VirtualProduction/DynamicMaterial/Source/DynamicMaterialEditor/Private/Components/DMMaterialStageExpression.cpp
@@ -24,6 +24,8 @@
 #include "Materials/MaterialExpressionRerouteBase.h"
 #include "Materials/MaterialExpressionSubstrate.h"
 
+#include UE_INLINE_GENERATED_CPP_BY_NAME(DMMaterialStageExpression)
+
 #define LOCTEXT_NAMESPACE "DMMaterialProperty"
 
 TSet<FName> UDMMaterialStageExpression::BlockedMaterialExpressionClasses = {
diff --git a/Engine/Plugins/VirtualProduction/DynamicMaterial/Source/DynamicMaterialEditor/Private/Components/DMMaterialStageFunction.cpp b/Engine/Plugins/VirtualProduction/DynamicMaterial/Source/DynamicMaterialEditor/Private/Components/DMMaterialStageFunction.cpp
--- a/Engine/Plugins/VirtualProduction/DynamicMaterial/Source/DynamicMaterialEditor/Private/Components/DMMaterialStageFunction.cpp
+++ b/Engine/Plugins/VirtualProduction/DynamicMaterial/Source/DynamicMaterialEditor/Private/Components/DMMaterialStageFunction.cpp
@@ -25,6 +25,8 @@
 #include "Utils/DMPrivate.h"
 #include "Utils/DMUtils.h"
 
+#include UE_INLINE_GENERATED_CPP_BY_NAME(DMMaterialStageFunction)
+
 #define LOCTEXT_NAMESPACE "DMMaterialStageFunction"
 
 TSoftObjectPtr<UMaterialFunctionInterface> UDMMaterialStageFunction::NoOp = TSoftObjectPtr<UMaterialFunctionInterface>(FSoftObjectPath(TEXT(
diff --git a/Engine/Plugins/VirtualProduction/DynamicMaterial/Source/DynamicMaterialEditor/Private/Components/DMMaterialStageGradient.cpp b/Engine/Plugins/VirtualProduction/DynamicMaterial/Source/DynamicMaterialEditor/Private/Components/DMMaterialStageGradient.cpp
--- a/Engine/Plugins/VirtualProduction/DynamicMaterial/Source/DynamicMaterialEditor/Private/Components/DMMaterialStageGradient.cpp
+++ b/Engine/Plugins/VirtualProduction/DynamicMaterial/Source/DynamicMaterialEditor/Private/Components/DMMaterialStageGradient.cpp
@@ -12,6 +12,8 @@
 #include "Model/DMMaterialBuildState.h"
 #include "Utils/DMUtils.h"
 
+#include UE_INLINE_GENERATED_CPP_BY_NAME(DMMaterialStageGradient)
+
 #define LOCTEXT_NAMESPACE "DMMaterialProperty"
 
 TArray<TStrongObjectPtr<UClass>> UDMMaterialStageGradient::Gradients = {};
diff --git a/Engine/Plugins/VirtualProduction/DynamicMaterial/Source/DynamicMaterialEditor/Private/Components/DMMaterialStageInput.cpp b/Engine/Plugins/VirtualProduction/DynamicMaterial/Source/DynamicMaterialEditor/Private/Components/DMMaterialStageInput.cpp
--- a/Engine/Plugins/VirtualProduction/DynamicMaterial/Source/DynamicMaterialEditor/Private/Components/DMMaterialStageInput.cpp
+++ b/Engine/Plugins/VirtualProduction/DynamicMaterial/Source/DynamicMaterialEditor/Private/Components/DMMaterialStageInput.cpp
@@ -10,6 +10,8 @@
 #include "Model/DMMaterialBuildUtils.h"
 #include "Model/DynamicMaterialModelEditorOnlyData.h"
 
+#include UE_INLINE_GENERATED_CPP_BY_NAME(DMMaterialStageInput)
+
 const FString UDMMaterialStageInput::StageInputPrefixStr = FString(TEXT("DMMaterialStageInput"));
 
 void UDMMaterialStageInput::Update(UDMMaterialComponent* InSource, EDMUpdateType InUpdateType)
diff --git a/Engine/Plugins/VirtualProduction/DynamicMaterial/Source/DynamicMaterialEditor/Private/Components/DMMaterialStageSource.cpp b/Engine/Plugins/VirtualProduction/DynamicMaterial/Source/DynamicMaterialEditor/Private/Components/DMMaterialStageSource.cpp
--- a/Engine/Plugins/VirtualProduction/DynamicMaterial/Source/DynamicMaterialEditor/Private/Components/DMMaterialStageSource.cpp
+++ b/Engine/Plugins/VirtualProduction/DynamicMaterial/Source/DynamicMaterialEditor/Private/Components/DMMaterialStageSource.cpp
@@ -15,6 +15,8 @@
 #include "UObject/Package.h"
 #include "UObject/UObjectIterator.h"
 
+#include UE_INLINE_GENERATED_CPP_BY_NAME(DMMaterialStageSource)
+
 TArray<TStrongObjectPtr<UClass>> UDMMaterialStageSource::SourceClasses = TArray<TStrongObjectPtr<UClass>>();
 
 UDMMaterialStage* UDMMaterialStageSource::GetStage() const
diff --git a/Engine/Plugins/VirtualProduction/DynamicMaterial/Source/DynamicMaterialEditor/Private/Components/DMMaterialStageThroughput.cpp b/Engine/Plugins/VirtualProduction/DynamicMaterial/Source/DynamicMaterialEditor/Private/Components/DMMaterialStageThroughput.cpp
--- a/Engine/Plugins/VirtualProduction/DynamicMaterial/Source/DynamicMaterialEditor/Private/Components/DMMaterialStageThroughput.cpp
+++ b/Engine/Plugins/VirtualProduction/DynamicMaterial/Source/DynamicMaterialEditor/Private/Components/DMMaterialStageThroughput.cpp
@@ -30,6 +30,8 @@
 #include "Utils/DMPrivate.h"
 #include "Utils/DMUtils.h"
 
+#include UE_INLINE_GENERATED_CPP_BY_NAME(DMMaterialStageThroughput)
+
 #define LOCTEXT_NAMESPACE "DMMaterialStageThroughput"
 
 TArray<TStrongObjectPtr<UClass>> UDMMaterialStageThroughput::Throughputs = {};
diff --git a/Engine/Plugins/VirtualProduction/DynamicMaterial/Source/DynamicMaterialEditor/Private/Components/DMMaterialStageThroughputLayerBlend.cpp b/Engine/Plugins/VirtualProduction/DynamicMaterial/Source/DynamicMaterialEditor/Private/Components/DMMaterialStageThroughputLayerBlend.cpp
--- a/Engine/Plugins/VirtualProduction/DynamicMaterial/Source/DynamicMaterialEditor/Private/Components/DMMaterialStageThroughputLayerBlend.cpp
+++ b/Engine/Plugins/VirtualProduction/DynamicMaterial/Source/DynamicMaterialEditor/Private/Components/DMMaterialStageThroughputLayerBlend.cpp
@@ -28,6 +28,8 @@
 #include "Utils/DMPrivate.h"
 #include "Utils/DMUtils.h"
 
+#include UE_INLINE_GENERATED_CPP_BY_NAME(DMMaterialStageThroughputLayerBlend)
+
 #define LOCTEXT_NAMESPACE "DMMaterialStageThroughputLayerBlend"
 
 namespace UE::DynamicMaterialEditor::Private
diff --git a/Engine/Plugins/VirtualProduction/DynamicMaterial/Source/DynamicMaterialEditor/Private/Components/DMMaterialSubStage.cpp b/Engine/Plugins/VirtualProduction/DynamicMaterial/Source/DynamicMaterialEditor/Private/Components/DMMaterialSubStage.cpp
--- a/Engine/Plugins/VirtualProduction/DynamicMaterial/Source/DynamicMaterialEditor/Private/Components/DMMaterialSubStage.cpp
+++ b/Engine/Plugins/VirtualProduction/DynamicMaterial/Source/DynamicMaterialEditor/Private/Components/DMMaterialSubStage.cpp
@@ -6,6 +6,8 @@
 #include "Components/MaterialStageInputs/DMMSIThroughput.h"
 #include "Utils/DMPrivate.h"
 
+#include UE_INLINE_GENERATED_CPP_BY_NAME(DMMaterialSubStage)
+
 UDMMaterialSubStage* UDMMaterialSubStage::CreateMaterialSubStage(UDMMaterialStage* InParentStage)
 {
 	UDMMaterialLayerObject* Layer = InParentStage->GetLayer();
diff --git a/Engine/Plugins/VirtualProduction/DynamicMaterial/Source/DynamicMaterialEditor/Private/Components/MaterialProperties/DMMPAmbientOcclusion.cpp b/Engine/Plugins/VirtualProduction/DynamicMaterial/Source/DynamicMaterialEditor/Private/Components/MaterialProperties/DMMPAmbientOcclusion.cpp
--- a/Engine/Plugins/VirtualProduction/DynamicMaterial/Source/DynamicMaterialEditor/Private/Components/MaterialProperties/DMMPAmbientOcclusion.cpp
+++ b/Engine/Plugins/VirtualProduction/DynamicMaterial/Source/DynamicMaterialEditor/Private/Components/MaterialProperties/DMMPAmbientOcclusion.cpp
@@ -3,6 +3,8 @@
 #include "Components/MaterialProperties/DMMPAmbientOcclusion.h"
 #include "Model/DynamicMaterialModelEditorOnlyData.h"
 
+#include UE_INLINE_GENERATED_CPP_BY_NAME(DMMPAmbientOcclusion)
+
 UDMMaterialPropertyAmbientOcclusion::UDMMaterialPropertyAmbientOcclusion()
 	: UDMMaterialProperty(
 		EDMMaterialPropertyType(EDMMaterialPropertyType::AmbientOcclusion),
diff --git a/Engine/Plugins/VirtualProduction/DynamicMaterial/Source/DynamicMaterialEditor/Private/Components/MaterialProperties/DMMPAnisotropy.cpp b/Engine/Plugins/VirtualProduction/DynamicMaterial/Source/DynamicMaterialEditor/Private/Components/MaterialProperties/DMMPAnisotropy.cpp
--- a/Engine/Plugins/VirtualProduction/DynamicMaterial/Source/DynamicMaterialEditor/Private/Components/MaterialProperties/DMMPAnisotropy.cpp
+++ b/Engine/Plugins/VirtualProduction/DynamicMaterial/Source/DynamicMaterialEditor/Private/Components/MaterialProperties/DMMPAnisotropy.cpp
@@ -3,6 +3,8 @@
 #include "Components/MaterialProperties/DMMPAnisotropy.h"
 #include "Model/DynamicMaterialModelEditorOnlyData.h"
 
+#include UE_INLINE_GENERATED_CPP_BY_NAME(DMMPAnisotropy)
+
 UDMMaterialPropertyAnisotropy::UDMMaterialPropertyAnisotropy()
 	: UDMMaterialProperty(
 		EDMMaterialPropertyType(EDMMaterialPropertyType::Anisotropy),
diff --git a/Engine/Plugins/VirtualProduction/DynamicMaterial/Source/DynamicMaterialEditor/Private/Components/MaterialProperties/DMMPBaseColor.cpp b/Engine/Plugins/VirtualProduction/DynamicMaterial/Source/DynamicMaterialEditor/Private/Components/MaterialProperties/DMMPBaseColor.cpp
--- a/Engine/Plugins/VirtualProduction/DynamicMaterial/Source/DynamicMaterialEditor/Private/Components/MaterialProperties/DMMPBaseColor.cpp
+++ b/Engine/Plugins/VirtualProduction/DynamicMaterial/Source/DynamicMaterialEditor/Private/Components/MaterialProperties/DMMPBaseColor.cpp
@@ -4,6 +4,8 @@
 #include "Components/DMMaterialSlot.h"
 #include "Model/DynamicMaterialModelEditorOnlyData.h"
 
+#include UE_INLINE_GENERATED_CPP_BY_NAME(DMMPBaseColor)
+
 UDMMaterialPropertyBaseColor::UDMMaterialPropertyBaseColor()
 	: UDMMaterialProperty(
 		EDMMaterialPropertyType(EDMMaterialPropertyType::BaseColor),
diff --git a/Engine/Plugins/VirtualProduction/DynamicMaterial/Source/DynamicMaterialEditor/Private/Components/MaterialProperties/DMMPDisplacement.cpp b/Engine/Plugins/VirtualProduction/DynamicMaterial/Source/DynamicMaterialEditor/Private/Components/MaterialProperties/DMMPDisplacement.cpp
--- a/Engine/Plugins/VirtualProduction/DynamicMaterial/Source/DynamicMaterialEditor/Private/Components/MaterialProperties/DMMPDisplacement.cpp
+++ b/Engine/Plugins/VirtualProduction/DynamicMaterial/Source/DynamicMaterialEditor/Private/Components/MaterialProperties/DMMPDisplacement.cpp
@@ -10,6 +10,8 @@
 #include "Utils/DMMaterialFunctionLibrary.h"
 #include "Utils/DMUtils.h"
 
+#include UE_INLINE_GENERATED_CPP_BY_NAME(DMMPDisplacement)
+
 namespace UE::DynamicMaterialEditor::Private
 {
 	UMaterialFunctionInterface* GetDisplacementAlpha()
diff --git a/Engine/Plugins/VirtualProduction/DynamicMaterial/Source/DynamicMaterialEditor/Private/Components/MaterialProperties/DMMPEmissiveColor.cpp b/Engine/Plugins/VirtualProduction/DynamicMaterial/Source/DynamicMaterialEditor/Private/Components/MaterialProperties/DMMPEmissiveColor.cpp
--- a/Engine/Plugins/VirtualProduction/DynamicMaterial/Source/DynamicMaterialEditor/Private/Components/MaterialProperties/DMMPEmissiveColor.cpp
+++ b/Engine/Plugins/VirtualProduction/DynamicMaterial/Source/DynamicMaterialEditor/Private/Components/MaterialProperties/DMMPEmissiveColor.cpp
@@ -5,6 +5,8 @@
 #include "Components/MaterialProperties/DMMPBaseColor.h"
 #include "Model/DynamicMaterialModelEditorOnlyData.h"
 
+#include UE_INLINE_GENERATED_CPP_BY_NAME(DMMPEmissiveColor)
+
 UDMMaterialPropertyEmissiveColor::UDMMaterialPropertyEmissiveColor()
 	: UDMMaterialProperty(
 		EDMMaterialPropertyType(EDMMaterialPropertyType::EmissiveColor),
diff --git a/Engine/Plugins/VirtualProduction/DynamicMaterial/Source/DynamicMaterialEditor/Private/Components/MaterialProperties/DMMPMetallic.cpp b/Engine/Plugins/VirtualProduction/DynamicMaterial/Source/DynamicMaterialEditor/Private/Components/MaterialProperties/DMMPMetallic.cpp
--- a/Engine/Plugins/VirtualProduction/DynamicMaterial/Source/DynamicMaterialEditor/Private/Components/MaterialProperties/DMMPMetallic.cpp
+++ b/Engine/Plugins/VirtualProduction/DynamicMaterial/Source/DynamicMaterialEditor/Private/Components/MaterialProperties/DMMPMetallic.cpp
@@ -3,6 +3,8 @@
 #include "Components/MaterialProperties/DMMPMetallic.h"
 #include "Model/DynamicMaterialModelEditorOnlyData.h"
 
+#include UE_INLINE_GENERATED_CPP_BY_NAME(DMMPMetallic)
+
 UDMMaterialPropertyMetallic::UDMMaterialPropertyMetallic()
 	: UDMMaterialProperty(
 		EDMMaterialPropertyType(EDMMaterialPropertyType::Metallic),
diff --git a/Engine/Plugins/VirtualProduction/DynamicMaterial/Source/DynamicMaterialEditor/Private/Components/MaterialProperties/DMMPNormal.cpp b/Engine/Plugins/VirtualProduction/DynamicMaterial/Source/DynamicMaterialEditor/Private/Components/MaterialProperties/DMMPNormal.cpp
--- a/Engine/Plugins/VirtualProduction/DynamicMaterial/Source/DynamicMaterialEditor/Private/Components/MaterialProperties/DMMPNormal.cpp
+++ b/Engine/Plugins/VirtualProduction/DynamicMaterial/Source/DynamicMaterialEditor/Private/Components/MaterialProperties/DMMPNormal.cpp
@@ -9,6 +9,8 @@
 #include "Utils/DMMaterialFunctionLibrary.h"
 #include "Utils/DMUtils.h"
 
+#include UE_INLINE_GENERATED_CPP_BY_NAME(DMMPNormal)
+
 namespace UE::DynamicMaterialEditor::Private
 {
 	UMaterialFunctionInterface* GetSafeNormalize()
diff --git a/Engine/Plugins/VirtualProduction/DynamicMaterial/Source/DynamicMaterialEditor/Private/Components/MaterialProperties/DMMPOpacity.cpp b/Engine/Plugins/VirtualProduction/DynamicMaterial/Source/DynamicMaterialEditor/Private/Components/MaterialProperties/DMMPOpacity.cpp
--- a/Engine/Plugins/VirtualProduction/DynamicMaterial/Source/DynamicMaterialEditor/Private/Components/MaterialProperties/DMMPOpacity.cpp
+++ b/Engine/Plugins/VirtualProduction/DynamicMaterial/Source/DynamicMaterialEditor/Private/Components/MaterialProperties/DMMPOpacity.cpp
@@ -4,6 +4,8 @@
 #include "Components/DMMaterialLayer.h"
 #include "Model/DynamicMaterialModelEditorOnlyData.h"
 
+#include UE_INLINE_GENERATED_CPP_BY_NAME(DMMPOpacity)
+
 UDMMaterialPropertyOpacity::UDMMaterialPropertyOpacity()
 	: UDMMaterialProperty(
 		EDMMaterialPropertyType(EDMMaterialPropertyType::Opacity),
diff --git a/Engine/Plugins/VirtualProduction/DynamicMaterial/Source/DynamicMaterialEditor/Private/Components/MaterialProperties/DMMPOpacityMask.cpp b/Engine/Plugins/VirtualProduction/DynamicMaterial/Source/DynamicMaterialEditor/Private/Components/MaterialProperties/DMMPOpacityMask.cpp
--- a/Engine/Plugins/VirtualProduction/DynamicMaterial/Source/DynamicMaterialEditor/Private/Components/MaterialProperties/DMMPOpacityMask.cpp
+++ b/Engine/Plugins/VirtualProduction/DynamicMaterial/Source/DynamicMaterialEditor/Private/Components/MaterialProperties/DMMPOpacityMask.cpp
@@ -5,6 +5,8 @@
 #include "Components/MaterialProperties/DMMPOpacity.h"
 #include "Model/DynamicMaterialModelEditorOnlyData.h"
 
+#include UE_INLINE_GENERATED_CPP_BY_NAME(DMMPOpacityMask)
+
 UDMMaterialPropertyOpacityMask::UDMMaterialPropertyOpacityMask()
 	: UDMMaterialProperty(
 		EDMMaterialPropertyType(EDMMaterialPropertyType::OpacityMask),
diff --git a/Engine/Plugins/VirtualProduction/DynamicMaterial/Source/DynamicMaterialEditor/Private/Components/MaterialProperties/DMMPPixelDepthOffset.cpp b/Engine/Plugins/VirtualProduction/DynamicMaterial/Source/DynamicMaterialEditor/Private/Components/MaterialProperties/DMMPPixelDepthOffset.cpp
--- a/Engine/Plugins/VirtualProduction/DynamicMaterial/Source/DynamicMaterialEditor/Private/Components/MaterialProperties/DMMPPixelDepthOffset.cpp
+++ b/Engine/Plugins/VirtualProduction/DynamicMaterial/Source/DynamicMaterialEditor/Private/Components/MaterialProperties/DMMPPixelDepthOffset.cpp
@@ -3,6 +3,8 @@
 #include "Components/MaterialProperties/DMMPPixelDepthOffset.h"
 #include "Model/DynamicMaterialModelEditorOnlyData.h"
 
+#include UE_INLINE_GENERATED_CPP_BY_NAME(DMMPPixelDepthOffset)
+
 UDMMaterialPropertyPixelDepthOffset::UDMMaterialPropertyPixelDepthOffset()
 	: UDMMaterialProperty(
 		EDMMaterialPropertyType(EDMMaterialPropertyType::PixelDepthOffset),
diff --git a/Engine/Plugins/VirtualProduction/DynamicMaterial/Source/DynamicMaterialEditor/Private/Components/MaterialProperties/DMMPRefraction.cpp b/Engine/Plugins/VirtualProduction/DynamicMaterial/Source/DynamicMaterialEditor/Private/Components/MaterialProperties/DMMPRefraction.cpp
--- a/Engine/Plugins/VirtualProduction/DynamicMaterial/Source/DynamicMaterialEditor/Private/Components/MaterialProperties/DMMPRefraction.cpp
+++ b/Engine/Plugins/VirtualProduction/DynamicMaterial/Source/DynamicMaterialEditor/Private/Components/MaterialProperties/DMMPRefraction.cpp
@@ -3,6 +3,8 @@
 #include "Components/MaterialProperties/DMMPRefraction.h"
 #include "Model/DynamicMaterialModelEditorOnlyData.h"
 
+#include UE_INLINE_GENERATED_CPP_BY_NAME(DMMPRefraction)
+
 UDMMaterialPropertyRefraction::UDMMaterialPropertyRefraction()
 	: UDMMaterialProperty(
 		EDMMaterialPropertyType(EDMMaterialPropertyType::Refraction),
diff --git a/Engine/Plugins/VirtualProduction/DynamicMaterial/Source/DynamicMaterialEditor/Private/Components/MaterialProperties/DMMPRoughness.cpp b/Engine/Plugins/VirtualProduction/DynamicMaterial/Source/DynamicMaterialEditor/Private/Components/MaterialProperties/DMMPRoughness.cpp
--- a/Engine/Plugins/VirtualProduction/DynamicMaterial/Source/DynamicMaterialEditor/Private/Components/MaterialProperties/DMMPRoughness.cpp
+++ b/Engine/Plugins/VirtualProduction/DynamicMaterial/Source/DynamicMaterialEditor/Private/Components/MaterialProperties/DMMPRoughness.cpp
@@ -3,6 +3,8 @@
 #include "Components/MaterialProperties/DMMPRoughness.h"
 #include "Model/DynamicMaterialModelEditorOnlyData.h"
 
+#include UE_INLINE_GENERATED_CPP_BY_NAME(DMMPRoughness)
+
 UDMMaterialPropertyRoughness::UDMMaterialPropertyRoughness()
 	: UDMMaterialProperty(
 		EDMMaterialPropertyType(EDMMaterialPropertyType::Roughness),
diff --git a/Engine/Plugins/VirtualProduction/DynamicMaterial/Source/DynamicMaterialEditor/Private/Components/MaterialProperties/DMMPSpecular.cpp b/Engine/Plugins/VirtualProduction/DynamicMaterial/Source/DynamicMaterialEditor/Private/Components/MaterialProperties/DMMPSpecular.cpp
--- a/Engine/Plugins/VirtualProduction/DynamicMaterial/Source/DynamicMaterialEditor/Private/Components/MaterialProperties/DMMPSpecular.cpp
+++ b/Engine/Plugins/VirtualProduction/DynamicMaterial/Source/DynamicMaterialEditor/Private/Components/MaterialProperties/DMMPSpecular.cpp
@@ -3,6 +3,8 @@
 #include "Components/MaterialProperties/DMMPSpecular.h"
 #include "Model/DynamicMaterialModelEditorOnlyData.h"
 
+#include UE_INLINE_GENERATED_CPP_BY_NAME(DMMPSpecular)
+
 UDMMaterialPropertySpecular::UDMMaterialPropertySpecular()
 	: UDMMaterialProperty(
 		EDMMaterialPropertyType(EDMMaterialPropertyType::Specular),
diff --git a/Engine/Plugins/VirtualProduction/DynamicMaterial/Source/DynamicMaterialEditor/Private/Components/MaterialProperties/DMMPSubsurfaceColor.cpp b/Engine/Plugins/VirtualProduction/DynamicMaterial/Source/DynamicMaterialEditor/Private/Components/MaterialProperties/DMMPSubsurfaceColor.cpp
--- a/Engine/Plugins/VirtualProduction/DynamicMaterial/Source/DynamicMaterialEditor/Private/Components/MaterialProperties/DMMPSubsurfaceColor.cpp
+++ b/Engine/Plugins/VirtualProduction/DynamicMaterial/Source/DynamicMaterialEditor/Private/Components/MaterialProperties/DMMPSubsurfaceColor.cpp
@@ -3,6 +3,8 @@
 #include "Components/MaterialProperties/DMMPSubsurfaceColor.h"
 #include "Model/DynamicMaterialModelEditorOnlyData.h"
 
+#include UE_INLINE_GENERATED_CPP_BY_NAME(DMMPSubsurfaceColor)
+
 UDMMaterialPropertySubsurfaceColor::UDMMaterialPropertySubsurfaceColor()
 	: UDMMaterialProperty(
 		EDMMaterialPropertyType(EDMMaterialPropertyType::SubsurfaceColor),
diff --git a/Engine/Plugins/VirtualProduction/DynamicMaterial/Source/DynamicMaterialEditor/Private/Components/MaterialProperties/DMMPSurfaceThickness.cpp b/Engine/Plugins/VirtualProduction/DynamicMaterial/Source/DynamicMaterialEditor/Private/Components/MaterialProperties/DMMPSurfaceThickness.cpp
--- a/Engine/Plugins/VirtualProduction/DynamicMaterial/Source/DynamicMaterialEditor/Private/Components/MaterialProperties/DMMPSurfaceThickness.cpp
+++ b/Engine/Plugins/VirtualProduction/DynamicMaterial/Source/DynamicMaterialEditor/Private/Components/MaterialProperties/DMMPSurfaceThickness.cpp
@@ -3,6 +3,8 @@
 #include "Components/MaterialProperties/DMMPSurfaceThickness.h"
 #include "Model/DynamicMaterialModelEditorOnlyData.h"
 
+#include UE_INLINE_GENERATED_CPP_BY_NAME(DMMPSurfaceThickness)
+
 UDMMaterialPropertySurfaceThickness::UDMMaterialPropertySurfaceThickness()
 	: UDMMaterialProperty(
 		EDMMaterialPropertyType(EDMMaterialPropertyType::SurfaceThickness),
diff --git a/Engine/Plugins/VirtualProduction/DynamicMaterial/Source/DynamicMaterialEditor/Private/Components/MaterialProperties/DMMPTangent.cpp b/Engine/Plugins/VirtualProduction/DynamicMaterial/Source/DynamicMaterialEditor/Private/Components/MaterialProperties/DMMPTangent.cpp
--- a/Engine/Plugins/VirtualProduction/DynamicMaterial/Source/DynamicMaterialEditor/Private/Components/MaterialProperties/DMMPTangent.cpp
+++ b/Engine/Plugins/VirtualProduction/DynamicMaterial/Source/DynamicMaterialEditor/Private/Components/MaterialProperties/DMMPTangent.cpp
@@ -3,6 +3,8 @@
 #include "Components/MaterialProperties/DMMPTangent.h"
 #include "Model/DynamicMaterialModelEditorOnlyData.h"
 
+#include UE_INLINE_GENERATED_CPP_BY_NAME(DMMPTangent)
+
 UDMMaterialPropertyTangent::UDMMaterialPropertyTangent()
 	: UDMMaterialProperty(
 		EDMMaterialPropertyType(EDMMaterialPropertyType::Tangent),
diff --git a/Engine/Plugins/VirtualProduction/DynamicMaterial/Source/DynamicMaterialEditor/Private/Components/MaterialProperties/DMMPWorldPositionOffset.cpp b/Engine/Plugins/VirtualProduction/DynamicMaterial/Source/DynamicMaterialEditor/Private/Components/MaterialProperties/DMMPWorldPositionOffset.cpp
--- a/Engine/Plugins/VirtualProduction/DynamicMaterial/Source/DynamicMaterialEditor/Private/Components/MaterialProperties/DMMPWorldPositionOffset.cpp
+++ b/Engine/Plugins/VirtualProduction/DynamicMaterial/Source/DynamicMaterialEditor/Private/Components/MaterialProperties/DMMPWorldPositionOffset.cpp
@@ -3,6 +3,8 @@
 #include "Components/MaterialProperties/DMMPWorldPositionOffset.h"
 #include "Model/DynamicMaterialModelEditorOnlyData.h"
 
+#include UE_INLINE_GENERATED_CPP_BY_NAME(DMMPWorldPositionOffset)
+
 UDMMaterialPropertyWorldPositionOffset::UDMMaterialPropertyWorldPositionOffset()
 	: UDMMaterialProperty(
 		EDMMaterialPropertyType(EDMMaterialPropertyType::WorldPositionOffset),
diff --git a/Engine/Plugins/VirtualProduction/DynamicMaterial/Source/DynamicMaterialEditor/Private/Components/MaterialStageBlends/DMMSBAdd.cpp b/Engine/Plugins/VirtualProduction/DynamicMaterial/Source/DynamicMaterialEditor/Private/Components/MaterialStageBlends/DMMSBAdd.cpp
--- a/Engine/Plugins/VirtualProduction/DynamicMaterial/Source/DynamicMaterialEditor/Private/Components/MaterialStageBlends/DMMSBAdd.cpp
+++ b/Engine/Plugins/VirtualProduction/DynamicMaterial/Source/DynamicMaterialEditor/Private/Components/MaterialStageBlends/DMMSBAdd.cpp
@@ -4,6 +4,8 @@
 
 #include "Materials/MaterialExpressionAdd.h"
 
+#include UE_INLINE_GENERATED_CPP_BY_NAME(DMMSBAdd)
+
 #define LOCTEXT_NAMESPACE "DMMaterialStageBlendAdd"
 
 UDMMaterialStageBlendAdd::UDMMaterialStageBlendAdd()
diff --git a/Engine/Plugins/VirtualProduction/DynamicMaterial/Source/DynamicMaterialEditor/Private/Components/MaterialStageBlends/DMMSBColor.cpp b/Engine/Plugins/VirtualProduction/DynamicMaterial/Source/DynamicMaterialEditor/Private/Components/MaterialStageBlends/DMMSBColor.cpp
--- a/Engine/Plugins/VirtualProduction/DynamicMaterial/Source/DynamicMaterialEditor/Private/Components/MaterialStageBlends/DMMSBColor.cpp
+++ b/Engine/Plugins/VirtualProduction/DynamicMaterial/Source/DynamicMaterialEditor/Private/Components/MaterialStageBlends/DMMSBColor.cpp
@@ -2,6 +2,8 @@
 
 #include "Components/MaterialStageBlends/DMMSBColor.h"
 
+#include UE_INLINE_GENERATED_CPP_BY_NAME(DMMSBColor)
+
 #define LOCTEXT_NAMESPACE "DMMaterialStageBlendColor"
 
 UDMMaterialStageBlendColor::UDMMaterialStageBlendColor()
diff --git a/Engine/Plugins/VirtualProduction/DynamicMaterial/Source/DynamicMaterialEditor/Private/Components/MaterialStageBlends/DMMSBColorBurn.cpp b/Engine/Plugins/VirtualProduction/DynamicMaterial/Source/DynamicMaterialEditor/Private/Components/MaterialStageBlends/DMMSBColorBurn.cpp
--- a/Engine/Plugins/VirtualProduction/DynamicMaterial/Source/DynamicMaterialEditor/Private/Components/MaterialStageBlends/DMMSBColorBurn.cpp
+++ b/Engine/Plugins/VirtualProduction/DynamicMaterial/Source/DynamicMaterialEditor/Private/Components/MaterialStageBlends/DMMSBColorBurn.cpp
@@ -2,6 +2,8 @@
 
 #include "Components/MaterialStageBlends/DMMSBColorBurn.h"
 
+#include UE_INLINE_GENERATED_CPP_BY_NAME(DMMSBColorBurn)
+
 #define LOCTEXT_NAMESPACE "DMMaterialStageBlendColorBurn"
 
 UDMMaterialStageBlendColorBurn::UDMMaterialStageBlendColorBurn()
diff --git a/Engine/Plugins/VirtualProduction/DynamicMaterial/Source/DynamicMaterialEditor/Private/Components/MaterialStageBlends/DMMSBColorDodge.cpp b/Engine/Plugins/VirtualProduction/DynamicMaterial/Source/DynamicMaterialEditor/Private/Components/MaterialStageBlends/DMMSBColorDodge.cpp
--- a/Engine/Plugins/VirtualProduction/DynamicMaterial/Source/DynamicMaterialEditor/Private/Components/MaterialStageBlends/DMMSBColorDodge.cpp
+++ b/Engine/Plugins/VirtualProduction/DynamicMaterial/Source/DynamicMaterialEditor/Private/Components/MaterialStageBlends/DMMSBColorDodge.cpp
@@ -2,6 +2,8 @@
 
 #include "Components/MaterialStageBlends/DMMSBColorDodge.h"
 
+#include UE_INLINE_GENERATED_CPP_BY_NAME(DMMSBColorDodge)
+
 #define LOCTEXT_NAMESPACE "DMMaterialStageBlendColorDodge"
 
 UDMMaterialStageBlendColorDodge::UDMMaterialStageBlendColorDodge()
diff --git a/Engine/Plugins/VirtualProduction/DynamicMaterial/Source/DynamicMaterialEditor/Private/Components/MaterialStageBlends/DMMSBContrastBase.cpp b/Engine/Plugins/VirtualProduction/DynamicMaterial/Source/DynamicMaterialEditor/Private/Components/MaterialStageBlends/DMMSBContrastBase.cpp
--- a/Engine/Plugins/VirtualProduction/DynamicMaterial/Source/DynamicMaterialEditor/Private/Components/MaterialStageBlends/DMMSBContrastBase.cpp
+++ b/Engine/Plugins/VirtualProduction/DynamicMaterial/Source/DynamicMaterialEditor/Private/Components/MaterialStageBlends/DMMSBContrastBase.cpp
@@ -2,6 +2,8 @@
 
 #include "Components/MaterialStageBlends/DMMSBContrastBase.h"
 
+#include UE_INLINE_GENERATED_CPP_BY_NAME(DMMSBContrastBase)
+
 #define LOCTEXT_NAMESPACE "DMMaterialStageBlendContrastBase"
 
 UDMMaterialStageBlendContrastBase::UDMMaterialStageBlendContrastBase()
