check if midi streams are the same before merging.

#rnx

diff --git a/Engine/Plugins/Runtime/Harmonix/Source/HarmonixMetasound/Private/HarmonixMetasound/DataTypes/MidiStream.cpp b/Engine/Plugins/Runtime/Harmonix/Source/HarmonixMetasound/Private/HarmonixMetasound/DataTypes/MidiStream.cpp
--- a/Engine/Plugins/Runtime/Harmonix/Source/HarmonixMetasound/Private/HarmonixMetasound/DataTypes/MidiStream.cpp
+++ b/Engine/Plugins/Runtime/Harmonix/Source/HarmonixMetasound/Private/HarmonixMetasound/DataTypes/MidiStream.cpp
@@ -133,6 +133,11 @@ namespace HarmonixMetasound
 
 	void FMidiStream::Merge(const FMidiStream& From, FMidiStream& To, const FEventFilter& Filter, const FEventTransformer& Transformer)
 	{
+		if (!ensure(&From != &To))
+		{
+			return;
+		}
+		
 		const auto FromClock = From.GetClock();
 		const auto ToClock = To.GetClock();
 
