Minor fixes and improvements to MusicClockComponent to better support custom clock drivers.

WallClockMusicClockDriver now sets ExperiencedTime directly instead of AudioRenderTime and applies latency offsets for AudioRender and VideoRender times, because the GameThread is querying wall time.
WallClockMusicClockDriver now sets RawAudioRenderTime to be equal to AudioRenderTime, as this driver doesn't apply smoothing, and was previously returning an invalid SongPos for RawAudioRenderTime.

#rnx

diff --git a/Engine/Plugins/Runtime/Harmonix/Source/HarmonixMetasound/Private/HarmonixMetasound/Components/MusicClockComponent.cpp b/Engine/Plugins/Runtime/Harmonix/Source/HarmonixMetasound/Private/HarmonixMetasound/Components/MusicClockComponent.cpp
--- a/Engine/Plugins/Runtime/Harmonix/Source/HarmonixMetasound/Private/HarmonixMetasound/Components/MusicClockComponent.cpp
+++ b/Engine/Plugins/Runtime/Harmonix/Source/HarmonixMetasound/Private/HarmonixMetasound/Components/MusicClockComponent.cpp
@@ -84,8 +84,7 @@ UMusicClockComponent* UMusicClockComponent::CreateMusicClockWithSettings(UObject
 	}
 	else
 	{
-		NewClock->MusicClockSettings = MusicClockSettings;
-		NewClock->ConnectToCustomClock();	
+		NewClock->ConnectToCustomClockWithSettings(MusicClockSettings);
 	}
 	
 	if (Start)
@@ -110,6 +109,13 @@ void UMusicClockComponent::ConnectToWallClockForMidi(UMidiFile* InTempoMap)
 	ConnectToWallClock();
 }
 
+void UMusicClockComponent::ConnectToCustomClockWithSettings(TInstancedStruct<FMusicClockSettingsBase> InMusicClockSettings)
+{
+	DriveMethod = EMusicClockDriveMethod::Custom;
+	MusicClockSettings = InMusicClockSettings;
+	ConnectToCustomClock();
+}
+
 void UMusicClockComponent::SetDefaultTempo(float TempoBpm)
 {
 	// Defaults only used to make the tempo map, once before ever running.
diff --git a/Engine/Plugins/Runtime/Harmonix/Source/HarmonixMetasound/Private/HarmonixMetasound/Components/WallClockMusicClockDriver.cpp b/Engine/Plugins/Runtime/Harmonix/Source/HarmonixMetasound/Private/HarmonixMetasound/Components/WallClockMusicClockDriver.cpp
--- a/Engine/Plugins/Runtime/Harmonix/Source/HarmonixMetasound/Private/HarmonixMetasound/Components/WallClockMusicClockDriver.cpp
+++ b/Engine/Plugins/Runtime/Harmonix/Source/HarmonixMetasound/Private/HarmonixMetasound/Components/WallClockMusicClockDriver.cpp
@@ -41,9 +41,12 @@ bool FWallClockMusicClockDriver::RefreshCurrentSongPos()
 	const ISongMapEvaluator* Maps = GetCurrentSongMapEvaluator();
 	check(Maps);
 
-	CurrentSmoothedAudioRenderSongPos.SetByTime((float)(RunTime * 1000.0), *Maps);
-	CurrentPlayerExperiencedSongPos.SetByTime(CurrentSmoothedAudioRenderSongPos.SecondsIncludingCountIn * 1000.0f - FHarmonixModule::GetMeasuredUserExperienceAndReactionToAudioRenderOffsetMs(), *Maps);
-	CurrentVideoRenderSongPos.SetByTime(CurrentSmoothedAudioRenderSongPos.SecondsIncludingCountIn * 1000.0f - FHarmonixModule::GetMeasuredVideoToAudioRenderOffsetMs(), *Maps);
+	// Set Experienced time here, as this code is running on the GameThread, near the same time as input is processed
+	CurrentPlayerExperiencedSongPos.SetByTime((float)(RunTime * 1000.0), *Maps);
+	CurrentRawAudioRenderSongPos.SetByTime(CurrentPlayerExperiencedSongPos.SecondsIncludingCountIn * 1000.0f + FHarmonixModule::GetMeasuredUserExperienceAndReactionToAudioRenderOffsetMs(), *Maps);
+	CurrentVideoRenderSongPos.SetByTime(CurrentRawAudioRenderSongPos.SecondsIncludingCountIn * 1000.0f - FHarmonixModule::GetMeasuredVideoToAudioRenderOffsetMs(), *Maps);
+	CurrentSmoothedAudioRenderSongPos = CurrentRawAudioRenderSongPos;
+
 	UpdateMusicPlaybackRate(CurrentSmoothedAudioRenderSongPos.Tempo, CurrentClockAdvanceRate, CurrentSmoothedAudioRenderSongPos.TimeSigNumerator, CurrentSmoothedAudioRenderSongPos.TimeSigDenominator);
 	return true;
 }
diff --git a/Engine/Plugins/Runtime/Harmonix/Source/HarmonixMetasound/Public/HarmonixMetasound/Components/MusicClockComponent.h b/Engine/Plugins/Runtime/Harmonix/Source/HarmonixMetasound/Public/HarmonixMetasound/Components/MusicClockComponent.h
--- a/Engine/Plugins/Runtime/Harmonix/Source/HarmonixMetasound/Public/HarmonixMetasound/Components/MusicClockComponent.h
+++ b/Engine/Plugins/Runtime/Harmonix/Source/HarmonixMetasound/Public/HarmonixMetasound/Components/MusicClockComponent.h
@@ -478,7 +478,7 @@ public:
 	static UE_API UMusicClockComponent* CreateMusicClockWithSettings(UObject* WorldContextObject, TInstancedStruct<FMusicClockSettingsBase> MusicClockSettings, bool Start = true);
 
 	template<typename MusicClockSettingsType = FMusicClockSettingsBase>
-	static UMusicClockComponent* CreateMusicClock(UObject* WorldContextObject, const MusicClockSettingsType& InMusicClockSettings, bool Start = true)
+	static UMusicClockComponent* CreateMusicClockWithSettings(UObject* WorldContextObject, const MusicClockSettingsType& InMusicClockSettings, bool Start = true)
 	{
 		return CreateMusicClockWithSettings(WorldContextObject, TInstancedStruct<FMusicClockSettingsBase>::Make(InMusicClockSettings), Start);
 	}
@@ -488,7 +488,16 @@ public:
 
 	UFUNCTION(BlueprintCallable, Category = "MusicClock")
 	UE_API void ConnectToWallClockForMidi(UMidiFile* InTempoMap);
+
+	UFUNCTION(BlueprintCallable, Category = "MusicClock")
+	UE_API void ConnectToCustomClockWithSettings(TInstancedStruct<FMusicClockSettingsBase> InMusicClockSettings);
 	
+	template<typename MusicClockSettingsType = FMusicClockSettingsBase>
+	void ConnectToCustomClockWithSettings(const MusicClockSettingsType& InMusicClockSettings)
+	{
+		ConnectToCustomClockWithSettings(TInstancedStruct<FMusicClockSettingsBase>::Make(InMusicClockSettings));
+	}
+
 	UE_API FMidiSongPos CalculateSongPosWithOffset(float MsOffset, ECalibratedMusicTimebase Timebase = ECalibratedMusicTimebase::VideoRenderTime) const;
 
 	UE_API const FMidiSongPos& GetRawUnsmoothedAudioRenderPos() const;
diff --git a/Engine/Plugins/Runtime/Harmonix/Source/HarmonixMetasound/Public/HarmonixMetasound/Components/MusicClockTypes.h b/Engine/Plugins/Runtime/Harmonix/Source/HarmonixMetasound/Public/HarmonixMetasound/Components/MusicClockTypes.h
--- a/Engine/Plugins/Runtime/Harmonix/Source/HarmonixMetasound/Public/HarmonixMetasound/Components/MusicClockTypes.h
+++ b/Engine/Plugins/Runtime/Harmonix/Source/HarmonixMetasound/Public/HarmonixMetasound/Components/MusicClockTypes.h
@@ -1,4 +1,4 @@
-﻿// Copyright Epic Games, Inc. All Rights Reserved.
+// Copyright Epic Games, Inc. All Rights Reserved.
 #pragma once
 
 #include "UObject/NameTypes.h"
@@ -9,7 +9,7 @@ class UAudioComponent;
 class UMidiFile;
 class FMusicClockDriverBase;
 
-USTRUCT()
+USTRUCT(MinimalAPI)
 struct FMusicClockSettingsBase
 {
 	GENERATED_BODY()
