Migrate MusicClockDriver members tied to specific timebases into arrays indexed by ECalibratedMusicTimebase

#tests PIE
#rnx

diff --git a/Engine/Plugins/Runtime/Harmonix/Source/Harmonix/Public/Harmonix/MusicalTimebase.h b/Engine/Plugins/Runtime/Harmonix/Source/Harmonix/Public/Harmonix/MusicalTimebase.h
--- a/Engine/Plugins/Runtime/Harmonix/Source/Harmonix/Public/Harmonix/MusicalTimebase.h
+++ b/Engine/Plugins/Runtime/Harmonix/Source/Harmonix/Public/Harmonix/MusicalTimebase.h
@@ -32,6 +32,9 @@ enum class ECalibratedMusicTimebase : uint8
 	 * Tells you exactly where the audio renderer is right now. NOTE: Un-smoothed, jittery,
 	 * raw position of the audio rendering. Only useful for debugging!
 	 */
-	RawAudioRenderTime
+	RawAudioRenderTime,
+
+	Count UMETA(Hidden)
 };
-ENUM_RANGE_BY_FIRST_AND_LAST(ECalibratedMusicTimebase, ECalibratedMusicTimebase::AudioRenderTime, ECalibratedMusicTimebase::RawAudioRenderTime);
+
+ENUM_RANGE_BY_COUNT(ECalibratedMusicTimebase, ECalibratedMusicTimebase::Count);
diff --git a/Engine/Plugins/Runtime/Harmonix/Source/HarmonixMetasound/Private/HarmonixMetasound/Components/MetasoundMusicClockDriver.cpp b/Engine/Plugins/Runtime/Harmonix/Source/HarmonixMetasound/Private/HarmonixMetasound/Components/MetasoundMusicClockDriver.cpp
--- a/Engine/Plugins/Runtime/Harmonix/Source/HarmonixMetasound/Private/HarmonixMetasound/Components/MetasoundMusicClockDriver.cpp
+++ b/Engine/Plugins/Runtime/Harmonix/Source/HarmonixMetasound/Private/HarmonixMetasound/Components/MetasoundMusicClockDriver.cpp
@@ -470,7 +470,7 @@ void FMetasoundMusicClockDriver::OnGeneratorDetached()
 		{
 			DefaultMaps.Copy(*(CurrentMapChain->SongMaps), 0, LastTickSeen);
 		}
-		SongPosOffsetMs = CurrentSmoothedAudioRenderSongPos.SecondsIncludingCountIn * 1000.0f;
+		SongPosOffsetMs = GetCurrentSmoothedAudioRenderSongPos().SecondsIncludingCountIn * 1000.0f;
 		FreeRunStartTimeSecs = GetWallClockTime();
 	}
 	SetClockHistory(nullptr);
@@ -482,12 +482,12 @@ void FMetasoundMusicClockDriver::OnGeneratorDetached()
 void FMetasoundMusicClockDriver::RefreshCurrentSongPosFromWallClock()
 {
 	double FreeRunTime = (GetWallClockTime() - FreeRunStartTimeSecs) * CurrentClockAdvanceRate;
+	SetCurrentSongPosByTime(ECalibratedMusicTimebase::RawAudioRenderTime, ((float)FreeRunTime * 1000.0) + SongPosOffsetMs, DefaultMaps);
+	SetCurrentSongPos(ECalibratedMusicTimebase::AudioRenderTime, GetCurrentRawAudioRenderSongPos());
+	SetCurrentSongPosByTime(ECalibratedMusicTimebase::ExperiencedTime, GetCurrentSmoothedAudioRenderSongPos().SecondsIncludingCountIn * 1000.0f - FHarmonixModule::GetMeasuredUserExperienceAndReactionToAudioRenderOffsetMs(), DefaultMaps);
+	SetCurrentSongPosByTime(ECalibratedMusicTimebase::VideoRenderTime, GetCurrentSmoothedAudioRenderSongPos().SecondsIncludingCountIn * 1000.0f - FHarmonixModule::GetMeasuredVideoToAudioRenderOffsetMs(), DefaultMaps);
 
-	CurrentRawAudioRenderSongPos.SetByTime(((float)FreeRunTime * 1000.0) + SongPosOffsetMs, DefaultMaps);
-	CurrentSmoothedAudioRenderSongPos = CurrentRawAudioRenderSongPos;
-	CurrentPlayerExperiencedSongPos.SetByTime(CurrentSmoothedAudioRenderSongPos.SecondsIncludingCountIn * 1000.0f - FHarmonixModule::GetMeasuredUserExperienceAndReactionToAudioRenderOffsetMs(), DefaultMaps);
-	CurrentVideoRenderSongPos.SetByTime(CurrentSmoothedAudioRenderSongPos.SecondsIncludingCountIn * 1000.0f - FHarmonixModule::GetMeasuredVideoToAudioRenderOffsetMs(), DefaultMaps);
-
+	const FMidiSongPos& CurrentSmoothedAudioRenderSongPos = GetCurrentSmoothedAudioRenderSongPos();
 	UpdateMusicPlaybackRate(CurrentSmoothedAudioRenderSongPos.Tempo, CurrentClockAdvanceRate, CurrentSmoothedAudioRenderSongPos.TimeSigNumerator, CurrentSmoothedAudioRenderSongPos.TimeSigDenominator);
 }
 
@@ -530,7 +530,7 @@ void FMetasoundMusicClockDriver::RefreshCurrentSongPosFromHistory()
 	}
 
 	auto Entry = ClockHistory->Positions.GetEntry(ClockHistory->Positions.GetLastWriteIndex());
-	CurrentRawAudioRenderSongPos.SetByTick(Entry->Item.UpToTick, *(CurrentMapChain->SongMaps));
+	SetCurrentSongPosByTick(ECalibratedMusicTimebase::RawAudioRenderTime, Entry->Item.UpToTick, *(CurrentMapChain->SongMaps));
 	Metasound::FSampleCount LastRenderPosSampleCount = Entry->Item.SampleCount;
 	float SpeedAtRawRenderTime = Entry->Item.CurrentSpeed;
 	LastTickSeen = Entry->Item.UpToTick;
@@ -585,6 +585,7 @@ void FMetasoundMusicClockDriver::RefreshCurrentSongPosFromHistory()
 
 	// first smoothed render time. This is closest to the actual render time. If we catch up to the render time
 	// it means we are rendering in such large blocks that we need to push up our "look behind" number for smoothing.
+	const FMidiSongPos& CurrentSmoothedAudioRenderSongPos = GetCurrentSmoothedAudioRenderSongPos();
 	float SpeedAtSmoothRenderedTime = 1.0f;
 	float SmoothedTick = GetCurrentSongMapEvaluator()->MsToTick(CurrentSmoothedAudioRenderSongPos.SecondsIncludingCountIn * 1000.0f);
 	float SmoothedTempoMapTick = AudioRenderState.TempoMapTick;
@@ -645,10 +646,10 @@ void FMetasoundMusicClockDriver::UpdateCurrentTicksForOffsetClock(float Smoothed
 	SmoothedPositionMs += RenderSmoothingLagSeconds * 1000.0f;
 
 	// Calculate the song position AND get the "local tick" for the Audio Render timebase
-	CurrentSmoothedAudioRenderSongPos = CalculateSongPosAtMsForOffsetClock(SmoothedPositionMs,
+	SetCurrentSongPos(ECalibratedMusicTimebase::AudioRenderTime, CalculateSongPosAtMsForOffsetClock(SmoothedPositionMs,
 		SmoothedTick - SmoothedTempoMapTick,
 		AudioRenderState.LocalTick,
-		AudioRenderSeekDetected);
+		AudioRenderSeekDetected));
 	float SmoothedTempoMapMs = CurrentMapChain->SongMaps->TickToMs(SmoothedTempoMapTick);
 	AudioRenderState.TempoMapMs = SmoothedTempoMapMs + (RenderSmoothingLagSeconds * 1000.0f);
 	AudioRenderState.TempoMapTick = CurrentMapChain->SongMaps->MsToTick(AudioRenderState.TempoMapMs);
@@ -659,10 +660,10 @@ void FMetasoundMusicClockDriver::UpdateCurrentTicksForOffsetClock(float Smoothed
 	float Ms = SmoothedPositionMs;
 	Ms -= LookBehindLagMs;
 	// Calculate the song position AND get the "local tick" for the Player Experience timebase
-	CurrentPlayerExperiencedSongPos = CalculateSongPosAtMsForOffsetClock(Ms,
+	SetCurrentSongPos(ECalibratedMusicTimebase::ExperiencedTime, CalculateSongPosAtMsForOffsetClock(Ms,
 		SmoothedTick - SmoothedTempoMapTick,
 		PlayerExperienceState.LocalTick,
-		PlayerExperiencedSeekDetected);
+		PlayerExperiencedSeekDetected));
 	PlayerExperienceState.TempoMapMs = AudioRenderState.TempoMapMs - LookBehindLagMs;
 	PlayerExperienceState.TempoMapTick = CurrentMapChain->SongMaps->MsToTick(PlayerExperienceState.TempoMapMs);
 
@@ -672,10 +673,10 @@ void FMetasoundMusicClockDriver::UpdateCurrentTicksForOffsetClock(float Smoothed
 	Ms = SmoothedPositionMs;
 	Ms -= LookBehindLagMs;
 	// Calculate the song position AND get the "local tick" for the Video Render timebase
-	CurrentVideoRenderSongPos = CalculateSongPosAtMsForOffsetClock(Ms,
+	SetCurrentSongPos(ECalibratedMusicTimebase::VideoRenderTime, CalculateSongPosAtMsForOffsetClock(Ms,
 		SmoothedTick - SmoothedTempoMapTick,
 		VideoRenderState.LocalTick,
-		VideoRenderSeekDetected);
+		VideoRenderSeekDetected));
 	VideoRenderState.TempoMapMs = AudioRenderState.TempoMapMs - LookBehindLagMs;
 	VideoRenderState.TempoMapTick = CurrentMapChain->SongMaps->MsToTick(VideoRenderState.TempoMapMs);
 }
@@ -689,30 +690,30 @@ void FMetasoundMusicClockDriver::UpdateCurrentTicksForLoopingOrMonotonicClock(fl
 
 	// Calculate the song position AND get the "local tick" for the Audio Render timebase
 	AudioRenderState.TempoMapTick = CurrentMapChain->SongMaps->MsToTick(AudioRenderState.TempoMapMs);
-	CurrentSmoothedAudioRenderSongPos = CalculateSongPosAtMsForLoopingOrMonotonicClock(AudioRenderState.TempoMapMs,
+	SetCurrentSongPos(ECalibratedMusicTimebase::AudioRenderTime, CalculateSongPosAtMsForLoopingOrMonotonicClock(AudioRenderState.TempoMapMs,
 		AudioRenderState.LocalTick,
 		AudioRenderSeekDetected,
-		AudioRenderLoopDetected);
+		AudioRenderLoopDetected));
 
 	// Now the time the user should actually be "experiencing" (ie "hearing") can be calculated as an offset
 	// from the smooth audio rendering time...
 	float LookBehindLagMs = FHarmonixModule::Get().GetMeasuredUserExperienceAndReactionToAudioRenderOffsetMs();
 	PlayerExperienceState.TempoMapMs = AudioRenderState.TempoMapMs - LookBehindLagMs;
 	PlayerExperienceState.TempoMapTick = CurrentMapChain->SongMaps->MsToTick(PlayerExperienceState.TempoMapMs);
-	CurrentPlayerExperiencedSongPos = CalculateSongPosAtMsForLoopingOrMonotonicClock(PlayerExperienceState.TempoMapMs,
+	SetCurrentSongPos(ECalibratedMusicTimebase::ExperiencedTime, CalculateSongPosAtMsForLoopingOrMonotonicClock(PlayerExperienceState.TempoMapMs,
 		PlayerExperienceState.LocalTick,
 		PlayerExperiencedSeekDetected,
-		PlayerExperiencedLoopDetected);
+		PlayerExperiencedLoopDetected));
 
 	// Now the time the game should be rendering graphics for can be calculated as an offset
 	// from the smooth audio rendering time...
 	LookBehindLagMs = FHarmonixModule::Get().GetMeasuredVideoToAudioRenderOffsetMs();
 	VideoRenderState.TempoMapMs = AudioRenderState.TempoMapMs - LookBehindLagMs;
 	VideoRenderState.TempoMapTick = CurrentMapChain->SongMaps->MsToTick(VideoRenderState.TempoMapMs);
-	CurrentVideoRenderSongPos = CalculateSongPosAtMsForLoopingOrMonotonicClock(VideoRenderState.TempoMapMs,
+	SetCurrentSongPos(ECalibratedMusicTimebase::VideoRenderTime, CalculateSongPosAtMsForLoopingOrMonotonicClock(VideoRenderState.TempoMapMs,
 		VideoRenderState.LocalTick,
 		VideoRenderSeekDetected,
-		VideoRenderLoopDetected);
+		VideoRenderLoopDetected));
 }
 
 FMetasoundMusicClockDriver::EHistoryFailureType FMetasoundMusicClockDriver::CalculateSmoothedTick(Metasound::FSampleCount ExpectedRenderPosSampleCount, Metasound::FSampleCount LastRenderPosSampleCount,
diff --git a/Engine/Plugins/Runtime/Harmonix/Source/HarmonixMetasound/Private/HarmonixMetasound/Components/MusicClockComponent.cpp b/Engine/Plugins/Runtime/Harmonix/Source/HarmonixMetasound/Private/HarmonixMetasound/Components/MusicClockComponent.cpp
--- a/Engine/Plugins/Runtime/Harmonix/Source/HarmonixMetasound/Private/HarmonixMetasound/Components/MusicClockComponent.cpp
+++ b/Engine/Plugins/Runtime/Harmonix/Source/HarmonixMetasound/Private/HarmonixMetasound/Components/MusicClockComponent.cpp
@@ -250,7 +250,7 @@ FMidiSongPos UMusicClockComponent::CalculateSongPosWithOffset(float MsOffset, EC
 const FMidiSongPos& UMusicClockComponent::GetRawUnsmoothedAudioRenderPos() const
 {
 	static const FMidiSongPos PosZero;
-	return ClockDriver ? ClockDriver->CurrentRawAudioRenderSongPos : PosZero;
+	return ClockDriver ? ClockDriver->GetCurrentRawAudioRenderSongPos() : PosZero;
 }
 
 void UMusicClockComponent::Activate(bool bReset)
@@ -595,34 +595,12 @@ float UMusicClockComponent::GetDistanceToClosestBar(ECalibratedMusicTimebase Tim
 
 float UMusicClockComponent::GetDeltaBar(ECalibratedMusicTimebase Timebase) const
 {
-	if (ClockDriver)
-	{
-		switch (Timebase)
-		{
-		case ECalibratedMusicTimebase::RawAudioRenderTime: return ClockDriver->RawAudioRenderDeltaBarF;
-		case ECalibratedMusicTimebase::AudioRenderTime:    return ClockDriver->AudioRenderDeltaBarF;
-		case ECalibratedMusicTimebase::ExperiencedTime:	   return ClockDriver->PlayerExperienceDeltaBarF;
-		case ECalibratedMusicTimebase::VideoRenderTime:
-		default:                                           return ClockDriver->VideoRenderDeltaBarF;
-		}
-	}
-	return 0.0f;
+	return ClockDriver ? ClockDriver->GetDeltaBarF(Timebase) : 0.0f;
 }
 
 float UMusicClockComponent::GetDeltaBeat(ECalibratedMusicTimebase Timebase) const
 {
-	if (ClockDriver)
-	{
-		switch (Timebase)
-		{
-		case ECalibratedMusicTimebase::RawAudioRenderTime: return ClockDriver->RawAudioRenderDeltaBeatF;
-		case ECalibratedMusicTimebase::AudioRenderTime:    return ClockDriver->AudioRenderDeltaBeatF;
-		case ECalibratedMusicTimebase::ExperiencedTime:	   return ClockDriver->PlayerExperienceDeltaBeatF;
-		case ECalibratedMusicTimebase::VideoRenderTime:
-		default:                                           return ClockDriver->VideoRenderDeltaBeatF;
-		}
-	}
-	return 0.0f;
+	return ClockDriver ? ClockDriver->GetDeltaBeatF(Timebase) : 0.0f;
 }
 
 const TArray<FSongSection>& UMusicClockComponent::GetSongSections() const
diff --git a/Engine/Plugins/Runtime/Harmonix/Source/HarmonixMetasound/Private/HarmonixMetasound/Components/MusicClockDriverBase.cpp b/Engine/Plugins/Runtime/Harmonix/Source/HarmonixMetasound/Private/HarmonixMetasound/Components/MusicClockDriverBase.cpp
--- a/Engine/Plugins/Runtime/Harmonix/Source/HarmonixMetasound/Private/HarmonixMetasound/Components/MusicClockDriverBase.cpp
+++ b/Engine/Plugins/Runtime/Harmonix/Source/HarmonixMetasound/Private/HarmonixMetasound/Components/MusicClockDriverBase.cpp
@@ -1,4 +1,4 @@
-﻿// Copyright Epic Games, Inc. All Rights Reserved.
+// Copyright Epic Games, Inc. All Rights Reserved.
 
 #include "HarmonixMetasound/Components/MusicClockDriverBase.h"
 #include "Engine/World.h"
@@ -29,6 +29,24 @@ FMusicClockDriverBase::~FMusicClockDriverBase()
 {
 }
 
+bool FMusicClockDriverBase::CalculateSongPosWithOffset(float MsOffset, ECalibratedMusicTimebase Timebase, FMidiSongPos& OutResult) const
+{
+	const ISongMapEvaluator* SongMapEvaluator = GetCurrentSongMapEvaluator();
+	if (!SongMapEvaluator)
+	{
+		return false;
+	}
+
+	const FMidiSongPos& SongPos = GetCurrentSongPos(Timebase);
+	if (!SongPos.IsValid())
+	{
+		return false;
+	}
+
+	OutResult.SetByTick((SongPos.SecondsIncludingCountIn * 1000.0f) + MsOffset, *SongMapEvaluator);
+	return true;
+}
+
 double FMusicClockDriverBase::GetWallClockTime() const
 {
 	if (TStrongObjectPtr<UObject> Obj = ContextObj.Pin())
@@ -50,50 +68,33 @@ double FMusicClockDriverBase::GetWallClockTime() const
 
 const FMidiSongPos& FMusicClockDriverBase::GetCurrentSongPos(ECalibratedMusicTimebase Timebase) const
 {
-	switch (Timebase)
-	{
-	case ECalibratedMusicTimebase::AudioRenderTime:
-		return CurrentSmoothedAudioRenderSongPos;
-	case ECalibratedMusicTimebase::ExperiencedTime:
-		return CurrentPlayerExperiencedSongPos;
-	case ECalibratedMusicTimebase::VideoRenderTime:
-		return CurrentVideoRenderSongPos;
-	case ECalibratedMusicTimebase::RawAudioRenderTime:
-		return CurrentRawAudioRenderSongPos;
-	}
-	checkNoEntry();
-	// minor robustness, return something even if there's an illegal value. 
-	return CurrentVideoRenderSongPos;
+	return CurrentSongPos[GetArrayIndexForTimebase(Timebase)];
 }
 
 const FMidiSongPos& FMusicClockDriverBase::GetPreviousSongPos(ECalibratedMusicTimebase Timebase) const
 {
-	switch (Timebase)
-	{
-	case ECalibratedMusicTimebase::AudioRenderTime:
-		return PrevAudioRenderSongPos;
-	case ECalibratedMusicTimebase::ExperiencedTime:
-		return PrevPlayerExperiencedSongPos;
-	case ECalibratedMusicTimebase::VideoRenderTime:
-		return PrevVideoRenderSongPos;
-	case ECalibratedMusicTimebase::RawAudioRenderTime:
-		return PrevRawAudioRenderSongPos;
-	}
-	checkNoEntry();
-	
-	// minor robustness, return something even if there's an illegal value. 
-	return CurrentVideoRenderSongPos;
+	return PrevSongPos[GetArrayIndexForTimebase(Timebase)];
 }
 
+float FMusicClockDriverBase::GetDeltaBarF(ECalibratedMusicTimebase Timebase) const
+{
+	return DeltaBarF[GetArrayIndexForTimebase(Timebase)];
+}
 
+float FMusicClockDriverBase::GetDeltaBeatF(ECalibratedMusicTimebase Timebase) const
+{
+	return DeltaBeatF[GetArrayIndexForTimebase(Timebase)];
+}
 
 void FMusicClockDriverBase::Start()
 {
 	OnStart();
-	CurrentSmoothedAudioRenderSongPos.Reset();
-    CurrentVideoRenderSongPos.Reset();
-    CurrentPlayerExperiencedSongPos.Reset();
-    CurrentRawAudioRenderSongPos.Reset();
+
+	for (ECalibratedMusicTimebase Timebase : TEnumRange<ECalibratedMusicTimebase>())
+	{
+		CurrentSongPos[(int32)Timebase].Reset();
+	}
+
 	MusicClockState = EMusicClockState::Running;
 }
 
@@ -123,10 +124,11 @@ void FMusicClockDriverBase::Stop()
 {
 	OnStop();
 	MusicClockState = EMusicClockState::Stopped;
-	CurrentSmoothedAudioRenderSongPos.Reset();
-	CurrentVideoRenderSongPos.Reset();
-	CurrentPlayerExperiencedSongPos.Reset();
-	CurrentRawAudioRenderSongPos.Reset();
+
+	for (ECalibratedMusicTimebase Timebase : TEnumRange<ECalibratedMusicTimebase>())
+	{
+		CurrentSongPos[(int32)Timebase].Reset();
+	}
 }
 
 void FMusicClockDriverBase::UpdateMusicPlaybackRate(float InTempo, float InClockAdvanceRate, int32 InTimeSigNum, int32 InTimeSigDenom)
@@ -184,22 +186,47 @@ void FMusicClockDriverBase::EnsureClockIsValidForGameFrame()
 	//	Run the actual clock update.
 	if (GetState() == EMusicClockState::Running)
 	{
-		PrevRawAudioRenderSongPos = CurrentRawAudioRenderSongPos;
-		PrevAudioRenderSongPos = CurrentSmoothedAudioRenderSongPos;
-		PrevPlayerExperiencedSongPos = CurrentPlayerExperiencedSongPos;
-		PrevVideoRenderSongPos = CurrentVideoRenderSongPos;
+		for (ECalibratedMusicTimebase Timebase : TEnumRange<ECalibratedMusicTimebase>())
+		{
+			PrevSongPos[(int32)Timebase] = CurrentSongPos[(int32)Timebase];
+		}
 
 		if (RefreshCurrentSongPos())
 		{
-			RawAudioRenderDeltaBarF = CurrentRawAudioRenderSongPos.BarsIncludingCountIn - PrevRawAudioRenderSongPos.BarsIncludingCountIn;
-			RawAudioRenderDeltaBeatF = CurrentRawAudioRenderSongPos.BeatsIncludingCountIn - PrevRawAudioRenderSongPos.BeatsIncludingCountIn;
-			AudioRenderDeltaBarF = CurrentSmoothedAudioRenderSongPos.BarsIncludingCountIn - PrevAudioRenderSongPos.BarsIncludingCountIn;
-			AudioRenderDeltaBeatF = CurrentSmoothedAudioRenderSongPos.BeatsIncludingCountIn - PrevAudioRenderSongPos.BeatsIncludingCountIn;
-			PlayerExperienceDeltaBarF = CurrentPlayerExperiencedSongPos.BarsIncludingCountIn - PrevPlayerExperiencedSongPos.BarsIncludingCountIn;
-			PlayerExperienceDeltaBeatF = CurrentPlayerExperiencedSongPos.BeatsIncludingCountIn - PrevPlayerExperiencedSongPos.BeatsIncludingCountIn;
-			VideoRenderDeltaBarF = CurrentVideoRenderSongPos.BarsIncludingCountIn - PrevVideoRenderSongPos.BarsIncludingCountIn;
-			VideoRenderDeltaBeatF = CurrentVideoRenderSongPos.BeatsIncludingCountIn - PrevVideoRenderSongPos.BeatsIncludingCountIn;
+			for (ECalibratedMusicTimebase Timebase : TEnumRange<ECalibratedMusicTimebase>())
+			{
+				DeltaBarF[(int32)Timebase] = CurrentSongPos[(int32)Timebase].BarsIncludingCountIn - PrevSongPos[(int32)Timebase].BarsIncludingCountIn;
+				DeltaBeatF[(int32)Timebase] = CurrentSongPos[(int32)Timebase].BeatsIncludingCountIn - PrevSongPos[(int32)Timebase].BeatsIncludingCountIn;
+			}
+
 			LastUpdateFrame = GFrameCounter;
 		}
 	}
-}
\ No newline at end of file
+}
+
+int32 FMusicClockDriverBase::GetArrayIndexForTimebase(ECalibratedMusicTimebase Timebase) const
+{
+	int32 TimebaseIndex = (int32)Timebase;
+	if (!ensure(TimebaseIndex >= 0 && TimebaseIndex < (int32)ECalibratedMusicTimebase::Count))
+	{
+		// minor robustness, return something even if there's an illegal value.
+		return (int32)ECalibratedMusicTimebase::VideoRenderTime;
+	}
+
+	return TimebaseIndex;
+}
+
+void FMusicClockDriverBase::SetCurrentSongPos(ECalibratedMusicTimebase Timebase, const FMidiSongPos& InSongPos)
+{
+	CurrentSongPos[GetArrayIndexForTimebase(Timebase)] = InSongPos;
+}
+
+void FMusicClockDriverBase::SetCurrentSongPosByTime(ECalibratedMusicTimebase Timebase, float InElapsedMs, const ISongMapEvaluator& Map)
+{
+	CurrentSongPos[GetArrayIndexForTimebase(Timebase)].SetByTime(InElapsedMs, Map);
+}
+
+void FMusicClockDriverBase::SetCurrentSongPosByTick(ECalibratedMusicTimebase Timebase, float InTick, const ISongMapEvaluator& Map)
+{
+	CurrentSongPos[GetArrayIndexForTimebase(Timebase)].SetByTick(InTick, Map);
+}
diff --git a/Engine/Plugins/Runtime/Harmonix/Source/HarmonixMetasound/Private/HarmonixMetasound/Components/WallClockMusicClockDriver.cpp b/Engine/Plugins/Runtime/Harmonix/Source/HarmonixMetasound/Private/HarmonixMetasound/Components/WallClockMusicClockDriver.cpp
--- a/Engine/Plugins/Runtime/Harmonix/Source/HarmonixMetasound/Private/HarmonixMetasound/Components/WallClockMusicClockDriver.cpp
+++ b/Engine/Plugins/Runtime/Harmonix/Source/HarmonixMetasound/Private/HarmonixMetasound/Components/WallClockMusicClockDriver.cpp
@@ -2,32 +2,6 @@
 #include "WallClockMusicClockDriver.h"
 #include "Harmonix.h"
 
-bool FWallClockMusicClockDriver::CalculateSongPosWithOffset(float MsOffset, ECalibratedMusicTimebase Timebase, FMidiSongPos& OutResult) const
-{
-	check(IsInGameThread());
-	if (!TempoMapMidi.IsValid())
-	{
-		return false;
-	}
-
-	const FSongMaps* Maps = TempoMapMidi->GetSongMaps();
-	switch (Timebase)
-	{
-	case ECalibratedMusicTimebase::AudioRenderTime:
-		OutResult.SetByTime((CurrentSmoothedAudioRenderSongPos.SecondsIncludingCountIn * 1000.0f) + MsOffset, *Maps);
-		break;
-	case ECalibratedMusicTimebase::ExperiencedTime:
-		OutResult.SetByTime((CurrentPlayerExperiencedSongPos.SecondsIncludingCountIn * 1000.0f) + MsOffset, *Maps);
-		break;
-	case ECalibratedMusicTimebase::VideoRenderTime:
-	default:
-		OutResult.SetByTime((CurrentVideoRenderSongPos.SecondsIncludingCountIn * 1000.0f) + MsOffset, *Maps);
-		break;
-	}
-
-	return true;
-}
-
 void FWallClockMusicClockDriver::Disconnect()
 {
 	TempoMapMidi = nullptr;
@@ -42,12 +16,14 @@ bool FWallClockMusicClockDriver::RefreshCurrentSongPos()
 	check(Maps);
 
 	// Set Experienced time here, as this code is running on the GameThread, near the same time as input is processed
-	CurrentPlayerExperiencedSongPos.SetByTime((float)(RunTime * 1000.0), *Maps);
-	CurrentRawAudioRenderSongPos.SetByTime(CurrentPlayerExperiencedSongPos.SecondsIncludingCountIn * 1000.0f + FHarmonixModule::GetMeasuredUserExperienceAndReactionToAudioRenderOffsetMs(), *Maps);
-	CurrentVideoRenderSongPos.SetByTime(CurrentRawAudioRenderSongPos.SecondsIncludingCountIn * 1000.0f - FHarmonixModule::GetMeasuredVideoToAudioRenderOffsetMs(), *Maps);
-	CurrentSmoothedAudioRenderSongPos = CurrentRawAudioRenderSongPos;
+	SetCurrentSongPosByTime(ECalibratedMusicTimebase::ExperiencedTime, (float)(RunTime * 1000.0), *Maps);
+	SetCurrentSongPosByTime(ECalibratedMusicTimebase::RawAudioRenderTime, GetCurrentPlayerExperiencedSongPos().SecondsIncludingCountIn * 1000.0f + FHarmonixModule::GetMeasuredUserExperienceAndReactionToAudioRenderOffsetMs(), *Maps);
+	SetCurrentSongPosByTime(ECalibratedMusicTimebase::VideoRenderTime, GetCurrentRawAudioRenderSongPos().SecondsIncludingCountIn * 1000.0f - FHarmonixModule::GetMeasuredVideoToAudioRenderOffsetMs(), *Maps);
+	SetCurrentSongPos(ECalibratedMusicTimebase::AudioRenderTime, GetCurrentRawAudioRenderSongPos());
 
+	const FMidiSongPos& CurrentSmoothedAudioRenderSongPos = GetCurrentSmoothedAudioRenderSongPos();
 	UpdateMusicPlaybackRate(CurrentSmoothedAudioRenderSongPos.Tempo, CurrentClockAdvanceRate, CurrentSmoothedAudioRenderSongPos.TimeSigNumerator, CurrentSmoothedAudioRenderSongPos.TimeSigDenominator);
+
 	return true;
 }
 
diff --git a/Engine/Plugins/Runtime/Harmonix/Source/HarmonixMetasound/Private/HarmonixMetasound/Components/WallClockMusicClockDriver.h b/Engine/Plugins/Runtime/Harmonix/Source/HarmonixMetasound/Private/HarmonixMetasound/Components/WallClockMusicClockDriver.h
--- a/Engine/Plugins/Runtime/Harmonix/Source/HarmonixMetasound/Private/HarmonixMetasound/Components/WallClockMusicClockDriver.h
+++ b/Engine/Plugins/Runtime/Harmonix/Source/HarmonixMetasound/Private/HarmonixMetasound/Components/WallClockMusicClockDriver.h
@@ -12,8 +12,6 @@ public:
 		, TempoMapMidi(Settings.TempoMap)
 	{}
 
-	virtual bool CalculateSongPosWithOffset(float MsOffset, ECalibratedMusicTimebase Timebase, FMidiSongPos& OutResult) const override;
-
 	virtual void Disconnect() override;
 	virtual bool RefreshCurrentSongPos() override;
 	virtual void OnStart() override;
diff --git a/Engine/Plugins/Runtime/Harmonix/Source/HarmonixMetasound/Public/HarmonixMetasound/Components/MusicClockDriverBase.h b/Engine/Plugins/Runtime/Harmonix/Source/HarmonixMetasound/Public/HarmonixMetasound/Components/MusicClockDriverBase.h
--- a/Engine/Plugins/Runtime/Harmonix/Source/HarmonixMetasound/Public/HarmonixMetasound/Components/MusicClockDriverBase.h
+++ b/Engine/Plugins/Runtime/Harmonix/Source/HarmonixMetasound/Public/HarmonixMetasound/Components/MusicClockDriverBase.h
@@ -20,13 +20,20 @@ public:
 
 	UE_API void EnsureClockIsValidForGameFrame();
 
-	virtual bool CalculateSongPosWithOffset(float MsOffset, ECalibratedMusicTimebase Timebase, FMidiSongPos& OutResult) const = 0;
+	UE_API virtual bool CalculateSongPosWithOffset(float MsOffset, ECalibratedMusicTimebase Timebase, FMidiSongPos& OutResult) const;
 
 	UE_API double GetWallClockTime() const;
 
 	EMusicClockState GetState() const { return MusicClockState; }
 	UE_API const FMidiSongPos& GetCurrentSongPos(ECalibratedMusicTimebase Timebase) const;
 	UE_API const FMidiSongPos& GetPreviousSongPos(ECalibratedMusicTimebase Timebase) const;
+	UE_API float GetDeltaBarF(ECalibratedMusicTimebase Timebase) const;
+	UE_API float GetDeltaBeatF(ECalibratedMusicTimebase Timebase) const;
+
+	const FMidiSongPos& GetCurrentSmoothedAudioRenderSongPos() const { return CurrentSongPos[(int32)ECalibratedMusicTimebase::AudioRenderTime]; }
+	const FMidiSongPos& GetCurrentPlayerExperiencedSongPos() const { return CurrentSongPos[(int32)ECalibratedMusicTimebase::ExperiencedTime]; }
+	const FMidiSongPos& GetCurrentVideoRenderSongPos() const { return CurrentSongPos[(int32)ECalibratedMusicTimebase::VideoRenderTime]; }
+	const FMidiSongPos& GetCurrentRawAudioRenderSongPos() const { return CurrentSongPos[(int32)ECalibratedMusicTimebase::RawAudioRenderTime]; }
 
 	UE_API void Start();
 	UE_API void Pause();
@@ -54,28 +61,14 @@ public:
 
 	float CurrentBeatsPerSecond = 0.0f;
 	float CurrentBarsPerSecond = 0.0f;
-
-	FMidiSongPos CurrentSmoothedAudioRenderSongPos;
-	FMidiSongPos CurrentVideoRenderSongPos;
-	FMidiSongPos CurrentPlayerExperiencedSongPos;
-	FMidiSongPos CurrentRawAudioRenderSongPos;
-	
-	FMidiSongPos PrevRawAudioRenderSongPos;
-	FMidiSongPos PrevAudioRenderSongPos;
-	FMidiSongPos PrevPlayerExperiencedSongPos;
-	FMidiSongPos PrevVideoRenderSongPos;
-	
 	float CurrentClockAdvanceRate = 1.0f;
 
+	FMidiSongPos CurrentSongPos[(int32)ECalibratedMusicTimebase::Count];
+	FMidiSongPos PrevSongPos[(int32)ECalibratedMusicTimebase::Count];
+
 	// TODO shouldn't these deltas be rolled into the MidiSongPos?
-	float RawAudioRenderDeltaBarF = 0.0f;
-	float RawAudioRenderDeltaBeatF = 0.0f;
-	float AudioRenderDeltaBarF = 0.0f;
-	float AudioRenderDeltaBeatF = 0.0f;
-	float PlayerExperienceDeltaBarF = 0.0f;
-	float PlayerExperienceDeltaBeatF = 0.0f;
-	float VideoRenderDeltaBarF = 0.0f;
-	float VideoRenderDeltaBeatF = 0.0f;
+	float DeltaBarF[(int32)ECalibratedMusicTimebase::Count];
+	float DeltaBeatF[(int32)ECalibratedMusicTimebase::Count];
 	
 protected:
 	
@@ -84,6 +77,12 @@ protected:
 	virtual void OnContinue() = 0;
 	virtual void OnStop() = 0;
 	
+	UE_API int32 GetArrayIndexForTimebase(ECalibratedMusicTimebase Timebase) const;
+
+	UE_API void SetCurrentSongPos(ECalibratedMusicTimebase Timebase, const FMidiSongPos& InSongPos);
+	UE_API void SetCurrentSongPosByTime(ECalibratedMusicTimebase Timebase, float InElapsedMs, const ISongMapEvaluator& Map);
+	UE_API void SetCurrentSongPosByTick(ECalibratedMusicTimebase Timebase, float InTick, const ISongMapEvaluator& Map);
+
 private:
 
 	UE_API void ResetDefaultSongMaps(float DefaultTempo, int32 DefaultTimeSigNum, int32 DefaultTimeSigDenom);
