Fix Simple Sampler jitter

#rnx
#tests Harmonix Internal Tests

diff --git a/Engine/Plugins/Runtime/Harmonix/Source/HarmonixMetasound/Private/HarmonixMetasound/Nodes/SimpleSamplerNode.cpp b/Engine/Plugins/Runtime/Harmonix/Source/HarmonixMetasound/Private/HarmonixMetasound/Nodes/SimpleSamplerNode.cpp
--- a/Engine/Plugins/Runtime/Harmonix/Source/HarmonixMetasound/Private/HarmonixMetasound/Nodes/SimpleSamplerNode.cpp
+++ b/Engine/Plugins/Runtime/Harmonix/Source/HarmonixMetasound/Private/HarmonixMetasound/Nodes/SimpleSamplerNode.cpp
@@ -338,7 +338,8 @@ namespace HarmonixMetasound::SimpleSamplerNode
 					// Are we still playing a previous trigger?
 					if (bCurrentlyPlaying)
 					{
-						NewPlayPosition = ForceFinishCurrentPlaybackByBlockFrame(NewPlayPosition);
+						NewPlayPosition = FMath::Max(Inputs.Play->Last(), kNumFadeFrames);
+						DoPlay(0, NewPlayPosition, true);
 						bRenderedOutput = true;
 					}
 					else
@@ -391,9 +392,6 @@ namespace HarmonixMetasound::SimpleSamplerNode
 				}
 			}
 
-			// Reset this since we are done capturing & playing this block...
-			bPlayFadeInNotNeededThisBlock = false;
-
 			// If we didn't render anything, and we are fully captured
 			if (!bRenderedOutput && (bFullyCaptured || bHavePlayedThisCapture))
 			{
@@ -438,42 +436,6 @@ namespace HarmonixMetasound::SimpleSamplerNode
 			}
 		}
 
-		/**
-		Ends the current playback as gracefully as possible.
-		@param CompleteByBlockFrame When the playback needs to be finished.
-		@return The frame we actually completed by, so that the new play
-		        request can be moved a few samples to avoid a click.
-		*/
-		int32 ForceFinishCurrentPlaybackByBlockFrame(int32 CompleteByBlockFrame)
-		{
-			int32 FramesRemainingToBePlayed = PlayLengthFrames - NextPlayFrame;
-			if (FramesRemainingToBePlayed <= CompleteByBlockFrame)
-			{
-				// We have enough space to finish rendering the last play INCLUDING its fade...
-				DoPlay(0, CompleteByBlockFrame, false);
-				// No need to modify the new play frame.
-			}
-			else
-			{
-				// We don't have enough space to finish rendering the last play...
-				// If it has already started fading, let it finish and then fudge the start sample
-				// of the next play...
-				if (FramesRemainingToBePlayed <= kNumFadeFrames)
-				{
-					DoPlay(0, CompleteByBlockFrame, false);
-					CompleteByBlockFrame = FramesRemainingToBePlayed;
-				}
-				else
-				{
-					// It hasn't already started to fade, so let's force it to fade (truncate it) and
-					// fudge the start sample of the next play...
-					DoPlay(0, kNumFadeFrames, true);
-					CompleteByBlockFrame = kNumFadeFrames;
-				}
-			}
-			return CompleteByBlockFrame;
-		}
-
 		/**
 		Plays the captured audio. Will do a little fade in at the beginning, and a little 
 		fade out at the end.
@@ -514,7 +476,7 @@ namespace HarmonixMetasound::SimpleSamplerNode
 					}
 				}
 				// fade in the beginning to avoid click...
-				if (NextPlayFrame < kNumFadeFrames && !bPlayFadeInNotNeededThisBlock)
+				if (NextPlayFrame < kNumFadeFrames)
 				{
 					int32 FadeFrames = FMath::Min(kNumFadeFrames, FramesToPlay);
 					float CurrentFade = float(NextPlayFrame) / float(kNumFadeFrames);
@@ -621,8 +583,23 @@ namespace HarmonixMetasound::SimpleSamplerNode
 				if (FirstFrameOfPlay > -1)
 				{
 					// We WILL be playing this block... so limit "capture output thru" up to that point...
-					BulkOutput(EInputOrSilence::Input, 0, FMath::Min(FirstFrameOfPlay, FirstFrameInBlockToCapture + NumFramesToCapture));
-					bPlayFadeInNotNeededThisBlock = FirstFrameOfPlay == FirstFrameInBlockToCapture && CaptureBuffer.GetNumValidFrames() == NumFramesToCapture;
+					FirstFrameOfPlay = FMath::Max(FirstFrameOfPlay, kNumFadeFrames);
+					int32 FrameToFinish = FMath::Min(FirstFrameOfPlay, FirstFrameInBlockToCapture + NumFramesToCapture);
+					BulkOutput(EInputOrSilence::Input, 0, FrameToFinish);
+
+					// fade out the input before the start of the next play
+					float CurrentFade = 0.0f;
+					float FadeInc = 1.0f/float(kNumFadeFrames);
+					int32 FadeAtIndex = FrameToFinish - 1;
+					for (int32 FadeIndex = 0; FadeIndex < kNumFadeFrames; ++FadeIndex)
+					{
+						for (int32 ChannelIndex = 0; ChannelIndex < CaptureBuffer.GetNumValidChannels(); ++ChannelIndex)
+						{
+							Outputs.AudioOutputs[ChannelIndex]->GetData()[FadeAtIndex] *= CurrentFade;
+						}
+						--FadeAtIndex;
+						CurrentFade = FMath::Clamp(CurrentFade + FadeInc, 0.0f, 1.0f);
+					}
 				}
 				else
 				{
@@ -669,7 +646,6 @@ namespace HarmonixMetasound::SimpleSamplerNode
 		bool     bFullyCaptured                = false;
 		bool     bCurrentlyPlaying             = false;
 		int32    CaptureLengthFrames           = 0;
-		bool     bPlayFadeInNotNeededThisBlock = false;
 		int32    PlayLengthFrames              = 0;
 		bool     bPlayingReverse               = false;
 		int32    NextPlayFrame                 = 0;
