Parse lyrics from MIDI

Added IsLyric() function in FMidiMsg struct. The function checks whether the midi message is a lyric message.

diff --git a/Engine/Plugins/Runtime/Harmonix/Source/HarmonixMidi/Public/HarmonixMidi/MidiMsg.h b/Engine/Plugins/Runtime/Harmonix/Source/HarmonixMidi/Public/HarmonixMidi/MidiMsg.h
--- a/Engine/Plugins/Runtime/Harmonix/Source/HarmonixMidi/Public/HarmonixMidi/MidiMsg.h
+++ b/Engine/Plugins/Runtime/Harmonix/Source/HarmonixMidi/Public/HarmonixMidi/MidiMsg.h
@@ -53,6 +53,7 @@ public:
 	bool  IsControlChange() const { return IsStd() && Harmonix::Midi::Constants::IsControl(Status); }
 	bool  IsTempo() const { return Type == EType::Tempo; }
 	bool  IsTimeSignature() const { return Type == EType::TimeSig; }
+	bool  IsLyric() const { return Type == EType::Text && TextType == Harmonix::Midi::Constants::GMeta_Lyric; }
 
 	/** Construct a standard (std) midi message */
 	UE_API FMidiMsg(uint8 InStatus, uint8 InData1, uint8 InData2);
