[FM] Replace SongBeatAndTime with a pointer to a music clock, to modernize and simplify the BeatSyncedAnimNode.

#tests PIE mainstage, frontend jam
#changelist validated
#virtualized

diff --git a/Engine/Plugins/Runtime/Harmonix/Source/HarmonixMetasound/Public/HarmonixMetasound/Components/MusicClockComponent.h b/Engine/Plugins/Runtime/Harmonix/Source/HarmonixMetasound/Public/HarmonixMetasound/Components/MusicClockComponent.h
--- a/Engine/Plugins/Runtime/Harmonix/Source/HarmonixMetasound/Public/HarmonixMetasound/Components/MusicClockComponent.h
+++ b/Engine/Plugins/Runtime/Harmonix/Source/HarmonixMetasound/Public/HarmonixMetasound/Components/MusicClockComponent.h
@@ -2,14 +2,12 @@
 #pragma once
 
 #include "Components/ActorComponent.h"
-#include "Delegates/DelegateCombinations.h"
 #include "Evaluation/IMovieSceneCustomClockSource.h"
 #include "Harmonix/MusicalTimebase.h"
 #include "HarmonixMidi/MidiSongPos.h"
 #include "HarmonixMidi/SongMaps.h"
 #include "MusicEnvironmentClockSource.h"
 #include "MusicClockTypes.h"
-#include "Templates/UniquePtr.h"
 #include "StructUtils/InstancedStruct.h"
 
 #include "MusicClockComponent.generated.h"
