Can't directly call "Seek" in "Reset" of clock.

Instead, reset other state, which will do the work of resetting and seeking

#rnx
#DisabledValidations nonUnityAndStaticAnalysis

diff --git a/Engine/Plugins/Runtime/Harmonix/Source/HarmonixMetasound/Private/HarmonixMetasound/DataTypes/MidiClock.cpp b/Engine/Plugins/Runtime/Harmonix/Source/HarmonixMetasound/Private/HarmonixMetasound/DataTypes/MidiClock.cpp
--- a/Engine/Plugins/Runtime/Harmonix/Source/HarmonixMetasound/Private/HarmonixMetasound/DataTypes/MidiClock.cpp
+++ b/Engine/Plugins/Runtime/Harmonix/Source/HarmonixMetasound/Private/HarmonixMetasound/DataTypes/MidiClock.cpp
@@ -200,6 +200,44 @@ namespace HarmonixMetasound
 		AttachToSongMapEvaluator(nullptr, false);
 	}
 
+	void FMidiClock::Reset(const FOperatorSettings& InSettings)
+	{
+		SongMapEvaluator = MakeShared<FSongMapsWithAlternateTempoSource>(MakeShared<FSongMaps>(120.0f, 4, 4));
+		CurrentTempoInfoPointIndex = 0;
+		CurrentTimeSignaturePointIndex = 0;
+		TickResidualWhenDriven = 0.0f;
+		BlockSize = InSettings.GetNumFramesPerBlock();
+		CurrentBlockFrameIndex = 0;
+		FirstTickProcessedThisBlock = -1;
+		LastProcessedMidiTick = -1;
+		NextMidiTickToProcess = 0;
+		NextTempoMapTickToProcess = 0;
+		SampleRate = InSettings.GetSampleRate();
+		SampleCount = 0;
+		FramesUntilNextProcess = 0;
+		TransportAtBlockStart = EMusicPlayerTransportState::Invalid;
+		TransportAtBlockEnd = EMusicPlayerTransportState::Invalid;
+		SpeedAtBlockStart = 0.0f;
+		SpeedAtBlockEnd = 0.0f;
+		CurrentLocalSpeed = -1.0f;
+		TempoAtBlockStart = 0.0f;
+		TempoAtBlockEnd = 0.0f;
+		TimeSignatureAtBlockStart = FTimeSignature(0, 1);
+		TimeSignatureAtBlockEnd = FTimeSignature(0, 1);
+		NumTransportChangeInBlock = 0;
+		NumSpeedChangeInBlock = 0;
+		NumTempoChangeInBlock = 0;
+		NumTimeSignatureChangeInBlock = 0;
+		NextTempoChangeTick = std::numeric_limits<int32>::max();
+		NextTimeSigChangeTick = std::numeric_limits<int32>::max();
+		NextTempoOrTimeSigChangeTick = std::numeric_limits<int32>::max();
+		FirstTickInLoop = -1;
+		LoopLengthTicks = 0;
+		MidiDataChangedInBlock = false;
+		NeedsSeekToDrivingClock = false;
+		MidiClockEventsInBlock.Reset();
+	}
+
 	void FMidiClock::SetDrivingClock(FConstSharedMidiClockPtr NewExternalClockDriver)
 	{
 		MidiDataChangedInBlock = true;
diff --git a/Engine/Plugins/Runtime/Harmonix/Source/HarmonixMetasound/Private/HarmonixMetasound/Nodes/MetronomeNode.cpp b/Engine/Plugins/Runtime/Harmonix/Source/HarmonixMetasound/Private/HarmonixMetasound/Nodes/MetronomeNode.cpp
--- a/Engine/Plugins/Runtime/Harmonix/Source/HarmonixMetasound/Private/HarmonixMetasound/Nodes/MetronomeNode.cpp
+++ b/Engine/Plugins/Runtime/Harmonix/Source/HarmonixMetasound/Private/HarmonixMetasound/Nodes/MetronomeNode.cpp
@@ -206,11 +206,15 @@ namespace HarmonixMetasound::Nodes::MetronomeNode
 
 	void FMetronomeOperatorBase::Reset(const FResetParams& Params)
 	{
+		MidiClockOutPin->Reset(Params.OperatorSettings);
 		BlockSize = Params.OperatorSettings.GetNumFramesPerBlock();
 		SampleRate = Params.OperatorSettings.GetSampleRate();
-		
 		LastProcessedClockTick = -1;
 		NextClockTickToProcess = 0;
+		CurrentTempo = GetInputTempoClamped();
+		CurrentTimeSigNum = GetInputTimeSigNumeratorClamped();
+		CurrentTimeSigDenom = GetInputTimeSigDenominatorClamped();
+        Init();
 	}
 
 	void FMetronomeOperatorBase::Init()
diff --git a/Engine/Plugins/Runtime/Harmonix/Source/HarmonixMetasound/Private/HarmonixMetasound/Nodes/MidiPlayerNode.cpp b/Engine/Plugins/Runtime/Harmonix/Source/HarmonixMetasound/Private/HarmonixMetasound/Nodes/MidiPlayerNode.cpp
--- a/Engine/Plugins/Runtime/Harmonix/Source/HarmonixMetasound/Private/HarmonixMetasound/Nodes/MidiPlayerNode.cpp
+++ b/Engine/Plugins/Runtime/Harmonix/Source/HarmonixMetasound/Private/HarmonixMetasound/Nodes/MidiPlayerNode.cpp
@@ -319,6 +319,7 @@ namespace HarmonixMetasound::Nodes::MidiPlayerNode
 	void FExternallyClockedMidiPlayerOperator::Reset(const FResetParams& Params)
 	{
 		FMidiPlayerOperator::Reset(Params);
+		MidiClockOut->SetDrivingClock(MidiClockIn->AsShared().ToSharedPtr());
 	}
 
 	void FSelfClockedMidiPlayerOperator::BindInputs(FInputVertexInterfaceData& InVertexData)
@@ -524,10 +525,11 @@ namespace HarmonixMetasound::Nodes::MidiPlayerNode
 	{
 		BlockSize = Params.OperatorSettings.GetNumFramesPerBlock();
 		CurrentBlockSpanStart = 0;
-
+		CurrentRenderBlockFrame = 0;
+		CurrentMidiFile = nullptr;
+		
 		MidiOutPin->SetClock(*MidiClockOut);
-		MidiClockOut->SeekTo(0, 0, 0);
-		MidiClockOut->SetTransportState(0, EMusicPlayerTransportState::Prepared);
+		MidiClockOut->Reset(Params.OperatorSettings);
 
 		NeedsTransportInit = true;
 	}
diff --git a/Engine/Plugins/Runtime/Harmonix/Source/HarmonixMetasound/Public/HarmonixMetasound/DataTypes/MidiClock.h b/Engine/Plugins/Runtime/Harmonix/Source/HarmonixMetasound/Public/HarmonixMetasound/DataTypes/MidiClock.h
--- a/Engine/Plugins/Runtime/Harmonix/Source/HarmonixMetasound/Public/HarmonixMetasound/DataTypes/MidiClock.h
+++ b/Engine/Plugins/Runtime/Harmonix/Source/HarmonixMetasound/Public/HarmonixMetasound/DataTypes/MidiClock.h
@@ -51,6 +51,7 @@ namespace HarmonixMetasound
 		UE_API void AttachToSongMapEvaluator(TSharedPtr<ISongMapEvaluator> SongMaps, bool ResetToStart = true);
 		UE_API void SongMapsChanged();
 		UE_API void DetachFromSongMaps();
+		UE_API void Reset(const Metasound::FOperatorSettings& InSettings);
 
 		const ISongMapEvaluator& GetSongMapEvaluator() const { return *SongMapEvaluator; }
 
