Harmonix/MidiEvent.h

FMidiEvent, default constructor

#rnx
#tests HarmonixDev

diff --git a/Engine/Plugins/Runtime/Harmonix/Source/HarmonixMidi/Public/HarmonixMidi/MidiEvent.h b/Engine/Plugins/Runtime/Harmonix/Source/HarmonixMidi/Public/HarmonixMidi/MidiEvent.h
--- a/Engine/Plugins/Runtime/Harmonix/Source/HarmonixMidi/Public/HarmonixMidi/MidiEvent.h
+++ b/Engine/Plugins/Runtime/Harmonix/Source/HarmonixMidi/Public/HarmonixMidi/MidiEvent.h
@@ -19,6 +19,11 @@ struct FMidiEvent
 	GENERATED_BODY()
 
 public:
+	// default construct for USTRUCT
+	FMidiEvent()
+	: Tick(0)
+	{}
+	
 	UE_API FMidiEvent(int32 Tick, const FMidiMsg& InMessage);
 
 	bool operator==(const FMidiEvent& Other) const
@@ -52,11 +57,6 @@ public:
 	{
 	}
 
-protected:
-	FMidiEvent()
-		: Tick(0)
-	{}
-
 private:
 	friend struct FMidiTrack;
 	friend struct FMidiFileData;
