Implement randomizer modulators on the FusionSampler so that the Randomizer settings actually work.

For Pitch modulation, the range goes from negative magnitude to positive magnitude, proportional to the depth

For Start Point Ms, the range goes from 0 to positive magnitude, proportional to the depth.

[RN] Implement "Randomizer" modulation on Fusion Sampler with Fusion Patch Settings. Hide "Velocity" modulation settings
#tests HarmonixDev

diff --git a/Engine/Plugins/Runtime/Harmonix/Source/HarmonixDsp/Private/HarmonixDsp/FusionSampler/FusionSampler.cpp b/Engine/Plugins/Runtime/Harmonix/Source/HarmonixDsp/Private/HarmonixDsp/FusionSampler/FusionSampler.cpp
--- a/Engine/Plugins/Runtime/Harmonix/Source/HarmonixDsp/Private/HarmonixDsp/FusionSampler/FusionSampler.cpp
+++ b/Engine/Plugins/Runtime/Harmonix/Source/HarmonixDsp/Private/HarmonixDsp/FusionSampler/FusionSampler.cpp
@@ -557,18 +557,8 @@ bool FFusionSampler::TryKeyOnZone(FMidiVoiceId InVoiceId, uint8 InTriggeredNote,
 	ActiveVoices.AddTail(Voice);
 
 	// From here on out we use the transposed note to that the correct pitch is played...
-
-	for (int32 ModulatorIdx = 0; ModulatorIdx < kNumModulators; ++ModulatorIdx)
-	{
-		Randomizers[ModulatorIdx].Modulate(FMath::FRand());
-	}
-
-
+	
 	float VelocityNorm = (float)InVelocity * sVelocityNormalizer;
-	for (int32 ModulatorIdx = 0; ModulatorIdx < kNumModulators; ++ModulatorIdx)
-	{
-		VelocityModulators[ModulatorIdx].Modulate(VelocityNorm);
-	}
 
 	float AttackGain = VelocityNorm; // should this really just be linear?
 	if (!InKeyzone->bVelocityToGain)
@@ -739,6 +729,16 @@ void FFusionSampler::ApplyPatchSettings()
 		LfoSettings[Idx].CopySettings(Settings.Lfo[Idx]);
 	}
 
+	for (int32 Idx = 0; Idx < kNumModulators; ++Idx)
+	{
+		RandomizerSettings[Idx] = Settings.Randomizer[Idx];
+	}
+
+	for (int32 Idx = 0; Idx < kNumModulators; ++Idx)
+	{
+		VelocityModulatorSettings[Idx] = Settings.VelocityModulator[Idx];
+	}
+
 	//------------------------------------
 	// Read in portamento settings
 	//------------------------------------
@@ -859,16 +859,13 @@ void FFusionSampler::ResetPatchRelatedState()
 	// set up the default randomizer settings
 	for (uint8 ModulatorIdx = 0; ModulatorIdx < (uint8)EModulatorTarget::Num; ++ModulatorIdx)
 	{
-		Randomizers[ModulatorIdx].Reset();
-		Randomizers[ModulatorIdx].SetTarget(&Harmonix::Dsp::Modulators::FModulatorTarget::kDummyTarget);
+		RandomizerSettings[ModulatorIdx] = {};
 	}
-
-
+	
 	// set up the default randomizer settings
 	for (int32 ModulatorIdx = 0; ModulatorIdx < kNumModulators; ++ModulatorIdx)
 	{
-		VelocityModulators[ModulatorIdx].Reset();
-		VelocityModulators[ModulatorIdx].SetTarget(&Harmonix::Dsp::Modulators::FModulatorTarget::kDummyTarget);
+		VelocityModulatorSettings[ModulatorIdx] = {};
 	}
 
 
diff --git a/Engine/Plugins/Runtime/Harmonix/Source/HarmonixDsp/Private/HarmonixDsp/FusionSampler/FusionVoice.cpp b/Engine/Plugins/Runtime/Harmonix/Source/HarmonixDsp/Private/HarmonixDsp/FusionSampler/FusionVoice.cpp
--- a/Engine/Plugins/Runtime/Harmonix/Source/HarmonixDsp/Private/HarmonixDsp/FusionSampler/FusionVoice.cpp
+++ b/Engine/Plugins/Runtime/Harmonix/Source/HarmonixDsp/Private/HarmonixDsp/FusionSampler/FusionVoice.cpp
@@ -298,7 +298,7 @@ void FFusionVoice::SetPitchShiftForMidiNote(uint8 InNote)
 	{
 		SemitoneOffset = 0;
 	}
-	double NewPitchOffsetCents = (double)SemitoneOffset * 100.0 + PitchOffsetCents;
+	double NewPitchOffsetCents = (double)SemitoneOffset * 100.0 + PitchModulation * 100 + PitchOffsetCents;
 
 	PitchShift = FMath::Pow(2.0, NewPitchOffsetCents * HarmonixDsp::kOctavesPerCent);
 }
@@ -340,6 +340,23 @@ void FFusionVoice::AttackWithTargetNote(uint8 InMidiNoteNumber, float InGain, in
 		SemitoneOffset = 0;
 	}
 
+	PitchModulation = 0.0;
+	float StartPointMsModulation = 0.0f;
+	for (int32 Idx = 0; Idx < FFusionPatchSettings::kNumModulators; ++Idx)
+	{
+		const FModulatorSettings& Settings = MySampler->RandomizerSettings[Idx];
+		if (Settings.Target == EModulatorTarget::Pitch)
+		{
+			float X = FMath::FRand();
+			PitchModulation += FMath::Clamp(FMath::Lerp(-Settings.Range, Settings.Range, X * Settings.Depth), -Settings.Range, Settings.Range);
+		}
+		else if (Settings.Target == EModulatorTarget::StartPoint)
+		{
+			float X = FMath::FRand();
+			StartPointMsModulation += FMath::Clamp(FMath::Lerp(0, Settings.Range, X * Settings.Depth), 0, Settings.Range);
+		}
+	}
+	
 	// this is the note (pitch) we are supposedly playing
 	// although the exact pitch will vary depending on various modulations
 	TargetMidiNote = (float)InMidiNoteNumber;
@@ -373,10 +390,9 @@ void FFusionVoice::AttackWithTargetNote(uint8 InMidiNoteNumber, float InGain, in
 		StartQuarterNote = (float)InEventTick / Harmonix::Midi::Constants::GTicksPerQuarterNote;
 	}
 
-	// multiply by 100.0?
-	PrepareWithPitchOffsetAndGain(SemitoneOffset * 100.0 + PitchOffsetCents, InGain);
+	PrepareWithPitchOffsetAndGain(SemitoneOffset * 100.0 + PitchModulation * 100.0 + PitchOffsetCents, InGain);
 
-	double OutputSampleOffset = 0.001 * InStartPointMs * MySampler->GetSamplesPerSecond();
+	double OutputSampleOffset = 0.001 * (InStartPointMs + StartPointMsModulation) * MySampler->GetSamplesPerSecond();
 	double KeyzoneStartOffset = (KeyZone->SampleStartOffset == -1) ?
 		0 :
 		KeyZone->SampleStartOffset;
diff --git a/Engine/Plugins/Runtime/Harmonix/Source/HarmonixDsp/Public/HarmonixDsp/FusionSampler/FusionSampler.h b/Engine/Plugins/Runtime/Harmonix/Source/HarmonixDsp/Public/HarmonixDsp/FusionSampler/FusionSampler.h
--- a/Engine/Plugins/Runtime/Harmonix/Source/HarmonixDsp/Public/HarmonixDsp/FusionSampler/FusionSampler.h
+++ b/Engine/Plugins/Runtime/Harmonix/Source/HarmonixDsp/Public/HarmonixDsp/FusionSampler/FusionSampler.h
@@ -7,7 +7,6 @@
 #include "HarmonixDsp/Effects/BiquadFilter.h"
 
 #include "HarmonixDsp/Modulators/Settings/AdsrSettings.h"
-#include "HarmonixDsp/Modulators/Modulator.h"
 #include "HarmonixDsp/Modulators/Lfo.h"
 
 #include "HarmonixDsp/FusionSampler/FusionPatch.h"
@@ -294,8 +293,8 @@ private:
 	FLfoSettings LfoSettings[kNumLfos];
 	Harmonix::Dsp::Modulators::FLfo Lfos[kNumLfos];
 	static constexpr int32 kNumModulators = 2;
-	Harmonix::Dsp::Modulators::FModulator Randomizers[kNumModulators];
-	Harmonix::Dsp::Modulators::FModulator VelocityModulators[kNumModulators];
+	FModulatorSettings RandomizerSettings[kNumModulators];
+	FModulatorSettings VelocityModulatorSettings[kNumModulators];
 
 	// soft limit, may be over max if some voices are in release stage
 	uint32 MaxNumVoices = 100;
diff --git a/Engine/Plugins/Runtime/Harmonix/Source/HarmonixDsp/Public/HarmonixDsp/FusionSampler/FusionVoice.h b/Engine/Plugins/Runtime/Harmonix/Source/HarmonixDsp/Public/HarmonixDsp/FusionSampler/FusionVoice.h
--- a/Engine/Plugins/Runtime/Harmonix/Source/HarmonixDsp/Public/HarmonixDsp/FusionSampler/FusionVoice.h
+++ b/Engine/Plugins/Runtime/Harmonix/Source/HarmonixDsp/Public/HarmonixDsp/FusionSampler/FusionVoice.h
@@ -210,7 +210,8 @@ private:
 	double ResampleRate = 0.0;
 
 	double PitchOffsetCents = 0.0;
-
+	double PitchModulation = 0.0;
+	
 	double EndOfSampleData = 0.0;
 	
 	FFusionVoicePool* VoicePool = nullptr;
