Additional null checks and logging for cases where a custom clock driver settings class fails to return a valid clock driver.

#tests PIE
#rnx

diff --git a/Engine/Plugins/Runtime/Harmonix/Source/HarmonixMetasound/Private/HarmonixMetasound/Components/MusicClockComponent.cpp b/Engine/Plugins/Runtime/Harmonix/Source/HarmonixMetasound/Private/HarmonixMetasound/Components/MusicClockComponent.cpp
--- a/Engine/Plugins/Runtime/Harmonix/Source/HarmonixMetasound/Private/HarmonixMetasound/Components/MusicClockComponent.cpp
+++ b/Engine/Plugins/Runtime/Harmonix/Source/HarmonixMetasound/Private/HarmonixMetasound/Components/MusicClockComponent.cpp
@@ -160,6 +160,7 @@ void UMusicClockComponent::CreateClockDriver()
 
 	if (!ClockDriver)
 	{
+		UE_LOG(LogMusicClock, Error, TEXT("Failed to create clock driver as configured for component '%s', falling back to wall clock"), *GetNameSafe(this));
 		ConnectToWallClock();
 	}
 }
@@ -215,7 +216,19 @@ void UMusicClockComponent::ConnectToCustomClock()
 	check(DriveMethod == EMusicClockDriveMethod::Custom);
 	DisconnectFromClockDriver();
 	const FMusicClockSettingsBase* Settings = MusicClockSettings.GetPtr<FMusicClockSettingsBase>();
+	if (!Settings)
+	{
+		UE_LOG(LogMusicClock, Error, TEXT("Invalid settings for custom clock driver on component '%s'"), *GetNameSafe(this));
+		return;
+	}
+
 	ClockDriver = Settings->MakeInstance(this);
+	if (!ClockDriver)
+	{
+		UE_LOG(LogMusicClock, Error, TEXT("Failed to create custom clock driver on component '%s'"), *GetNameSafe(this));
+		return;
+	}
+
 	if (IsActive())
 	{
 		ClockDriver->Start();
@@ -234,7 +247,7 @@ void UMusicClockComponent::DisconnectFromClockDriver()
 FMidiSongPos UMusicClockComponent::CalculateSongPosWithOffset(float MsOffset, ECalibratedMusicTimebase Timebase) const
 {
 	FMidiSongPos Result;
-	if (ClockDriver != nullptr)
+	if (ClockDriver)
 	{
 		if (ClockDriver->CalculateSongPosWithOffset(MsOffset, Timebase, Result))
 		{
@@ -290,11 +303,15 @@ void UMusicClockComponent::Activate(bool bReset)
 	{
 		CreateClockDriver();
 	}
-	ClockDriver->Start();
-	LastBroadcastBeat = -1;
-	LastBroadcastBar  = -1;
-	UMidiClockUpdateSubsystem::TrackMusicClockComponent(this);
-	PlayStateEvent.Broadcast(EMusicClockState::Running);
+
+	if (ensure(ClockDriver))
+	{
+		ClockDriver->Start();
+		LastBroadcastBeat = -1;
+		LastBroadcastBar = -1;
+		UMidiClockUpdateSubsystem::TrackMusicClockComponent(this);
+		PlayStateEvent.Broadcast(EMusicClockState::Running);
+	}
 }
 
 void UMusicClockComponent::Deactivate()
@@ -881,6 +898,11 @@ const FMidiSongPos& UMusicClockComponent::GetPreviousSongPosInternal(ECalibrated
 
 float UMusicClockComponent::GetPositionSeconds() const
 {
+	if (!ClockDriver)
+	{
+		return 0.0f;
+	}
+
 	const FMidiSongPos& SongPos = ClockDriver->GetCurrentSongPos(ECalibratedMusicTimebase::VideoRenderTime);
 	return SongPos.SecondsIncludingCountIn;
 }
