FusionSampler: Fix comment

#rnx

diff --git a/Engine/Plugins/Runtime/Harmonix/Source/HarmonixDsp/Private/HarmonixDsp/FusionSampler/FusionSampler.cpp b/Engine/Plugins/Runtime/Harmonix/Source/HarmonixDsp/Private/HarmonixDsp/FusionSampler/FusionSampler.cpp
--- a/Engine/Plugins/Runtime/Harmonix/Source/HarmonixDsp/Private/HarmonixDsp/FusionSampler/FusionSampler.cpp
+++ b/Engine/Plugins/Runtime/Harmonix/Source/HarmonixDsp/Private/HarmonixDsp/FusionSampler/FusionSampler.cpp
@@ -862,13 +862,12 @@ void FFusionSampler::ResetPatchRelatedState()
 		RandomizerSettings[ModulatorIdx] = {};
 	}
 	
-	// set up the default randomizer settings
+	// set up the default velocity modulator settings
 	for (int32 ModulatorIdx = 0; ModulatorIdx < kNumModulators; ++ModulatorIdx)
 	{
 		VelocityModulatorSettings[ModulatorIdx] = {};
 	}
-
-
+	
 	MaxNumVoices = Defaults::kMaxNumVoices;
 
 	SetIsPortamentoEnabled(false);
