Bug.

#rnx

diff --git a/Engine/Plugins/Runtime/Harmonix/Source/HarmonixMetasound/Private/HarmonixMetasound/Components/MetasoundMusicClockDriver.cpp b/Engine/Plugins/Runtime/Harmonix/Source/HarmonixMetasound/Private/HarmonixMetasound/Components/MetasoundMusicClockDriver.cpp
--- a/Engine/Plugins/Runtime/Harmonix/Source/HarmonixMetasound/Private/HarmonixMetasound/Components/MetasoundMusicClockDriver.cpp
+++ b/Engine/Plugins/Runtime/Harmonix/Source/HarmonixMetasound/Private/HarmonixMetasound/Components/MetasoundMusicClockDriver.cpp
@@ -330,7 +330,7 @@ bool FMetasoundMusicClockDriver::ConnectToAudioComponentsMetasound(
 	UMetasoundGeneratorHandle::FOnAttached::FDelegate OnGeneratorAttachedCallback,
 	UMetasoundGeneratorHandle::FOnDetached::FDelegate OnGeneratorDetachedCallback)
 {
-	OnAttachedDelegate = OnGeneratorDetachedCallback;
+	OnAttachedDelegate = OnGeneratorAttachedCallback;
 	OnDetachedDelegate = OnGeneratorDetachedCallback;
 	
 	AudioComponentToWatch = InAudioComponent;
