MusicClockNDI NDC

#tests PIE

#changelist validated
#virtualized

diff --git a/Engine/Plugins/Runtime/Harmonix/Source/HarmonixMetasound/Private/HarmonixMetasound/Subsystems/MidiClockUpdateSubsystem.cpp b/Engine/Plugins/Runtime/Harmonix/Source/HarmonixMetasound/Private/HarmonixMetasound/Subsystems/MidiClockUpdateSubsystem.cpp
--- a/Engine/Plugins/Runtime/Harmonix/Source/HarmonixMetasound/Private/HarmonixMetasound/Subsystems/MidiClockUpdateSubsystem.cpp
+++ b/Engine/Plugins/Runtime/Harmonix/Source/HarmonixMetasound/Private/HarmonixMetasound/Subsystems/MidiClockUpdateSubsystem.cpp
@@ -154,10 +154,14 @@ UMidiClockUpdateSubsystem::FClockHistoryPtr UMidiClockUpdateSubsystem::GetOrCrea
 void UMidiClockUpdateSubsystem::TrackMusicClockComponentImpl(UMusicClockComponent* Clock)
 {
 	TrackedMusicClockComponents.AddUnique(Clock);
+
+	OnTrackMusicClockComponentDelegate.Broadcast(Clock);
 }
 
 void UMidiClockUpdateSubsystem::StopTrackingMusicClockComponentImpl(UMusicClockComponent* Clock)
 {
+	OnStopTrackingMusicClockComponentDelegate.Broadcast(Clock);
+
 	TrackedMusicClockComponents.Remove(Clock);
 }
 
@@ -204,6 +208,8 @@ void UMidiClockUpdateSubsystem::UpdateUMusicClockComponents()
 		}
 	}
 
+	OnMusicClockComponentsUpdatedDelegate.Broadcast(TrackedMusicClockComponents);
+
 	for (auto It = ClockHistories.CreateIterator(); It; ++It)
 	{
 		if (!It.Value().IsValid())
diff --git a/Engine/Plugins/Runtime/Harmonix/Source/HarmonixMetasound/Public/HarmonixMetasound/Subsystems/MidiClockUpdateSubsystem.h b/Engine/Plugins/Runtime/Harmonix/Source/HarmonixMetasound/Public/HarmonixMetasound/Subsystems/MidiClockUpdateSubsystem.h
--- a/Engine/Plugins/Runtime/Harmonix/Source/HarmonixMetasound/Public/HarmonixMetasound/Subsystems/MidiClockUpdateSubsystem.h
+++ b/Engine/Plugins/Runtime/Harmonix/Source/HarmonixMetasound/Public/HarmonixMetasound/Subsystems/MidiClockUpdateSubsystem.h
@@ -70,6 +70,31 @@ public:
 	using FClockHistoryPtr = TSharedPtr<HarmonixMetasound::Analysis::FMidiClockSongPositionHistory>;
 	static UE_API FClockHistoryPtr GetOrCreateClockHistory(const Metasound::Frontend::FAnalyzerAddress& ForAddress);
 
+	DECLARE_MULTICAST_DELEGATE_OneParam(FOnMusicClockComponentsUpdated, const TArray<TWeakObjectPtr<UMusicClockComponent>>&);
+	FOnMusicClockComponentsUpdated& OnMusicClockComponentsUpdated()
+	{
+		check(IsInGameThread());
+		return OnMusicClockComponentsUpdatedDelegate;
+	}
+
+	DECLARE_MULTICAST_DELEGATE_OneParam(FOnTrackMusicClockComponentChange, UMusicClockComponent*);
+	FOnTrackMusicClockComponentChange& OnTrackMusicClockComponent()
+	{
+		check(IsInGameThread());
+		return OnTrackMusicClockComponentDelegate;
+	}
+
+	FOnTrackMusicClockComponentChange& OnStopTrackingMusicClockComponent()
+	{
+		check(IsInGameThread());
+		return OnStopTrackingMusicClockComponentDelegate;
+	}
+	
+	const TArray<TWeakObjectPtr<UMusicClockComponent>>& GetTrackedMusicClockComponents() const
+	{
+		return TrackedMusicClockComponents;
+	}
+
 private:
 
 	mutable FCriticalSection TrackedMidiClocksMutex;
@@ -90,6 +115,10 @@ private:
 
 	static UE_API uint32 MakeMidiSongPosAnalyzerAddressHash(const Metasound::Frontend::FAnalyzerAddress& ForAddress);
 
+	FOnTrackMusicClockComponentChange OnTrackMusicClockComponentDelegate;
+	FOnTrackMusicClockComponentChange OnStopTrackingMusicClockComponentDelegate;
+	FOnMusicClockComponentsUpdated OnMusicClockComponentsUpdatedDelegate;
+
 public:
 	// Declare a "tick" method that can be used during automated testing so that
 	// the test code doesn't need knowledge of how the low-res clocks are being ticked...
