[MusicClock] Changed the MusicClockComponent::Create static methods so that we don't have to provide an actor (matches CreateSound2D behavior)

#rnx

diff --git a/Engine/Plugins/Runtime/Harmonix/Source/HarmonixMetasound/Private/HarmonixMetasound/Components/MusicClockComponent.cpp b/Engine/Plugins/Runtime/Harmonix/Source/HarmonixMetasound/Private/HarmonixMetasound/Components/MusicClockComponent.cpp
--- a/Engine/Plugins/Runtime/Harmonix/Source/HarmonixMetasound/Private/HarmonixMetasound/Components/MusicClockComponent.cpp
+++ b/Engine/Plugins/Runtime/Harmonix/Source/HarmonixMetasound/Private/HarmonixMetasound/Components/MusicClockComponent.cpp
@@ -5,6 +5,7 @@
 #include "HarmonixMetasound/Subsystems/MidiClockUpdateSubsystem.h"
 #include "HarmonixMidi/MusicTimeSpan.h"
 #include "Components/AudioComponent.h"
+#include "Engine/World.h"
 #include "Templates/SharedPointer.h"
 #include "GameFramework/Actor.h"
 
@@ -22,12 +23,16 @@ UMusicClockComponent::UMusicClockComponent()
 UMusicClockComponent* UMusicClockComponent::CreateMetasoundDrivenMusicClock(UObject* WorldContextObject, UAudioComponent* InAudioComponent, FName MetasoundOutputPinName, bool Start)
 {
 	UMusicClockComponent* NewClock = NewObject<UMusicClockComponent>(WorldContextObject);
-	if (!NewClock->GetOwner())
+
+	if (NewClock->GetOwner())
 	{
-		UE_LOG(LogMusicClock, Warning, TEXT("Can't create a music clock in a non-actor context!"));
-		return nullptr;
+		NewClock->RegisterComponent();
+	}
+	else if (UWorld* World = WorldContextObject ? WorldContextObject->GetWorld() : nullptr)
+	{
+		NewClock->RegisterComponentWithWorld(World);
 	}
-	NewClock->RegisterComponent();
+
 	NewClock->DriveMethod = EMusicClockDriveMethod::MetaSound;
 	NewClock->MetasoundOutputName = MetasoundOutputPinName;
 	NewClock->ConnectToMetasoundOnAudioComponent(InAudioComponent);
@@ -41,12 +46,16 @@ UMusicClockComponent* UMusicClockComponent::CreateMetasoundDrivenMusicClock(UObj
 UMusicClockComponent* UMusicClockComponent::CreateWallClockDrivenMusicClock(UObject* WorldContextObject, UMidiFile* InTempoMap, bool Start)
 {
 	UMusicClockComponent* NewClock = NewObject<UMusicClockComponent>(WorldContextObject);
-	if (!NewClock->GetOwner())
+
+	if (NewClock->GetOwner())
 	{
-		UE_LOG(LogMusicClock, Warning, TEXT("Can't create a music clock in a non-actor context!"));
-		return nullptr;
+		NewClock->RegisterComponent();
+	}
+	else if (UWorld* World = WorldContextObject ? WorldContextObject->GetWorld() : nullptr)
+	{
+		NewClock->RegisterComponentWithWorld(World);
 	}
-	NewClock->RegisterComponent();
+
 	NewClock->ConnectToWallClockForMidi(InTempoMap);
 	if (Start)
 	{
@@ -58,13 +67,16 @@ UMusicClockComponent* UMusicClockComponent::CreateWallClockDrivenMusicClock(UObj
 UMusicClockComponent* UMusicClockComponent::CreateMusicClockWithSettings(UObject* WorldContextObject, TInstancedStruct<FMusicClockSettingsBase> MusicClockSettings, bool Start)
 {
 	UMusicClockComponent* NewClock = NewObject<UMusicClockComponent>(WorldContextObject);
-	if (!NewClock->GetOwner())
+
+	if (NewClock->GetOwner())
 	{
-		UE_LOG(LogMusicClock, Warning, TEXT("Can't create a music clock in a non-actor context!"));
-		return nullptr;
+		NewClock->RegisterComponent();
 	}
-	NewClock->RegisterComponent();
-	
+	else if (UWorld* World = WorldContextObject ? WorldContextObject->GetWorld() : nullptr)
+	{
+		NewClock->RegisterComponentWithWorld(World);
+	}
+
 	if (const FMetasoundMusicClockSettings* MetasoundClockSettings = MusicClockSettings.GetPtr<FMetasoundMusicClockSettings>())
 	{
 		NewClock->DriveMethod = EMusicClockDriveMethod::MetaSound;
