UMusicClockComponent::TryGetAudioSourceLocation

#tests PIE

diff --git a/Engine/Plugins/Runtime/Harmonix/Source/HarmonixMetasound/Private/HarmonixMetasound/Components/MetasoundMusicClockDriver.cpp b/Engine/Plugins/Runtime/Harmonix/Source/HarmonixMetasound/Private/HarmonixMetasound/Components/MetasoundMusicClockDriver.cpp
--- a/Engine/Plugins/Runtime/Harmonix/Source/HarmonixMetasound/Private/HarmonixMetasound/Components/MetasoundMusicClockDriver.cpp
+++ b/Engine/Plugins/Runtime/Harmonix/Source/HarmonixMetasound/Private/HarmonixMetasound/Components/MetasoundMusicClockDriver.cpp
@@ -49,6 +49,15 @@ namespace MetasoundMusicClockDriver
 		ECVF_Default);
 }
 
+TOptional<FVector> FMetasoundMusicClockDriver::TryGetAudioSourceLocation() const
+{
+	if (UAudioComponent* AudioComponent = AudioComponentToWatch.Get())
+	{
+		return AudioComponent->GetComponentLocation();
+	}
+	return {};
+}
+
 bool FMetasoundMusicClockDriver::CalculateSongPosWithOffset(float MsOffset, ECalibratedMusicTimebase Timebase, FMidiSongPos& OutResult) const
 {
 	using namespace HarmonixMetasound::Analysis;
diff --git a/Engine/Plugins/Runtime/Harmonix/Source/HarmonixMetasound/Private/HarmonixMetasound/Components/MetasoundMusicClockDriver.h b/Engine/Plugins/Runtime/Harmonix/Source/HarmonixMetasound/Private/HarmonixMetasound/Components/MetasoundMusicClockDriver.h
--- a/Engine/Plugins/Runtime/Harmonix/Source/HarmonixMetasound/Private/HarmonixMetasound/Components/MetasoundMusicClockDriver.h
+++ b/Engine/Plugins/Runtime/Harmonix/Source/HarmonixMetasound/Private/HarmonixMetasound/Components/MetasoundMusicClockDriver.h
@@ -17,6 +17,8 @@ public:
 	FMetasoundMusicClockDriver(UObject* WorldContextObj, const FMetasoundMusicClockSettings& Settings)
 		: FMusicClockDriverBase(WorldContextObj, Settings) {}
 
+	virtual TOptional<FVector> TryGetAudioSourceLocation() const override;
+
 	virtual bool CalculateSongPosWithOffset(float MsOffset, ECalibratedMusicTimebase Timebase, FMidiSongPos& OutResult) const override;
 
 	virtual void Disconnect() override;
diff --git a/Engine/Plugins/Runtime/Harmonix/Source/HarmonixMetasound/Private/HarmonixMetasound/Components/MusicClockComponent.cpp b/Engine/Plugins/Runtime/Harmonix/Source/HarmonixMetasound/Private/HarmonixMetasound/Components/MusicClockComponent.cpp
--- a/Engine/Plugins/Runtime/Harmonix/Source/HarmonixMetasound/Private/HarmonixMetasound/Components/MusicClockComponent.cpp
+++ b/Engine/Plugins/Runtime/Harmonix/Source/HarmonixMetasound/Private/HarmonixMetasound/Components/MusicClockComponent.cpp
@@ -746,6 +746,15 @@ bool UMusicClockComponent::LoopedThisFrame(ECalibratedMusicTimebase Timebase) co
 	return false;
 }
 
+TOptional<FVector> UMusicClockComponent::TryGetAudioSourceLocation() const
+{
+	if (ClockDriver)
+	{
+		return ClockDriver->TryGetAudioSourceLocation();
+	}
+	return {};
+}
+
 const ISongMapEvaluator& UMusicClockComponent::GetSongMaps() const
 {
 	const ISongMapEvaluator* SongMaps = ClockDriver ? ClockDriver->GetCurrentSongMapEvaluator() : nullptr;
diff --git a/Engine/Plugins/Runtime/Harmonix/Source/HarmonixMetasound/Private/HarmonixMetasound/Components/WallClockMusicClockDriver.cpp b/Engine/Plugins/Runtime/Harmonix/Source/HarmonixMetasound/Private/HarmonixMetasound/Components/WallClockMusicClockDriver.cpp
--- a/Engine/Plugins/Runtime/Harmonix/Source/HarmonixMetasound/Private/HarmonixMetasound/Components/WallClockMusicClockDriver.cpp
+++ b/Engine/Plugins/Runtime/Harmonix/Source/HarmonixMetasound/Private/HarmonixMetasound/Components/WallClockMusicClockDriver.cpp
@@ -2,6 +2,11 @@
 #include "WallClockMusicClockDriver.h"
 #include "Harmonix.h"
 
+TOptional<FVector> FWallClockMusicClockDriver::TryGetAudioSourceLocation() const
+{
+	return {};
+}
+
 void FWallClockMusicClockDriver::Disconnect()
 {
 	TempoMapMidi = nullptr;
diff --git a/Engine/Plugins/Runtime/Harmonix/Source/HarmonixMetasound/Private/HarmonixMetasound/Components/WallClockMusicClockDriver.h b/Engine/Plugins/Runtime/Harmonix/Source/HarmonixMetasound/Private/HarmonixMetasound/Components/WallClockMusicClockDriver.h
--- a/Engine/Plugins/Runtime/Harmonix/Source/HarmonixMetasound/Private/HarmonixMetasound/Components/WallClockMusicClockDriver.h
+++ b/Engine/Plugins/Runtime/Harmonix/Source/HarmonixMetasound/Private/HarmonixMetasound/Components/WallClockMusicClockDriver.h
@@ -12,6 +12,8 @@ public:
 		, TempoMapMidi(Settings.TempoMap)
 	{}
 
+	virtual TOptional<FVector> TryGetAudioSourceLocation() const override;
+
 	virtual void Disconnect() override;
 	virtual bool RefreshCurrentSongPos() override;
 	virtual void OnStart() override;
diff --git a/Engine/Plugins/Runtime/Harmonix/Source/HarmonixMetasound/Public/HarmonixMetasound/Components/MusicClockComponent.h b/Engine/Plugins/Runtime/Harmonix/Source/HarmonixMetasound/Public/HarmonixMetasound/Components/MusicClockComponent.h
--- a/Engine/Plugins/Runtime/Harmonix/Source/HarmonixMetasound/Public/HarmonixMetasound/Components/MusicClockComponent.h
+++ b/Engine/Plugins/Runtime/Harmonix/Source/HarmonixMetasound/Public/HarmonixMetasound/Components/MusicClockComponent.h
@@ -312,6 +312,9 @@ public:
 	UFUNCTION(BlueprintPure, Category = "MusicClock")
 	UE_API bool LoopedThisFrame(ECalibratedMusicTimebase Timebase = ECalibratedMusicTimebase::VideoRenderTime) const;
 
+	UFUNCTION(BlueprintPure, Category = "MusicClock")
+	UE_API TOptional<FVector> TryGetAudioSourceLocation() const;
+
 	UFUNCTION(BlueprintPure, Category = "Count In")
 	UE_API float GetCountInSeconds() const;
 
diff --git a/Engine/Plugins/Runtime/Harmonix/Source/HarmonixMetasound/Public/HarmonixMetasound/Components/MusicClockDriverBase.h b/Engine/Plugins/Runtime/Harmonix/Source/HarmonixMetasound/Public/HarmonixMetasound/Components/MusicClockDriverBase.h
--- a/Engine/Plugins/Runtime/Harmonix/Source/HarmonixMetasound/Public/HarmonixMetasound/Components/MusicClockDriverBase.h
+++ b/Engine/Plugins/Runtime/Harmonix/Source/HarmonixMetasound/Public/HarmonixMetasound/Components/MusicClockDriverBase.h
@@ -18,6 +18,8 @@ public:
 	UE_API FMusicClockDriverBase(UObject* WorldContextObject, const FMusicClockSettingsBase& Settings);
 	UE_API virtual ~FMusicClockDriverBase();
 
+	virtual TOptional<FVector> TryGetAudioSourceLocation() const = 0;
+
 	UE_API void EnsureClockIsValidForGameFrame();
 
 	UE_API virtual bool CalculateSongPosWithOffset(float MsOffset, ECalibratedMusicTimebase Timebase, FMidiSongPos& OutResult) const;
@@ -69,7 +71,7 @@ public:
 	// TODO shouldn't these deltas be rolled into the MidiSongPos?
 	float DeltaBarF[(int32)ECalibratedMusicTimebase::Count];
 	float DeltaBeatF[(int32)ECalibratedMusicTimebase::Count];
-	
+
 protected:
 	
 	virtual void OnStart() = 0;
