[Harmonix] Remove MusicClockComponent tick. The component now directly broadcasts changed events after updating the clock.

Renamed instances of "EnsureClockIsValid" to "UpdateClock" to better describe what the function is doing.

#tests PIE - test level
#rnx

diff --git a/Engine/Plugins/Runtime/Harmonix/Source/HarmonixMetasound/Private/HarmonixMetasound/Components/MusicClockComponent.cpp b/Engine/Plugins/Runtime/Harmonix/Source/HarmonixMetasound/Private/HarmonixMetasound/Components/MusicClockComponent.cpp
--- a/Engine/Plugins/Runtime/Harmonix/Source/HarmonixMetasound/Private/HarmonixMetasound/Components/MusicClockComponent.cpp
+++ b/Engine/Plugins/Runtime/Harmonix/Source/HarmonixMetasound/Private/HarmonixMetasound/Components/MusicClockComponent.cpp
@@ -16,7 +16,7 @@ DEFINE_LOG_CATEGORY(LogMusicClock)
 UMusicClockComponent::UMusicClockComponent()
 {
 	MakeDefaultSongMap();
-	PrimaryComponentTick.bCanEverTick = true;
+	PrimaryComponentTick.bCanEverTick = false;
 	bAutoActivate = false;
 }
 
@@ -147,11 +147,16 @@ void UMusicClockComponent::SetDefaultTimeSignatureDenom(int Denom)
 	DefaultTimeSignatureDenom = Denom;
 }
 
-void UMusicClockComponent::EnsureClockIsValidForGameFrameFromSubsystem()
+void UMusicClockComponent::UpdateClock()
 {
 	if (ClockDriver)
 	{
-		ClockDriver->EnsureClockIsValidForGameFrame();
+		ClockDriver->UpdateClockForGameFrame();
+
+		BroadcastSongPosChanges(TimebaseForBarAndBeatEvents);
+		BroadcastSeekLoopDetections(ECalibratedMusicTimebase::AudioRenderTime);
+		BroadcastSeekLoopDetections(ECalibratedMusicTimebase::ExperiencedTime);
+		BroadcastSeekLoopDetections(ECalibratedMusicTimebase::VideoRenderTime);
 	}
 }
 
@@ -371,23 +376,6 @@ void UMusicClockComponent::BeginDestroy()
 	DisconnectFromClockDriver();
 }
 
-void UMusicClockComponent::TickComponent(float DeltaTime, enum ELevelTick TickType, FActorComponentTickFunction* ThisTickFunction)
-{
-	Super::TickComponent(DeltaTime, TickType, ThisTickFunction);
-	TickComponentInternal();
-}
-
-void UMusicClockComponent::TickComponentInternal()
-{
-	if (GetState() == EMusicClockState::Running)
-	{
-		BroadcastSongPosChanges(TimebaseForBarAndBeatEvents);
-		BroadcastSeekLoopDetections(ECalibratedMusicTimebase::AudioRenderTime);
-		BroadcastSeekLoopDetections(ECalibratedMusicTimebase::ExperiencedTime);
-		BroadcastSeekLoopDetections(ECalibratedMusicTimebase::VideoRenderTime);
-	}
-}
-
 void UMusicClockComponent::SetTempoMapForWallClock(UMidiFile* InTempoMap)
 {
 	ensure(GetState() == EMusicClockState::Stopped);
diff --git a/Engine/Plugins/Runtime/Harmonix/Source/HarmonixMetasound/Private/HarmonixMetasound/Components/MusicClockDriverBase.cpp b/Engine/Plugins/Runtime/Harmonix/Source/HarmonixMetasound/Private/HarmonixMetasound/Components/MusicClockDriverBase.cpp
--- a/Engine/Plugins/Runtime/Harmonix/Source/HarmonixMetasound/Private/HarmonixMetasound/Components/MusicClockDriverBase.cpp
+++ b/Engine/Plugins/Runtime/Harmonix/Source/HarmonixMetasound/Private/HarmonixMetasound/Components/MusicClockDriverBase.cpp
@@ -168,7 +168,7 @@ float FMusicClockDriverBase::GetBeatsPerMinute() const
 	return Tempo * TimeSignatureDenom / 4.0f; 
 }
 
-void FMusicClockDriverBase::EnsureClockIsValidForGameFrame()
+void FMusicClockDriverBase::UpdateClockForGameFrame()
 {
 	//	Not for use outside the game thread.
 	if (ensureMsgf(
diff --git a/Engine/Plugins/Runtime/Harmonix/Source/HarmonixMetasound/Private/HarmonixMetasound/DataTypes/HarmonixMusicHandle.cpp b/Engine/Plugins/Runtime/Harmonix/Source/HarmonixMetasound/Private/HarmonixMetasound/DataTypes/HarmonixMusicHandle.cpp
--- a/Engine/Plugins/Runtime/Harmonix/Source/HarmonixMetasound/Private/HarmonixMetasound/DataTypes/HarmonixMusicHandle.cpp
+++ b/Engine/Plugins/Runtime/Harmonix/Source/HarmonixMetasound/Private/HarmonixMetasound/DataTypes/HarmonixMusicHandle.cpp
@@ -101,10 +101,7 @@ bool UHarmonixMusicHandle::Play(float FromSeconds)
 
 void UHarmonixMusicHandle::Tick(float DeltaSeconds)
 {
-	if (MusicClockComponent)
-	{
-		MusicClockComponent->TickComponentInternal();
-	}
+	// MusicClock no longer ticks, this function is empty
 }
 
 void UHarmonixMusicHandle::Pause()
diff --git a/Engine/Plugins/Runtime/Harmonix/Source/HarmonixMetasound/Private/HarmonixMetasound/MusicEnvironmentSupport/HarmonixMusicEnvironmentMetronome.cpp b/Engine/Plugins/Runtime/Harmonix/Source/HarmonixMetasound/Private/HarmonixMetasound/MusicEnvironmentSupport/HarmonixMusicEnvironmentMetronome.cpp
--- a/Engine/Plugins/Runtime/Harmonix/Source/HarmonixMetasound/Private/HarmonixMetasound/MusicEnvironmentSupport/HarmonixMusicEnvironmentMetronome.cpp
+++ b/Engine/Plugins/Runtime/Harmonix/Source/HarmonixMetasound/Private/HarmonixMetasound/MusicEnvironmentSupport/HarmonixMusicEnvironmentMetronome.cpp
@@ -34,10 +34,7 @@ bool UHarmonixMusicEnvironmentMetronome::Initialize(UWorld* InWorld)
 
 void UHarmonixMusicEnvironmentMetronome::Tick(float DeltaSecs)
 {
-	if (MusicClockComponent)
-	{
-		MusicClockComponent->TickComponentInternal();
-	}
+	// MusicClock no longer ticks, this function is empty
 }
 
 void UHarmonixMusicEnvironmentMetronome::Start(double FromSeconds)
diff --git a/Engine/Plugins/Runtime/Harmonix/Source/HarmonixMetasound/Private/HarmonixMetasound/Subsystems/MidiClockUpdateSubsystem.cpp b/Engine/Plugins/Runtime/Harmonix/Source/HarmonixMetasound/Private/HarmonixMetasound/Subsystems/MidiClockUpdateSubsystem.cpp
--- a/Engine/Plugins/Runtime/Harmonix/Source/HarmonixMetasound/Private/HarmonixMetasound/Subsystems/MidiClockUpdateSubsystem.cpp
+++ b/Engine/Plugins/Runtime/Harmonix/Source/HarmonixMetasound/Private/HarmonixMetasound/Subsystems/MidiClockUpdateSubsystem.cpp
@@ -200,7 +200,7 @@ void UMidiClockUpdateSubsystem::UpdateUMusicClockComponents()
 	{
 		if (UMusicClockComponent* MusicClockPtr = ClockIterator->Get())
 		{
-			MusicClockPtr->EnsureClockIsValidForGameFrameFromSubsystem();
+			MusicClockPtr->UpdateClock();
 		}
 		else
 		{
diff --git a/Engine/Plugins/Runtime/Harmonix/Source/HarmonixMetasound/Public/HarmonixMetasound/Components/MusicClockComponent.h b/Engine/Plugins/Runtime/Harmonix/Source/HarmonixMetasound/Public/HarmonixMetasound/Components/MusicClockComponent.h
--- a/Engine/Plugins/Runtime/Harmonix/Source/HarmonixMetasound/Public/HarmonixMetasound/Components/MusicClockComponent.h
+++ b/Engine/Plugins/Runtime/Harmonix/Source/HarmonixMetasound/Public/HarmonixMetasound/Components/MusicClockComponent.h
@@ -512,10 +512,6 @@ public:
 	UE_API virtual FMusicalTime Quantize(const FMusicalTime& MusicalTime, int32 QuantizationInterval, UFrameBasedMusicMap::EQuantizeDirection Direction = UFrameBasedMusicMap::EQuantizeDirection::Nearest) const;
 	virtual bool CanAuditionInEditor() const override { return false; }
 	// ^^^ END IMusicEnvironmentClockSource Implementation ^^^
-
-public:
-	UE_API virtual void TickComponent(float DeltaTime, enum ELevelTick TickType, FActorComponentTickFunction* ThisTickFunction);
-	UE_API void TickComponentInternal();
 	
 private:
 	friend class  UMidiClockUpdateSubsystem;
@@ -538,7 +534,7 @@ private:
 	UE_API void ConnectToCustomClock();
 	UE_API void DisconnectFromClockDriver();
 	
-	UE_API void EnsureClockIsValidForGameFrameFromSubsystem();
+	UE_API void UpdateClock();
 	UE_API const FMidiSongPos& GetCurrentSongPosInternal(ECalibratedMusicTimebase Timebase) const;
 	UE_API const FMidiSongPos& GetPreviousSongPosInternal(ECalibratedMusicTimebase Timebase) const;
 };
diff --git a/Engine/Plugins/Runtime/Harmonix/Source/HarmonixMetasound/Public/HarmonixMetasound/Components/MusicClockDriverBase.h b/Engine/Plugins/Runtime/Harmonix/Source/HarmonixMetasound/Public/HarmonixMetasound/Components/MusicClockDriverBase.h
--- a/Engine/Plugins/Runtime/Harmonix/Source/HarmonixMetasound/Public/HarmonixMetasound/Components/MusicClockDriverBase.h
+++ b/Engine/Plugins/Runtime/Harmonix/Source/HarmonixMetasound/Public/HarmonixMetasound/Components/MusicClockDriverBase.h
@@ -20,7 +20,7 @@ public:
 
 	virtual TOptional<FVector> TryGetAudioSourceLocation() const = 0;
 
-	UE_API void EnsureClockIsValidForGameFrame();
+	UE_API void UpdateClockForGameFrame();
 
 	UE_API virtual bool CalculateSongPosWithOffset(float MsOffset, ECalibratedMusicTimebase Timebase, FMidiSongPos& OutResult) const;
 
