[Harmonix] FFTAnalyzer: mute static analyzer warning, divide-by-zero

#rnx

diff --git a/Engine/Plugins/Runtime/Harmonix/Source/HarmonixDsp/Private/HarmonixDsp/AudioAnalysis/FFTAnalyzer.cpp b/Engine/Plugins/Runtime/Harmonix/Source/HarmonixDsp/Private/HarmonixDsp/AudioAnalysis/FFTAnalyzer.cpp
--- a/Engine/Plugins/Runtime/Harmonix/Source/HarmonixDsp/Private/HarmonixDsp/AudioAnalysis/FFTAnalyzer.cpp
+++ b/Engine/Plugins/Runtime/Harmonix/Source/HarmonixDsp/Private/HarmonixDsp/AudioAnalysis/FFTAnalyzer.cpp
@@ -184,7 +184,7 @@ namespace Harmonix::Dsp::AudioAnalysis
 		for (int32 DstBinIdx = 0; DstBinIdx <= NumResultBins; DstBinIdx++)
 		{
 			// compute the desired frequency at this bin boundary
-			float DstFrac = static_cast<float>(DstBinIdx) / NumResultBins;
+			float DstFrac = static_cast<float>(DstBinIdx) / NumResultBins; //-V609
 			float Freq;
 			if (InSettings.MelScaleBinning)
 			{
