[MusicClockComponent] Update the static Create methods to more closely match CreateSound2D for better safety, especially on level transitions.

#rnx

diff --git a/Engine/Plugins/Runtime/Harmonix/Source/HarmonixMetasound/Private/HarmonixMetasound/Components/MusicClockComponent.cpp b/Engine/Plugins/Runtime/Harmonix/Source/HarmonixMetasound/Private/HarmonixMetasound/Components/MusicClockComponent.cpp
--- a/Engine/Plugins/Runtime/Harmonix/Source/HarmonixMetasound/Private/HarmonixMetasound/Components/MusicClockComponent.cpp
+++ b/Engine/Plugins/Runtime/Harmonix/Source/HarmonixMetasound/Private/HarmonixMetasound/Components/MusicClockComponent.cpp
@@ -20,19 +20,44 @@ UMusicClockComponent::UMusicClockComponent()
 	bAutoActivate = false;
 }
 
-UMusicClockComponent* UMusicClockComponent::CreateMetasoundDrivenMusicClock(UObject* WorldContextObject, UAudioComponent* InAudioComponent, FName MetasoundOutputPinName, bool Start)
+TNotNull<UMusicClockComponent*> UMusicClockComponent::CreateMusicClockShared(UObject* WorldContextObject, bool bPersistAcrossLevelTransition)
 {
-	UMusicClockComponent* NewClock = NewObject<UMusicClockComponent>(WorldContextObject);
+	// Mirroring UGameplayStatics::CreateSound2D / FAudioDevice::CreateComponent:
+	// - If we have an owning actor, register with them, use them as the Outer
+	// - If not, register with the world, use the transient package as the Outer
+	// - But if we want to persist through level transitions, don't register at all
+
+	AActor* OwningActor = nullptr;
+	UWorld* World = nullptr;
+	if (!bPersistAcrossLevelTransition && WorldContextObject)
+	{
+		OwningActor = Cast<AActor>(WorldContextObject);
+		if (!OwningActor)
+		{
+			OwningActor = WorldContextObject->GetTypedOuter<AActor>();
+		}
+
+		World = WorldContextObject->GetWorld();
+	}
+
+	UObject* OuterToUse = OwningActor ? OwningActor : GetTransientPackageAsObject();
+	TNotNull<UMusicClockComponent*> NewClock = NewObject<UMusicClockComponent>(OuterToUse);
 
 	if (NewClock->GetOwner())
 	{
 		NewClock->RegisterComponent();
 	}
-	else if (UWorld* World = WorldContextObject ? WorldContextObject->GetWorld() : nullptr)
+	else if (World)
 	{
 		NewClock->RegisterComponentWithWorld(World);
 	}
 
+	return NewClock;
+}
+
+UMusicClockComponent* UMusicClockComponent::CreateMetasoundDrivenMusicClock(UObject* WorldContextObject, UAudioComponent* InAudioComponent, FName MetasoundOutputPinName, bool Start, bool bPersistAcrossLevelTransition)
+{
+	TNotNull<UMusicClockComponent*> NewClock = CreateMusicClockShared(WorldContextObject, bPersistAcrossLevelTransition);
 	NewClock->DriveMethod = EMusicClockDriveMethod::MetaSound;
 	NewClock->MetasoundOutputName = MetasoundOutputPinName;
 	NewClock->ConnectToMetasoundOnAudioComponent(InAudioComponent);
@@ -43,19 +68,9 @@ UMusicClockComponent* UMusicClockComponent::CreateMetasoundDrivenMusicClock(UObj
 	return NewClock;
 }
 
-UMusicClockComponent* UMusicClockComponent::CreateWallClockDrivenMusicClock(UObject* WorldContextObject, UMidiFile* InTempoMap, bool Start)
+UMusicClockComponent* UMusicClockComponent::CreateWallClockDrivenMusicClock(UObject* WorldContextObject, UMidiFile* InTempoMap, bool Start, bool bPersistAcrossLevelTransition)
 {
-	UMusicClockComponent* NewClock = NewObject<UMusicClockComponent>(WorldContextObject);
-
-	if (NewClock->GetOwner())
-	{
-		NewClock->RegisterComponent();
-	}
-	else if (UWorld* World = WorldContextObject ? WorldContextObject->GetWorld() : nullptr)
-	{
-		NewClock->RegisterComponentWithWorld(World);
-	}
-
+	TNotNull<UMusicClockComponent*> NewClock = CreateMusicClockShared(WorldContextObject, bPersistAcrossLevelTransition);
 	NewClock->ConnectToWallClockForMidi(InTempoMap);
 	if (Start)
 	{
@@ -64,18 +79,9 @@ UMusicClockComponent* UMusicClockComponent::CreateWallClockDrivenMusicClock(UObj
 	return NewClock;
 }
 
-UMusicClockComponent* UMusicClockComponent::CreateMusicClockWithSettings(UObject* WorldContextObject, TInstancedStruct<FMusicClockSettingsBase> MusicClockSettings, bool Start)
+UMusicClockComponent* UMusicClockComponent::CreateMusicClockWithSettings(UObject* WorldContextObject, TInstancedStruct<FMusicClockSettingsBase> MusicClockSettings, bool Start, bool bPersistAcrossLevelTransition)
 {
-	UMusicClockComponent* NewClock = NewObject<UMusicClockComponent>(WorldContextObject);
-
-	if (NewClock->GetOwner())
-	{
-		NewClock->RegisterComponent();
-	}
-	else if (UWorld* World = WorldContextObject ? WorldContextObject->GetWorld() : nullptr)
-	{
-		NewClock->RegisterComponentWithWorld(World);
-	}
+	TNotNull<UMusicClockComponent*> NewClock = CreateMusicClockShared(WorldContextObject, bPersistAcrossLevelTransition);
 
 	if (const FMetasoundMusicClockSettings* MetasoundClockSettings = MusicClockSettings.GetPtr<FMetasoundMusicClockSettings>())
 	{
diff --git a/Engine/Plugins/Runtime/Harmonix/Source/HarmonixMetasound/Public/HarmonixMetasound/Components/MusicClockComponent.h b/Engine/Plugins/Runtime/Harmonix/Source/HarmonixMetasound/Public/HarmonixMetasound/Components/MusicClockComponent.h
--- a/Engine/Plugins/Runtime/Harmonix/Source/HarmonixMetasound/Public/HarmonixMetasound/Components/MusicClockComponent.h
+++ b/Engine/Plugins/Runtime/Harmonix/Source/HarmonixMetasound/Public/HarmonixMetasound/Components/MusicClockComponent.h
@@ -471,17 +471,17 @@ public:
 	UE_API float MeasureSpanProgress(const FMusicalTimeSpan& Span, ECalibratedMusicTimebase Timebase = ECalibratedMusicTimebase::VideoRenderTime) const;
 
 	UFUNCTION(BlueprintCallable, Category = "Audio|MusicClock", meta = (WorldContext = "WorldContextObject", DefaultToSelf = "WorldContextObject"))
-	static UE_API UMusicClockComponent* CreateMetasoundDrivenMusicClock(UObject* WorldContextObject, UAudioComponent* InAudioComponent, FName MetasoundOuputPinName = "MIDI Clock", bool Start = true);
+	static UE_API UMusicClockComponent* CreateMetasoundDrivenMusicClock(UObject* WorldContextObject, UAudioComponent* InAudioComponent, FName MetasoundOuputPinName = "MIDI Clock", bool Start = true, bool bPersistAcrossLevelTransition = false);
 
 	UFUNCTION(BlueprintCallable, Category = "Audio|MusicClock", meta = (WorldContext = "WorldContextObject", DefaultToSelf = "WorldContextObject"))
-	static UE_API UMusicClockComponent* CreateWallClockDrivenMusicClock(UObject* WorldContextObject, UMidiFile* WithTempoMap, bool Start = true);
+	static UE_API UMusicClockComponent* CreateWallClockDrivenMusicClock(UObject* WorldContextObject, UMidiFile* WithTempoMap, bool Start = true, bool bPersistAcrossLevelTransition = false);
 
-	static UE_API UMusicClockComponent* CreateMusicClockWithSettings(UObject* WorldContextObject, TInstancedStruct<FMusicClockSettingsBase> MusicClockSettings, bool Start = true);
+	static UE_API UMusicClockComponent* CreateMusicClockWithSettings(UObject* WorldContextObject, TInstancedStruct<FMusicClockSettingsBase> MusicClockSettings, bool Start = true, bool bPersistAcrossLevelTransition = false);
 
 	template<typename MusicClockSettingsType = FMusicClockSettingsBase>
-	static UMusicClockComponent* CreateMusicClockWithSettings(UObject* WorldContextObject, const MusicClockSettingsType& InMusicClockSettings, bool Start = true)
+	static UMusicClockComponent* CreateMusicClockWithSettings(UObject* WorldContextObject, const MusicClockSettingsType& InMusicClockSettings, bool Start = true, bool bPersistAcrossLevelTransition = false)
 	{
-		return CreateMusicClockWithSettings(WorldContextObject, TInstancedStruct<FMusicClockSettingsBase>::Make(InMusicClockSettings), Start);
+		return CreateMusicClockWithSettings(WorldContextObject, TInstancedStruct<FMusicClockSettingsBase>::Make(InMusicClockSettings), Start, bPersistAcrossLevelTransition);
 	}
 
 	UFUNCTION(BlueprintCallable, Category = "MusicClock")
@@ -524,6 +524,8 @@ private:
 	int32 LastBroadcastBeat = -1;
 	FSongSection LastBroadcastSongSection;
 	
+	static TNotNull<UMusicClockComponent*> CreateMusicClockShared(UObject* WorldContextObject, bool bPersistAcrossLevelTransition);
+	
 	UE_API void CreateClockDriver();
 	UE_API void BroadcastSongPosChanges(ECalibratedMusicTimebase Timebase);
 	UE_API void BroadcastSeekLoopDetections(ECalibratedMusicTimebase Timebase) const;
