[Harmonix] FusionSampler stores "NodePtr" instead of "RawPointer" to avoid unsafe pointer use

#rnx

diff --git a/Engine/Plugins/Runtime/Harmonix/Source/HarmonixDsp/Private/HarmonixDsp/FusionSampler/FusionSampler.cpp b/Engine/Plugins/Runtime/Harmonix/Source/HarmonixDsp/Private/HarmonixDsp/FusionSampler/FusionSampler.cpp
--- a/Engine/Plugins/Runtime/Harmonix/Source/HarmonixDsp/Private/HarmonixDsp/FusionSampler/FusionSampler.cpp
+++ b/Engine/Plugins/Runtime/Harmonix/Source/HarmonixDsp/Private/HarmonixDsp/FusionSampler/FusionSampler.cpp
@@ -674,7 +674,7 @@ void FFusionSampler::AllNotesOffWithFrameOffset(int32 InNumFrames)
 	NoteOffWithFrameOffset(FMidiVoiceId::Any(), 0, 0, InNumFrames);
 }
 
-void FFusionSampler::SetPatch(FFusionPatchData* PatchData)
+void FFusionSampler::SetPatch(FFusionPatchDataProxy::NodePtr PatchData)
 {
 	if (PatchData == FusionPatchData)
 	{
diff --git a/Engine/Plugins/Runtime/Harmonix/Source/HarmonixDsp/Public/HarmonixDsp/FusionSampler/FusionSampler.h b/Engine/Plugins/Runtime/Harmonix/Source/HarmonixDsp/Public/HarmonixDsp/FusionSampler/FusionSampler.h
--- a/Engine/Plugins/Runtime/Harmonix/Source/HarmonixDsp/Public/HarmonixDsp/FusionSampler/FusionSampler.h
+++ b/Engine/Plugins/Runtime/Harmonix/Source/HarmonixDsp/Public/HarmonixDsp/FusionSampler/FusionSampler.h
@@ -140,7 +140,7 @@ protected:
 	UE_API virtual void Process(uint32 InSliceIndex, uint32 InSubSliceIndex, TAudioBuffer<float>& OutBuffer) override;
 	bool CanProcessFromWorkerThread() const override { return true; }
 
-	UE_API void SetPatch(FFusionPatchData* PatchData);
+	UE_API void SetPatch(FFusionPatchDataProxy::NodePtr PatchData);
 
 	UE_API FString GetPatchPath() const;
 
@@ -340,7 +340,7 @@ private:
 	TArray<FPendingNoteAction> PendingNoteActions;
 	FMIDINoteStatus NoteStatus[Harmonix::Midi::Constants::GMaxNumNotes];
 
-	struct FFusionPatchData* FusionPatchData = nullptr;
+	FFusionPatchDataProxy::NodePtr FusionPatchData = nullptr;
 
 	int8 LastStartLayerSelect[Harmonix::Midi::Constants::GMaxNumNotes];
 	int8 LastStopLayerSelect[Harmonix::Midi::Constants::GMaxNumNotes];
