ISongMapEvaluator was modified to expose FractionalBeatIncludingCountInToTick

#tests Verified that I can now access FractionalBeatIncludingCountInToTick

#Synced-CL 48226943

diff --git a/Engine/Plugins/Runtime/Harmonix/Source/HarmonixMidi/Private/HarmonixMidi/SongMaps.cpp b/Engine/Plugins/Runtime/Harmonix/Source/HarmonixMidi/Private/HarmonixMidi/SongMaps.cpp
--- a/Engine/Plugins/Runtime/Harmonix/Source/HarmonixMidi/Private/HarmonixMidi/SongMaps.cpp
+++ b/Engine/Plugins/Runtime/Harmonix/Source/HarmonixMidi/Private/HarmonixMidi/SongMaps.cpp
@@ -1062,6 +1062,11 @@ float ISongMapEvaluator::FractionalBarIncludingCountInToTick(float FractionalBar
 	return GetBarMap().FractionalBarIncludingCountInToTick(FractionalBarIndex);
 }
 
+float ISongMapEvaluator::FractionalBeatIncludingCountInToTick(float FractionalBeatIndex) const
+{
+	return GetBarMap().FractionalBeatIncludingCountInToTick(FractionalBeatIndex);
+}
+
 int32 ISongMapEvaluator::TickToBarIncludingCountIn(int32 Tick) const
 {
 	return GetBarMap().TickToBarIncludingCountIn(Tick);
diff --git a/Engine/Plugins/Runtime/Harmonix/Source/HarmonixMidi/Public/HarmonixMidi/SongMaps.h b/Engine/Plugins/Runtime/Harmonix/Source/HarmonixMidi/Public/HarmonixMidi/SongMaps.h
--- a/Engine/Plugins/Runtime/Harmonix/Source/HarmonixMidi/Public/HarmonixMidi/SongMaps.h
+++ b/Engine/Plugins/Runtime/Harmonix/Source/HarmonixMidi/Public/HarmonixMidi/SongMaps.h
@@ -135,6 +135,7 @@ public:
 	UE_API int32 BarIncludingCountInToTick(int32 BarIndex, int32* OutBeatsPerBar = nullptr, int32* OutTicksPerBeat = nullptr) const;
 	UE_API int32 BarBeatTickIncludingCountInToTick(int32 BarIndex, int32 BeatInBar, int32 TickInBeat) const;
 	UE_API float FractionalBarIncludingCountInToTick(float FractionalBarIndex) const;
+	UE_API float FractionalBeatIncludingCountInToTick(float FractionalBeatIndex) const;
 	UE_API int32 TickToBarIncludingCountIn(int32 Tick) const;
 	UE_API float TickToFractionalBarIncludingCountIn(float Tick) const;
 	UE_API void TickToBarBeatTickIncludingCountIn(int32 RawTick, int32& OutBarIndex, int32& OutBeatInBar, int32& OutTickIndexInBeat, int32* OutBeatsPerBar = nullptr, int32* OutTicksPerBeat = nullptr) const;
