[Harmonix] Add DisplayName to MusicClock to show the name of the source.  For debugging.

#tests PIE jam

#Synced-CL 48426002

diff --git a/Engine/Plugins/Runtime/Harmonix/Source/HarmonixMetasound/Private/HarmonixMetasound/Components/MetasoundMusicClockDriver.cpp b/Engine/Plugins/Runtime/Harmonix/Source/HarmonixMetasound/Private/HarmonixMetasound/Components/MetasoundMusicClockDriver.cpp
--- a/Engine/Plugins/Runtime/Harmonix/Source/HarmonixMetasound/Private/HarmonixMetasound/Components/MetasoundMusicClockDriver.cpp
+++ b/Engine/Plugins/Runtime/Harmonix/Source/HarmonixMetasound/Private/HarmonixMetasound/Components/MetasoundMusicClockDriver.cpp
@@ -332,6 +332,22 @@ bool FMetasoundMusicClockDriver::SeekedThisFrame(ECalibratedMusicTimebase Timeba
 	}
 }
 
+#if !UE_BUILD_SHIPPING
+FString FMetasoundMusicClockDriver::GetDisplayName() const
+{
+#if WITH_EDITORONLY_DATA
+	// that method is only good WITH_EDITORONLY_DATA
+	if (UAudioComponent* Component = AudioComponentToWatch.Get())
+	{
+		if (UMetaSoundSource* Source = Cast<UMetaSoundSource>(Component->GetSound()))
+		{
+			return FString::Printf(TEXT("Metasound: %s"), *Source->GetDisplayName().ToString());
+		}
+	}
+#endif
+	return TEXT("Metasound");
+}
+#endif
 
 bool FMetasoundMusicClockDriver::ConnectToAudioComponentsMetasound(
 	UAudioComponent* InAudioComponent,
diff --git a/Engine/Plugins/Runtime/Harmonix/Source/HarmonixMetasound/Private/HarmonixMetasound/Components/MetasoundMusicClockDriver.h b/Engine/Plugins/Runtime/Harmonix/Source/HarmonixMetasound/Private/HarmonixMetasound/Components/MetasoundMusicClockDriver.h
--- a/Engine/Plugins/Runtime/Harmonix/Source/HarmonixMetasound/Private/HarmonixMetasound/Components/MetasoundMusicClockDriver.h
+++ b/Engine/Plugins/Runtime/Harmonix/Source/HarmonixMetasound/Private/HarmonixMetasound/Components/MetasoundMusicClockDriver.h
@@ -31,7 +31,11 @@ public:
 	
 	virtual bool LoopedThisFrame(ECalibratedMusicTimebase Timebase) const override;
 	virtual bool SeekedThisFrame(ECalibratedMusicTimebase Timebase) const override;
-	
+
+#if !UE_BUILD_SHIPPING
+	virtual FString GetDisplayName() const override;
+#endif
+
 	bool ConnectToAudioComponentsMetasound(
 		UAudioComponent* InAudioComponent,
 		FName MetasoundOuputPinName = "MIDI Clock",
diff --git a/Engine/Plugins/Runtime/Harmonix/Source/HarmonixMetasound/Private/HarmonixMetasound/Components/MusicClockComponent.cpp b/Engine/Plugins/Runtime/Harmonix/Source/HarmonixMetasound/Private/HarmonixMetasound/Components/MusicClockComponent.cpp
--- a/Engine/Plugins/Runtime/Harmonix/Source/HarmonixMetasound/Private/HarmonixMetasound/Components/MusicClockComponent.cpp
+++ b/Engine/Plugins/Runtime/Harmonix/Source/HarmonixMetasound/Private/HarmonixMetasound/Components/MusicClockComponent.cpp
@@ -755,6 +755,26 @@ const ISongMapEvaluator& UMusicClockComponent::GetSongMaps() const
 	return SongMaps ? *SongMaps : DefaultMaps;
 }
 
+#if !UE_BUILD_SHIPPING
+FString UMusicClockComponent::GetDisplayName() const
+{
+	if (!DisplayName.IsEmpty())
+	{
+		return DisplayName;
+	}
+	if (ClockDriver)
+	{
+		return ClockDriver->GetDisplayName();
+	}
+	return FString::Printf(TEXT("Driverless: %s %s"), *GetNameSafe(this), *GetNameSafe(GetOwner()));
+}
+
+void UMusicClockComponent::SetDisplayName(FString Name)
+{
+	DisplayName = Name;
+}
+#endif
+
 const FMidiSongPos& UMusicClockComponent::GetSongPos(ECalibratedMusicTimebase Timebase) const
 {
 	return GetCurrentSongPosInternal(Timebase);
diff --git a/Engine/Plugins/Runtime/Harmonix/Source/HarmonixMetasound/Private/HarmonixMetasound/Components/WallClockMusicClockDriver.cpp b/Engine/Plugins/Runtime/Harmonix/Source/HarmonixMetasound/Private/HarmonixMetasound/Components/WallClockMusicClockDriver.cpp
--- a/Engine/Plugins/Runtime/Harmonix/Source/HarmonixMetasound/Private/HarmonixMetasound/Components/WallClockMusicClockDriver.cpp
+++ b/Engine/Plugins/Runtime/Harmonix/Source/HarmonixMetasound/Private/HarmonixMetasound/Components/WallClockMusicClockDriver.cpp
@@ -63,3 +63,10 @@ const ISongMapEvaluator* FWallClockMusicClockDriver::GetCurrentSongMapEvaluator(
 	}
 	return &DefaultMaps;
 }
+
+#if !UE_BUILD_SHIPPING
+FString FWallClockMusicClockDriver::GetDisplayName() const
+{
+	return FString::Printf(TEXT("Wall: %.3f bpm"), GetBeatsPerMinute());
+}
+#endif
diff --git a/Engine/Plugins/Runtime/Harmonix/Source/HarmonixMetasound/Private/HarmonixMetasound/Components/WallClockMusicClockDriver.h b/Engine/Plugins/Runtime/Harmonix/Source/HarmonixMetasound/Private/HarmonixMetasound/Components/WallClockMusicClockDriver.h
--- a/Engine/Plugins/Runtime/Harmonix/Source/HarmonixMetasound/Private/HarmonixMetasound/Components/WallClockMusicClockDriver.h
+++ b/Engine/Plugins/Runtime/Harmonix/Source/HarmonixMetasound/Private/HarmonixMetasound/Components/WallClockMusicClockDriver.h
@@ -2,7 +2,6 @@
 #pragma once
 #include "HarmonixMetasound/Components/MusicClockDriverBase.h"
 #include "HarmonixMidi/MidiFile.h"
-#include "UObject/StrongObjectPtr.h"
 
 class FWallClockMusicClockDriver : public FMusicClockDriverBase
 {
@@ -25,6 +24,10 @@ public:
 	virtual bool LoopedThisFrame(ECalibratedMusicTimebase Timebase) const override { return false; }
 	virtual bool SeekedThisFrame(ECalibratedMusicTimebase Timebase) const override { return false; }
 
+#if !UE_BUILD_SHIPPING
+	virtual FString GetDisplayName() const override;
+#endif
+
 private:
 	TWeakObjectPtr<UMidiFile> TempoMapMidi;
 
diff --git a/Engine/Plugins/Runtime/Harmonix/Source/HarmonixMetasound/Public/HarmonixMetasound/Components/MusicClockComponent.h b/Engine/Plugins/Runtime/Harmonix/Source/HarmonixMetasound/Public/HarmonixMetasound/Components/MusicClockComponent.h
--- a/Engine/Plugins/Runtime/Harmonix/Source/HarmonixMetasound/Public/HarmonixMetasound/Components/MusicClockComponent.h
+++ b/Engine/Plugins/Runtime/Harmonix/Source/HarmonixMetasound/Public/HarmonixMetasound/Components/MusicClockComponent.h
@@ -365,6 +365,14 @@ public:
 
 	UE_API const ISongMapEvaluator& GetSongMaps() const;
 
+#if !UE_BUILD_SHIPPING
+	// A short string description of the clock, for debugging, like song name or animation name or whatever is driving it.
+	UE_API FString GetDisplayName() const;
+
+	// Sets the DisplayName. Will be used by GetDisplayName if set, otherwise falls back to drivers. 
+	UE_API void SetDisplayName(FString Name);
+#endif
+
 	UPROPERTY(EditAnywhere, BlueprintReadWrite, Category = "MusicClock")
 	ECalibratedMusicTimebase TimebaseForBarAndBeatEvents = ECalibratedMusicTimebase::VideoRenderTime;
 
@@ -523,6 +531,10 @@ private:
 	int32 LastBroadcastBar = -1;
 	int32 LastBroadcastBeat = -1;
 	FSongSection LastBroadcastSongSection;
+#if !UE_BUILD_SHIPPING
+	// Holds a name for display, ideally shows the original source for the clock.
+	FString DisplayName;
+#endif
 	
 	static TNotNull<UMusicClockComponent*> CreateMusicClockShared(UObject* WorldContextObject, bool bPersistAcrossLevelTransition);
 	
diff --git a/Engine/Plugins/Runtime/Harmonix/Source/HarmonixMetasound/Public/HarmonixMetasound/Components/MusicClockDriverBase.h b/Engine/Plugins/Runtime/Harmonix/Source/HarmonixMetasound/Public/HarmonixMetasound/Components/MusicClockDriverBase.h
--- a/Engine/Plugins/Runtime/Harmonix/Source/HarmonixMetasound/Public/HarmonixMetasound/Components/MusicClockDriverBase.h
+++ b/Engine/Plugins/Runtime/Harmonix/Source/HarmonixMetasound/Public/HarmonixMetasound/Components/MusicClockDriverBase.h
@@ -53,7 +53,12 @@ public:
 
 	// This returns our BeatsPerMinute which is our true beats, not quarter notes.
 	UE_API float GetBeatsPerMinute() const;
-	
+
+#if !UE_BUILD_SHIPPING
+	// Short user readable description good for logging. Should include name of song or sound source or whatnot. 
+	virtual FString GetDisplayName() const = 0;
+#endif
+
 	FSongMaps DefaultMaps;
 	
 	// Note that tempo is MidiTempo, quarter notes per minute.
