FSongMapChange

Use atomic pointers to fix SharedPtr access race conditions

#rnx

#Synced-CL 48881229

diff --git a/Engine/Plugins/Runtime/Harmonix/Source/HarmonixMetasound/Private/HarmonixMetasound/Components/MetasoundMusicClockDriver.cpp b/Engine/Plugins/Runtime/Harmonix/Source/HarmonixMetasound/Private/HarmonixMetasound/Components/MetasoundMusicClockDriver.cpp
--- a/Engine/Plugins/Runtime/Harmonix/Source/HarmonixMetasound/Private/HarmonixMetasound/Components/MetasoundMusicClockDriver.cpp
+++ b/Engine/Plugins/Runtime/Harmonix/Source/HarmonixMetasound/Private/HarmonixMetasound/Components/MetasoundMusicClockDriver.cpp
@@ -545,7 +545,7 @@ void FMetasoundMusicClockDriver::RefreshCurrentSongPosFromHistory()
 		return;
 	}
 
-	if (!CurrentMapChain || !CurrentMapChain->SongMaps || CurrentMapChain->NewSongMaps)
+	if (!CurrentMapChain || !CurrentMapChain->SongMaps || !CurrentMapChain->IsLatest())
 	{
 		CurrentMapChain = ClockHistory->GetLatestMapsForConsumer();
 		if (!CurrentMapChain || !CurrentMapChain->SongMaps)
diff --git a/Engine/Plugins/Runtime/Harmonix/Source/HarmonixMetasound/Public/HarmonixMetasound/Analysis/MidiClockSongPos.h b/Engine/Plugins/Runtime/Harmonix/Source/HarmonixMetasound/Public/HarmonixMetasound/Analysis/MidiClockSongPos.h
--- a/Engine/Plugins/Runtime/Harmonix/Source/HarmonixMetasound/Public/HarmonixMetasound/Analysis/MidiClockSongPos.h
+++ b/Engine/Plugins/Runtime/Harmonix/Source/HarmonixMetasound/Public/HarmonixMetasound/Analysis/MidiClockSongPos.h
@@ -6,6 +6,7 @@
 #include "HarmonixMidi/SongMaps.h"
 #include "HarmonixMetasound/Analysis/SpmcAnalysisResultQueue.h"
 #include "HarmonixMetasound/DataTypes/MusicTransport.h"
+#include <atomic>
 
 namespace HarmonixMetasound::Analysis
 {
@@ -25,7 +26,7 @@ namespace HarmonixMetasound::Analysis
 		EMarkerType MarkerType {EMarkerType::None};
 	};
 
-	struct FSongMapChain
+	struct FSongMapChain : TSharedFromThis<FSongMapChain>
 	{
 		FSongMapChain() {}
 		FSongMapChain(TSharedPtr<const ISongMapEvaluator>& InMaps, int32 InFirstTickInLoop, int32 InLoopLengthTicks)
@@ -37,6 +38,35 @@ namespace HarmonixMetasound::Analysis
 		TSharedPtr<const ISongMapEvaluator> SongMaps;
 		int32 FirstTickInLoop {-1};
 		int32 LoopLengthTicks {0};
+
+		TSharedPtr<FSongMapChain> GetLatest()
+		{
+			FSongMapChain* Latest = this;
+			while (Latest->NewSongMapsPtr.load() != nullptr)
+			{
+				Latest = Latest->NewSongMapsPtr;
+			}
+			return Latest->AsShared();
+		}
+
+		void Update(TSharedPtr<const ISongMapEvaluator>& InMaps, int32 InFirstTickInLoop, int32 InLoopLengthTicks)
+		{
+			FSongMapChain* Latest = this;
+			while (Latest->NewSongMapsPtr.load() != nullptr)
+			{
+				Latest = Latest->NewSongMapsPtr;
+			}
+			Latest->NewSongMaps = MakeShared<FSongMapChain>(InMaps, InFirstTickInLoop, InLoopLengthTicks);
+			Latest->NewSongMapsPtr = Latest->NewSongMaps.Get();
+		}
+
+		bool IsLatest() const
+		{
+			return NewSongMapsPtr.load() == nullptr;
+		}
+
+	private:
+		std::atomic<FSongMapChain*> NewSongMapsPtr = nullptr;
 		TSharedPtr<FSongMapChain> NewSongMaps;
 	};
 
@@ -69,33 +99,18 @@ namespace HarmonixMetasound::Analysis
 
 		TSharedPtr<const FSongMapChain> GetLatestMapsForConsumer()
 		{
-			TSharedPtr<FSongMapChain> Returnable = SongMapChain;
-			while (Returnable->NewSongMaps)
-			{
-				Returnable = Returnable->NewSongMaps;
-			}
-			SongMapChain = Returnable;
-			return Returnable;
+			SongMapChain = SongMapChain->GetLatest();
+			return SongMapChain;
 		}
 
-		TSharedPtr<const FSongMapChain> GetLatestMapsForProducer()
+		TSharedPtr<const FSongMapChain> GetLatestMapsForProducer() const
 		{
-			TSharedPtr<FSongMapChain> Returnable = SongMapChain;
-			while (Returnable->NewSongMaps)
-			{
-				Returnable = Returnable->NewSongMaps;
-			}
-			return Returnable;
+			return SongMapChain->GetLatest();
 		}
 
 		void UpdateMaps(TSharedPtr<const ISongMapEvaluator> Maps, int32 InFirstTickInLoop, int32 InLoopLengthTicks)
 		{
-			TSharedPtr<FSongMapChain> Chain = SongMapChain;
-			while (Chain->NewSongMaps)
-			{
-				Chain = Chain->NewSongMaps;
-			}
-			Chain->NewSongMaps = MakeShared<FSongMapChain>(Maps, InFirstTickInLoop, InLoopLengthTicks);
+			SongMapChain->Update(Maps, InFirstTickInLoop, InLoopLengthTicks);
 		}
 	};
 }
