Music Clock Refactor, Pt 1.

"Rename" `FMusicClockDriverBase` to `FMusicClock`

First steps in only needing "2" subclasses

- MetasoundMusicClock: uses output of Metasound to determine Time and Music Maps

- RuntimeMusicClock: uses an external "TimeSource" and external "Music Maps" to determine music time

ITimeSource
- Base class for external time sources
- Things like "Animation Time" and "Media Player Time" should be turned into implementations of this class
- Exposes which "Calibrated Time Base" it represents

#rnx

#Synced-CL 49167537

diff --git a/Engine/Plugins/Runtime/Harmonix/Source/HarmonixMetasound/Private/HarmonixMetasound/MusicClock/MetasoundMusicClock.cpp b/Engine/Plugins/Runtime/Harmonix/Source/HarmonixMetasound/Private/HarmonixMetasound/MusicClock/MetasoundMusicClock.cpp
new file mode 100644
--- /dev/null
+++ b/Engine/Plugins/Runtime/Harmonix/Source/HarmonixMetasound/Private/HarmonixMetasound/MusicClock/MetasoundMusicClock.cpp
@@ -0,0 +1,783 @@
+// Copyright Epic Games, Inc. All Rights Reserved.
+
+#include "HarmonixMetasound/MusicClock/MetasoundMusicClock.h"
+#include "MetasoundGeneratorHandle.h"
+#include "Components/AudioComponent.h"
+#include "HarmonixMetasound/Analysis/MidiSongPosVertexAnalyzer.h"
+#include "HarmonixMetasound/DataTypes/MidiClock.h"
+#include "MetasoundGenerator.h"
+#include "Async/Async.h"
+#include "Harmonix.h"
+#include "MetasoundSource.h"
+
+DEFINE_LOG_CATEGORY_STATIC(LogMetasoundMusicClock, Log, All)
+
+namespace MetasoundMusicClock
+{
+	namespace Smoothing
+	{
+		// Proportional factor for error correction
+		float kP = 0.18;
+
+		// Baked in Audio Renderer Lag Seconds
+		float LagSeconds = 0.030f; // 30 ms
+
+		// Max allowed error before force jumping to sync to the current time
+		float MaxErrorSecondsBeforeJump = 0.060f; // 60ms
+
+		// The min and max allowed speed to maintain sync with the audio time
+		double MinSyncSpeed = 0.98;
+		double MaxSyncSpeed = 1.02;
+	}
+}
+
+TOptional<FVector> Harmonix::FMetasoundMusicClock::TryGetAudioSourceLocation() const
+{
+	if (UAudioComponent* AudioComponent = AudioComponentToWatch.Get())
+	{
+		return AudioComponent->GetComponentLocation();
+	}
+	return {};
+}
+
+bool Harmonix::FMetasoundMusicClock::CalculateSongPosWithOffset(float MsOffset, ECalibratedMusicTimebase Timebase, FMidiSongPos& OutResult) const
+{
+	using namespace HarmonixMetasound::Analysis;
+
+	check(IsInGameThread());
+
+	// if we have an owner, ask them directly
+	if (!ClockHistory)
+	{
+		return false;
+	}
+
+	if (!CurrentMapChain || !CurrentMapChain->SongMaps)
+	{
+		return false;
+	}
+
+	const FPerTimebaseSmoothedClockState* ClockState = nullptr;
+
+	switch (Timebase)
+	{
+	case ECalibratedMusicTimebase::AudioRenderTime:  ClockState = &AudioRenderState; break;
+	case ECalibratedMusicTimebase::ExperiencedTime:  ClockState = &PlayerExperienceState; break;
+	case ECalibratedMusicTimebase::VideoRenderTime:
+	default:                                         ClockState = &VideoRenderState; break;
+	}
+
+	float AbsMs = ClockState->TempoMapMs + MsOffset;
+	float TempoMapOffsetTick = CurrentMapChain->SongMaps->MsToTick(AbsMs);
+	float RelativeTicks = TempoMapOffsetTick - ClockState->TempoMapTick;
+	float SmoothedOffsetTick = ClockState->LocalTick + RelativeTicks;
+
+	if (CurrentMapChain->LoopLengthTicks > 0)
+	{
+		while (SmoothedOffsetTick < 0.0f)
+		{
+			SmoothedOffsetTick += CurrentMapChain->LoopLengthTicks;
+		}
+		while (SmoothedOffsetTick >= CurrentMapChain->LoopLengthTicks)
+		{
+			SmoothedOffsetTick -= CurrentMapChain->LoopLengthTicks;
+		}
+	}
+
+	// first 99% of the song pos...
+	OutResult.SetByTick(SmoothedOffsetTick, *(CurrentMapChain->SongMaps));
+	// but tempo needs to be set to authoritiesTempo...
+	OutResult.Tempo = CurrentMapChain->SongMaps->GetTempoAtTick(FMath::FloorToInt32(TempoMapOffsetTick));
+
+	return true;
+}
+
+FMidiSongPos Harmonix::FMetasoundMusicClock::CalculateSongPosAtMsForLoopingOrMonotonicClock(float AbsoluteMs, float& PositionTick, bool& SeekDetected, bool& LoopDetected) const
+{
+	using namespace HarmonixMetasound::Analysis;
+	FMidiSongPos OutSongPos;
+	if (!ClockHistory)
+	{
+		PositionTick = 0.0f;
+		return OutSongPos;
+	}
+
+	if (!CurrentMapChain || !CurrentMapChain->SongMaps)
+	{
+		PositionTick = 0.0f;
+		return OutSongPos;
+	}
+
+	float NewPositionTick = 0.0f;
+	if (CurrentMapChain->LoopLengthTicks > 0)
+	{
+		float DrivingTick = CurrentMapChain->SongMaps->MsToTick(AbsoluteMs);
+		float TickPastLoop = (float)(CurrentMapChain->FirstTickInLoop + CurrentMapChain->LoopLengthTicks);
+		if (DrivingTick >= TickPastLoop)
+		{
+			float WrappedTick = FMath::Fmod(DrivingTick - CurrentMapChain->FirstTickInLoop, (float)CurrentMapChain->LoopLengthTicks);
+			LoopDetected = (PositionTick - WrappedTick) > (float)(CurrentMapChain->LoopLengthTicks - 240);
+			if (LoopDetected)
+			{
+				UE_LOG(LogMetasoundMusicClock, VeryVerbose, TEXT("Detected loop when calculating song pos (%f -> %f -> %f -> %d )"), PositionTick, WrappedTick, PositionTick - WrappedTick, CurrentMapChain->LoopLengthTicks);
+			}
+			OutSongPos.SetByTick(WrappedTick, *CurrentMapChain->SongMaps);
+			OutSongPos.Tempo = CurrentMapChain->SongMaps->GetTempoAtTick(FMath::FloorToInt32(DrivingTick));
+
+			if (!LoopDetected)
+			{
+				SeekDetected = CheckForSeek(PositionTick, WrappedTick, OutSongPos.Tempo, CurrentMapChain->SongMaps->GetTicksPerQuarterNote());
+				if (SeekDetected)
+				{
+					UE_LOG(LogMetasoundMusicClock, VeryVerbose, TEXT("Detected seek when calculating song pos (%f -> %f)"), PositionTick, WrappedTick);
+				}
+			}
+
+			PositionTick = WrappedTick;
+			return OutSongPos;
+		}
+		NewPositionTick = DrivingTick;
+	}
+	else
+	{
+		NewPositionTick = CurrentMapChain->SongMaps->MsToTick(AbsoluteMs);
+	}
+
+	OutSongPos.SetByTimeAndTick(AbsoluteMs, NewPositionTick, *(CurrentMapChain->SongMaps));
+	SeekDetected = CheckForSeek(PositionTick, NewPositionTick, OutSongPos.Tempo, CurrentMapChain->SongMaps->GetTicksPerQuarterNote());
+	if (SeekDetected)
+	{
+		UE_LOG(LogMetasoundMusicClock, VeryVerbose, TEXT("Detected seek when calculating song pos (%f -> %f)"), PositionTick, NewPositionTick);
+	}
+	PositionTick = NewPositionTick;
+	return OutSongPos;
+}
+
+FMidiSongPos Harmonix::FMetasoundMusicClock::CalculateSongPosAtMsForOffsetClock(float PositionMs, float ClockTickOffsetFromDrivingClock, float& PositionTick, bool& SeekDetected) const
+{
+	using namespace HarmonixMetasound::Analysis;
+	
+	FMidiSongPos OutSongPos;
+	
+	if (!ClockHistory)
+	{
+		PositionTick = 0.0f;
+		return OutSongPos;
+	}
+
+	if (!CurrentMapChain || !CurrentMapChain->SongMaps)
+	{
+		PositionTick = 0.0f;
+		return OutSongPos;
+	}
+
+	float NewPositionTick = CurrentMapChain->SongMaps->MsToTick(PositionMs);
+
+	OutSongPos.SetByTick(NewPositionTick, *CurrentMapChain->SongMaps);
+	OutSongPos.Tempo = CurrentMapChain->SongMaps->GetTempoAtTick(FMath::FloorToInt32(NewPositionTick - ClockTickOffsetFromDrivingClock));
+
+	SeekDetected = CheckForSeek(PositionTick, NewPositionTick, OutSongPos.Tempo, CurrentMapChain->SongMaps->GetTicksPerQuarterNote());
+	if (SeekDetected)
+	{
+		UE_LOG(LogMetasoundMusicClock, VeryVerbose, TEXT("Detected seek when calculating song pos (%f -> %f)"), PositionTick, NewPositionTick);
+	}
+
+	PositionTick = NewPositionTick;
+
+	return OutSongPos;
+}
+
+bool Harmonix::FMetasoundMusicClock::CheckForSeek(float FirstTick, float NextTick, float CurrentTempo, int32 TicksPerQuarter) const
+{
+	float QuartersPerSecond = CurrentTempo / 60.0f;
+	float ExpectedDeltaQuarters = QuartersPerSecond * DeltaSecondsBetweenRefreshes;
+	float ExpectedDeltaTicks = ExpectedDeltaQuarters * (float)TicksPerQuarter;
+	return FMath::Abs(ExpectedDeltaTicks - (NextTick - FirstTick)) > (ExpectedDeltaTicks * 2.0f);
+}
+
+bool Harmonix::FMetasoundMusicClock::RefreshCurrentSongPos()
+{
+	//	Only for use when on the game thread.
+	if (ensureMsgf(
+			IsInGameThread(),
+			TEXT("%hs called from non-game thread.  This is not supported"),
+			__FUNCTION__) == false)
+	{
+		return false;
+	}
+
+	if (AudioComponentToWatch.IsValid())
+	{
+		if (!CurrentGeneratorHandle)
+		{
+			// cursor connection is not pending
+			AttemptToConnectToAudioComponentsMetasound(AudioComponentToWatch.Get(), OnAttachedDelegate, OnDetachedDelegate);
+		}
+	}
+
+	if (bRunning)
+	{
+		if (ClockHistory)
+		{
+			// cursor is attached and has the current info
+			RefreshCurrentSongPosFromHistory();
+			return true;
+		}
+		else
+		{
+			// Cursor not attached so use wall clock
+			if (!WasEverConnected || RunPastMusicEnd)
+			{
+				RefreshCurrentSongPosFromWallClock();
+				return true;
+			}
+		}
+	}
+
+	return false;
+}
+
+void Harmonix::FMetasoundMusicClock::OnStart()
+{
+	check(IsInGameThread());
+
+	SongPosOffsetMs = 0.0f;
+	RenderStartSampleCount = 0;
+	RenderStartWallClockTimeSeconds = 0.0;
+	FreeRunStartTimeSecs = GetWallClockTime();
+	bRunning = true;
+}
+
+void Harmonix::FMetasoundMusicClock::OnPause()
+{
+	check(IsInGameThread());
+	bRunning = false;
+}
+
+void Harmonix::FMetasoundMusicClock::OnContinue()
+{
+	check(IsInGameThread());
+	if (!ClockHistory)
+	{
+		RefreshCurrentSongPosFromWallClock();
+	}
+	bRunning = true;
+}
+
+void Harmonix::FMetasoundMusicClock::OnStop()
+{
+	check(IsInGameThread());
+	bRunning = false;
+}
+
+void Harmonix::FMetasoundMusicClock::Disconnect()
+{
+	check(IsInGameThread());
+	SetClockHistory(nullptr);
+	DetachAllCallbacks();
+	AudioComponentToWatch.Reset();
+	CurrentGeneratorHandle.Reset();
+}
+
+const ISongMapEvaluator* Harmonix::FMetasoundMusicClock::GetCurrentSongMapEvaluator() const
+{
+	using namespace HarmonixMetasound::Analysis;
+	check(IsInGameThread());
+	if (ClockHistory && CurrentMapChain && CurrentMapChain->SongMaps)
+	{
+		return CurrentMapChain->SongMaps.Get();
+	}
+	return &DefaultMaps;
+}
+
+bool Harmonix::FMetasoundMusicClock::LoopedThisFrame(ECalibratedMusicTimebase Timebase) const
+{
+	switch (Timebase)
+	{
+	case ECalibratedMusicTimebase::RawAudioRenderTime:
+	case ECalibratedMusicTimebase::AudioRenderTime:    return AudioRenderLoopDetected;
+	case ECalibratedMusicTimebase::ExperiencedTime:	   return PlayerExperiencedLoopDetected;
+	case ECalibratedMusicTimebase::VideoRenderTime:
+	default:                                           return VideoRenderLoopDetected;
+	}
+}
+
+bool Harmonix::FMetasoundMusicClock::SeekedThisFrame(ECalibratedMusicTimebase Timebase) const
+{
+	switch (Timebase)
+	{
+	case ECalibratedMusicTimebase::RawAudioRenderTime:
+	case ECalibratedMusicTimebase::AudioRenderTime:    return AudioRenderSeekDetected;
+	case ECalibratedMusicTimebase::ExperiencedTime:	   return PlayerExperiencedSeekDetected;
+	case ECalibratedMusicTimebase::VideoRenderTime:
+	default:                                           return VideoRenderSeekDetected;
+	}
+}
+
+#if !UE_BUILD_SHIPPING
+FString Harmonix::FMetasoundMusicClock::GetDisplayName() const
+{
+#if WITH_EDITORONLY_DATA
+	// that method is only good WITH_EDITORONLY_DATA
+	if (UAudioComponent* Component = AudioComponentToWatch.Get())
+	{
+		if (UMetaSoundSource* Source = Cast<UMetaSoundSource>(Component->GetSound()))
+		{
+			return FString::Printf(TEXT("Metasound: %s"), *Source->GetDisplayName().ToString());
+		}
+	}
+#endif
+	return TEXT("Metasound");
+}
+#endif
+
+bool Harmonix::FMetasoundMusicClock::ConnectToAudioComponentsMetasound(
+	UAudioComponent* InAudioComponent,
+	FName MetasoundOuputPinName,
+	UMetasoundGeneratorHandle::FOnAttached::FDelegate OnGeneratorAttachedCallback,
+	UMetasoundGeneratorHandle::FOnDetached::FDelegate OnGeneratorDetachedCallback)
+{
+	OnAttachedDelegate = OnGeneratorAttachedCallback;
+	OnDetachedDelegate = OnGeneratorDetachedCallback;
+	
+	AudioComponentToWatch = InAudioComponent;
+	MetasoundOutputName = MetasoundOuputPinName;
+	bool Connected = AttemptToConnectToAudioComponentsMetasound(InAudioComponent, OnGeneratorAttachedCallback, OnGeneratorDetachedCallback);
+
+	if(GetState() == EMusicClockState::Running)
+	{
+		OnStart();
+	}
+	
+	return Connected;
+}
+
+bool Harmonix::FMetasoundMusicClock::AttemptToConnectToAudioComponentsMetasound(
+	UAudioComponent* InAudioComponent,
+	UMetasoundGeneratorHandle::FOnAttached::FDelegate OnGeneratorAttachedCallback,
+	UMetasoundGeneratorHandle::FOnDetached::FDelegate OnGeneratorDetachedCallback)
+{
+	using namespace HarmonixMetasound::Analysis;
+
+	check(IsInGameThread());
+	if (!AudioComponentToWatch.IsValid() || MetasoundOutputName.IsNone())
+	{
+		return false;
+	}
+
+	// we have an audio component, but it's not set to play a Metasound (yet)
+	if (!Cast<UMetaSoundSource>(AudioComponentToWatch->GetSound()))
+	{
+		return false;
+	}
+
+	DetachAllCallbacks();
+	CurrentGeneratorHandle.Reset(UMetasoundGeneratorHandle::CreateMetaSoundGeneratorHandle(AudioComponentToWatch.Get()));
+	if (!CurrentGeneratorHandle)
+	{
+		return false;
+	}
+
+	bool WatchingOutput = CurrentGeneratorHandle->WatchOutput(MetasoundOutputName, 
+		FOnMetasoundOutputValueChangedNative::CreateLambda([](FName, const FMetaSoundOutput&) {}),
+		FMidiSongPosVertexAnalyzer::GetAnalyzerName(),
+		FMidiSongPosVertexAnalyzer::SongPosition.Name);
+	if (WatchingOutput)
+	{
+		ensure(CurrentGeneratorHandle->TryCreateAnalyzerAddress(MetasoundOutputName, 
+				FMidiSongPosVertexAnalyzer::GetAnalyzerName(), 
+				FMidiSongPosVertexAnalyzer::SongPosition.Name, 
+				MidiSongPosAnalyzerAddress));
+	}
+
+	GeneratorAttachedCallbackHandle = CurrentGeneratorHandle->OnGeneratorHandleAttached.AddLambda([this, OnGeneratorAttachedCallback]()
+	{
+		OnGeneratorAttached();
+		OnGeneratorAttachedCallback.ExecuteIfBound();
+	});
+	GeneratorDetachedCallbackHandle = CurrentGeneratorHandle->OnGeneratorHandleDetached.AddLambda([this, OnGeneratorDetachedCallback]()
+	{
+		OnGeneratorDetached();
+		OnGeneratorDetachedCallback.ExecuteIfBound();
+	});
+	GeneratorIOUpdatedCallbackHandle = CurrentGeneratorHandle->OnIOUpdatedWithChanges.AddLambda([this](const TArray<Metasound::FVertexInterfaceChange>& VertexInterfaceChanges){OnGeneratorIOUpdatedWithChanges(VertexInterfaceChanges);});
+	UMetasoundGeneratorHandle::FOnSetGraph::FDelegate OnSetGraph;
+	OnSetGraph.BindLambda([this](){OnGraphSet();});
+	GraphChangedCallbackHandle = CurrentGeneratorHandle->AddGraphSetCallback(MoveTemp(OnSetGraph));
+	// OnGeneratorAttached();
+	return true;
+}
+
+void Harmonix::FMetasoundMusicClock::DetachAllCallbacks()
+{
+	if (CurrentGeneratorHandle)
+	{
+		CurrentGeneratorHandle->OnGeneratorHandleAttached.Remove(GeneratorAttachedCallbackHandle);
+		GeneratorAttachedCallbackHandle.Reset();
+		CurrentGeneratorHandle->OnGeneratorHandleDetached.Remove(GeneratorDetachedCallbackHandle);
+		GeneratorDetachedCallbackHandle.Reset();
+
+		CurrentGeneratorHandle->OnIOUpdatedWithChanges.Remove(GeneratorIOUpdatedCallbackHandle);
+		GeneratorIOUpdatedCallbackHandle.Reset();
+		CurrentGeneratorHandle->RemoveGraphSetCallback(GraphChangedCallbackHandle);
+		GraphChangedCallbackHandle.Reset();
+
+	}
+	SetClockHistory(nullptr);
+}
+
+void Harmonix::FMetasoundMusicClock::OnGeneratorAttached()
+{
+	WasEverConnected = true;
+	UnderlyingGenerator = CurrentGeneratorHandle->GetGenerator();
+	SetClockHistory(UMidiClockUpdateSubsystem::GetOrCreateClockHistory(MidiSongPosAnalyzerAddress));
+	SmoothedAudioRenderClockHistoryCursor = ClockHistory->CreateReadCursor();
+	SmoothedPlayerExperienceClockHistoryCursor = ClockHistory->CreateReadCursor();
+	SmoothedVideoRenderClockHistoryCursor = ClockHistory->CreateReadCursor();
+}
+
+void Harmonix::FMetasoundMusicClock::OnGraphSet()
+{
+	SetClockHistory(UMidiClockUpdateSubsystem::GetOrCreateClockHistory(MidiSongPosAnalyzerAddress));
+	SmoothedAudioRenderClockHistoryCursor = ClockHistory->CreateReadCursor();
+	SmoothedPlayerExperienceClockHistoryCursor = ClockHistory->CreateReadCursor();
+	SmoothedVideoRenderClockHistoryCursor = ClockHistory->CreateReadCursor();
+}
+
+
+void Harmonix::FMetasoundMusicClock::OnGeneratorIOUpdatedWithChanges(const TArray<Metasound::FVertexInterfaceChange>& VertexInterfaceChanges)
+{
+	if (!MetasoundOutputName.IsNone())
+	{
+		if (VertexInterfaceChanges.Num() > 0)
+		{
+			SetClockHistory(UMidiClockUpdateSubsystem::GetOrCreateClockHistory(MidiSongPosAnalyzerAddress));
+			SmoothedAudioRenderClockHistoryCursor = ClockHistory->CreateReadCursor();
+			SmoothedPlayerExperienceClockHistoryCursor = ClockHistory->CreateReadCursor();
+			SmoothedVideoRenderClockHistoryCursor = ClockHistory->CreateReadCursor();
+		}
+	}
+	
+}
+
+
+void Harmonix::FMetasoundMusicClock::OnGeneratorDetached()
+{
+	TWeakPtr<Metasound::FMetasoundGenerator> DetachingGenerator = CurrentGeneratorHandle->GetGenerator();
+	if (DetachingGenerator != UnderlyingGenerator)
+	{
+		// yes... it is possible we will get a detach callback for a generator that is not ours! yikes.
+		return;
+	}
+
+	using namespace HarmonixMetasound::Analysis;
+	if (GetState() != EMusicClockState::Stopped)
+	{
+		if (ClockHistory && CurrentMapChain && CurrentMapChain->SongMaps)
+		{
+			DefaultMaps.Copy(*(CurrentMapChain->SongMaps), 0, LastTickSeen);
+		}
+		SongPosOffsetMs = GetCurrentSongPos(ECalibratedMusicTimebase::AudioRenderTime).SecondsIncludingCountIn * 1000.0f;
+		FreeRunStartTimeSecs = GetWallClockTime();
+	}
+	SetClockHistory(nullptr);
+	SmoothedAudioRenderClockHistoryCursor = HarmonixMetasound::Analysis::FMidiClockSongPositionHistory::FReadCursor();
+	SmoothedPlayerExperienceClockHistoryCursor = HarmonixMetasound::Analysis::FMidiClockSongPositionHistory::FReadCursor();
+	SmoothedVideoRenderClockHistoryCursor = HarmonixMetasound::Analysis::FMidiClockSongPositionHistory::FReadCursor();
+}
+
+void Harmonix::FMetasoundMusicClock::RefreshCurrentSongPosFromWallClock()
+{
+	double FreeRunTime = (GetWallClockTime() - FreeRunStartTimeSecs) * CurrentClockAdvanceRate;
+	SetCurrentSongPosByTime(ECalibratedMusicTimebase::RawAudioRenderTime, ((float)FreeRunTime * 1000.0) + SongPosOffsetMs, DefaultMaps);
+	SetCurrentSongPos(ECalibratedMusicTimebase::AudioRenderTime, GetCurrentSongPos(ECalibratedMusicTimebase::RawAudioRenderTime));
+	SetCurrentSongPosByTime(ECalibratedMusicTimebase::ExperiencedTime, GetCurrentSongPos(ECalibratedMusicTimebase::AudioRenderTime).SecondsIncludingCountIn * 1000.0f - FHarmonixModule::GetMeasuredUserExperienceAndReactionToAudioRenderOffsetMs(), DefaultMaps);
+	SetCurrentSongPosByTime(ECalibratedMusicTimebase::VideoRenderTime, GetCurrentSongPos(ECalibratedMusicTimebase::AudioRenderTime).SecondsIncludingCountIn * 1000.0f - FHarmonixModule::GetMeasuredVideoToAudioRenderOffsetMs(), DefaultMaps);
+
+	const FMidiSongPos& CurrentSmoothedAudioRenderSongPos = GetCurrentSongPos(ECalibratedMusicTimebase::AudioRenderTime);
+	UpdateMusicPlaybackRate(CurrentSmoothedAudioRenderSongPos.Tempo, CurrentClockAdvanceRate, CurrentSmoothedAudioRenderSongPos.TimeSigNumerator, CurrentSmoothedAudioRenderSongPos.TimeSigDenominator);
+}
+
+FString Harmonix::FMetasoundMusicClock::HistoryFailureTypeToString(Harmonix::FMetasoundMusicClock::EHistoryFailureType Error)
+{
+	switch (Error)
+	{
+	case Harmonix::FMetasoundMusicClock::EHistoryFailureType::None: return "None";
+	case Harmonix::FMetasoundMusicClock::EHistoryFailureType::NotEnoughDataInTheHistoryRing: return "NotEnoughDataInTheHistoryRing";
+	case Harmonix::FMetasoundMusicClock::EHistoryFailureType::NotEnoughHistory: return "NotEnoughHistory";
+	case Harmonix::FMetasoundMusicClock::EHistoryFailureType::LookingForTimeInTheFutureOfWhatHasEvenRendered: return "LookingForTimeInTheFutureOfWhatHasEvenRendered";
+	case Harmonix::FMetasoundMusicClock::EHistoryFailureType::CaughtUpToRenderPosition: return "CaughtUpToRenderPosition";
+	}
+	return "Unrecognized";
+}
+
+void Harmonix::FMetasoundMusicClock::RefreshCurrentSongPosFromHistory()
+{
+	using namespace HarmonixMetasound::Analysis;
+	using namespace MetasoundMusicClock;
+
+	check(IsInGameThread());
+
+	if (!bRunning || !SmoothedAudioRenderClockHistoryCursor.DataAvailable())
+	{
+		return;
+	}
+	
+	if (!SmoothedAudioRenderClockHistoryCursor.Queue)
+	{
+		return;
+	}
+
+	if (!CurrentMapChain || !CurrentMapChain->SongMaps || !CurrentMapChain->IsLatest())
+	{
+		CurrentMapChain = ClockHistory->GetLatestMapsForConsumer();
+		if (!CurrentMapChain || !CurrentMapChain->SongMaps)
+		{
+			return;
+		}
+	}
+
+	auto Entry = ClockHistory->Positions.GetEntry(ClockHistory->Positions.GetLastWriteIndex());
+	SetCurrentSongPosByTick(ECalibratedMusicTimebase::RawAudioRenderTime, Entry->Item.UpToTick, *(CurrentMapChain->SongMaps));
+	Metasound::FSampleCount LastRenderPosSampleCount = Entry->Item.SampleCount;
+	float SpeedAtRawRenderTime = Entry->Item.CurrentSpeed;
+	LastTickSeen = Entry->Item.UpToTick;
+	bool bClockIsStopped = Entry->Item.CurrentTransportState != HarmonixMetasound::EMusicPlayerTransportState::Playing;
+
+	if (RenderStartWallClockTimeSeconds == 0.0)
+	{
+		// We are just starting up. We need to create the initial "sync point"...
+		// Wall clock <-> Render Samples...
+		RenderStartSampleCount = LastRenderPosSampleCount;
+		RenderStartWallClockTimeSeconds = GetWallClockTime() - (double)RenderStartSampleCount / (double)ClockHistory->SampleRate;
+		RenderSmoothingLagSeconds = Smoothing::LagSeconds;
+		ErrorTracker.Reset();
+		LastRefreshWallClockTimeSeconds = RenderStartWallClockTimeSeconds;
+	}
+
+	double CurrentWallClockSeconds = GetWallClockTime();
+	DeltaSecondsBetweenRefreshes = CurrentWallClockSeconds - LastRefreshWallClockTimeSeconds;
+	LastRefreshWallClockTimeSeconds = CurrentWallClockSeconds;
+
+	double ExpectedRenderedSeconds = (CurrentWallClockSeconds - RenderStartWallClockTimeSeconds) * SyncSpeed;
+	double RenderedSeconds = (double)LastRenderPosSampleCount / (double)ClockHistory->SampleRate;
+	double Error = RenderedSeconds - ExpectedRenderedSeconds;
+
+	if (!bClockIsStopped)
+	{
+		ErrorTracker.Push(Error);
+
+		if (FMath::Abs(ErrorTracker.Min()) > Smoothing::MaxErrorSecondsBeforeJump)
+		{
+			UE_LOG(LogMetasoundMusicClock, Verbose, TEXT("======== MASSIVE ERROR (%f) - SEEKING ==========="), (float)Error);
+			RenderStartSampleCount = LastRenderPosSampleCount;
+			RenderStartWallClockTimeSeconds = GetWallClockTime() - (double)RenderStartSampleCount / (double)ClockHistory->SampleRate;
+			ExpectedRenderedSeconds = RenderedSeconds;
+			RenderSmoothingLagSeconds = Smoothing::LagSeconds;
+			ErrorTracker.Reset();
+			Error = 0;
+			SyncSpeed = 1.0;
+		}
+
+		// Use proportional part of error to adjust speed ever so slightly...
+		if (ExpectedRenderedSeconds > 0.0)
+		{
+			SyncSpeed += Smoothing::kP * ErrorTracker.Min() / ExpectedRenderedSeconds;
+		}
+		SyncSpeed = FMath::Clamp(SyncSpeed, Smoothing::MinSyncSpeed, Smoothing::MaxSyncSpeed);
+	}
+
+	//UE_LOG(LogMetasoundMusicClock, Log, TEXT("Local Minimum Error: %f Average: %f Speed: %f"), (float)ErrorTracker.Min(), (float)ErrorTracker.Average(), (float)SyncSpeed);
+
+	Metasound::FSampleCount ExpectedRenderPosSampleCount = (Metasound::FSampleCount)(ExpectedRenderedSeconds * (double)ClockHistory->SampleRate);
+
+	// first smoothed render time. This is closest to the actual render time. If we catch up to the render time
+	// it means we are rendering in such large blocks that we need to push up our "look behind" number for smoothing.
+	const FMidiSongPos& CurrentSmoothedAudioRenderSongPos = GetCurrentSongPos(ECalibratedMusicTimebase::AudioRenderTime);
+	float SpeedAtSmoothRenderedTime = 1.0f;
+	float SmoothedTick = GetCurrentSongMapEvaluator()->MsToTick(CurrentSmoothedAudioRenderSongPos.SecondsIncludingCountIn * 1000.0f);
+	float SmoothedTempoMapTick = AudioRenderState.TempoMapTick;
+	EHistoryFailureType Result = CalculateSmoothedTick(ExpectedRenderPosSampleCount, LastRenderPosSampleCount,
+		SmoothedTick, SmoothedTempoMapTick, SpeedAtSmoothRenderedTime, SmoothedAudioRenderClockHistoryCursor, RenderSmoothingLagSeconds);
+	if (Result != EHistoryFailureType::None)
+	{
+		if (LastRenderPosSampleCount >  (RenderSmoothingLagSeconds * ClockHistory->SampleRate * 2))
+		{
+			if (RenderSmoothingLagSeconds < 0.250f)
+			{
+				RenderSmoothingLagSeconds += 0.005f;
+				UE_LOG(LogMetasoundMusicClock, Verbose, TEXT("(%d) Smoothed Render Time too close to actual render time. Bumping up smoothing lag! (%f)"), LastRenderPosSampleCount, RenderSmoothingLagSeconds);
+			}
+			else
+			{
+				UE_LOG(LogMetasoundMusicClock, Verbose, TEXT("(%d) Smoothed Render Time too close to actual render time. ALREADY MAX SMOOTHING LAG! (%f)"), LastRenderPosSampleCount, RenderSmoothingLagSeconds);
+			}
+		}
+		else
+		{
+			UE_LOG(LogMetasoundMusicClock, Verbose, TEXT("(%d) Smoothed Render Time too close to actual render time. WAITING!"), LastRenderPosSampleCount);
+		}
+		return;
+	}
+
+	if (SmoothedTempoMapTick != SmoothedTick && CurrentMapChain->LoopLengthTicks <= 0)
+	{
+		// The clock is offset from its SongMaps (eg. It is the output of a 
+		// clock offset node.)
+		UpdateCurrentTicksForOffsetClock(SmoothedTick, SmoothedTempoMapTick);
+	}
+	else
+	{
+		// The clock is looping or monotonically increasing... so we deal with it this way...
+		UpdateCurrentTicksForLoopingOrMonotonicClock(SmoothedTick, SmoothedTempoMapTick);
+	}
+
+	UpdateMusicPlaybackRate(CurrentSmoothedAudioRenderSongPos.Tempo, SpeedAtRawRenderTime, CurrentSmoothedAudioRenderSongPos.TimeSigNumerator, CurrentSmoothedAudioRenderSongPos.TimeSigDenominator);
+}
+
+void Harmonix::FMetasoundMusicClock::SetClockHistory(const UMidiClockUpdateSubsystem::FClockHistoryPtr& NewHistory)
+{
+	if (NewHistory != ClockHistory)
+	{
+		ClockHistory = NewHistory;
+		// We null this out because ultimately we get the CurrentMapChain from the ClockHistory so if it ever
+		// changes, we need to get a new CurrentMapChain assigned.
+		CurrentMapChain = nullptr;
+	}
+}
+
+void Harmonix::FMetasoundMusicClock::UpdateCurrentTicksForOffsetClock(float SmoothedTick, float SmoothedTempoMapTick)
+{
+	float SmoothedPositionMs = CurrentMapChain->SongMaps->TickToMs(SmoothedTick);
+	// We are behind the actual render time because of the lag we introduce to have enough
+	// history... SO... Push forward to get a time that is approx. where the renderer is.
+	SmoothedPositionMs += RenderSmoothingLagSeconds * 1000.0f;
+
+	// Calculate the song position AND get the "local tick" for the Audio Render timebase
+	SetCurrentSongPos(ECalibratedMusicTimebase::AudioRenderTime, CalculateSongPosAtMsForOffsetClock(SmoothedPositionMs,
+		SmoothedTick - SmoothedTempoMapTick,
+		AudioRenderState.LocalTick,
+		AudioRenderSeekDetected));
+	float SmoothedTempoMapMs = CurrentMapChain->SongMaps->TickToMs(SmoothedTempoMapTick);
+	AudioRenderState.TempoMapMs = SmoothedTempoMapMs + (RenderSmoothingLagSeconds * 1000.0f);
+	AudioRenderState.TempoMapTick = CurrentMapChain->SongMaps->MsToTick(AudioRenderState.TempoMapMs);
+
+	// Now the time the user should actually be "experiencing" (ie "hearing") can be calculated as an offset
+	// from the smooth audio rendering time...
+	float LookBehindLagMs = FHarmonixModule::Get().GetMeasuredUserExperienceAndReactionToAudioRenderOffsetMs();
+	float Ms = SmoothedPositionMs;
+	Ms -= LookBehindLagMs;
+	// Calculate the song position AND get the "local tick" for the Player Experience timebase
+	SetCurrentSongPos(ECalibratedMusicTimebase::ExperiencedTime, CalculateSongPosAtMsForOffsetClock(Ms,
+		SmoothedTick - SmoothedTempoMapTick,
+		PlayerExperienceState.LocalTick,
+		PlayerExperiencedSeekDetected));
+	PlayerExperienceState.TempoMapMs = AudioRenderState.TempoMapMs - LookBehindLagMs;
+	PlayerExperienceState.TempoMapTick = CurrentMapChain->SongMaps->MsToTick(PlayerExperienceState.TempoMapMs);
+
+	// Now the time the game should be rendering graphics for can be calculated as an offset
+	// from the smooth audio rendering time...
+	LookBehindLagMs = FHarmonixModule::Get().GetMeasuredVideoToAudioRenderOffsetMs();
+	Ms = SmoothedPositionMs;
+	Ms -= LookBehindLagMs;
+	// Calculate the song position AND get the "local tick" for the Video Render timebase
+	SetCurrentSongPos(ECalibratedMusicTimebase::VideoRenderTime, CalculateSongPosAtMsForOffsetClock(Ms,
+		SmoothedTick - SmoothedTempoMapTick,
+		VideoRenderState.LocalTick,
+		VideoRenderSeekDetected));
+	VideoRenderState.TempoMapMs = AudioRenderState.TempoMapMs - LookBehindLagMs;
+	VideoRenderState.TempoMapTick = CurrentMapChain->SongMaps->MsToTick(VideoRenderState.TempoMapMs);
+}
+
+void Harmonix::FMetasoundMusicClock::UpdateCurrentTicksForLoopingOrMonotonicClock(float SmoothedTick, float SmoothedTempoMapTick)
+{
+	AudioRenderState.TempoMapMs = CurrentMapChain->SongMaps->TickToMs(SmoothedTempoMapTick);
+	// We are behind the actual render time because of the lag we introduce to have enough
+	// history... SO... Push forward to get a time that is approx. where the renderer is.
+	AudioRenderState.TempoMapMs += RenderSmoothingLagSeconds * 1000.0f;
+
+	// Calculate the song position AND get the "local tick" for the Audio Render timebase
+	AudioRenderState.TempoMapTick = CurrentMapChain->SongMaps->MsToTick(AudioRenderState.TempoMapMs);
+	SetCurrentSongPos(ECalibratedMusicTimebase::AudioRenderTime, CalculateSongPosAtMsForLoopingOrMonotonicClock(AudioRenderState.TempoMapMs,
+		AudioRenderState.LocalTick,
+		AudioRenderSeekDetected,
+		AudioRenderLoopDetected));
+
+	// Now the time the user should actually be "experiencing" (ie "hearing") can be calculated as an offset
+	// from the smooth audio rendering time...
+	float LookBehindLagMs = FHarmonixModule::Get().GetMeasuredUserExperienceAndReactionToAudioRenderOffsetMs();
+	PlayerExperienceState.TempoMapMs = AudioRenderState.TempoMapMs - LookBehindLagMs;
+	PlayerExperienceState.TempoMapTick = CurrentMapChain->SongMaps->MsToTick(PlayerExperienceState.TempoMapMs);
+	SetCurrentSongPos(ECalibratedMusicTimebase::ExperiencedTime, CalculateSongPosAtMsForLoopingOrMonotonicClock(PlayerExperienceState.TempoMapMs,
+		PlayerExperienceState.LocalTick,
+		PlayerExperiencedSeekDetected,
+		PlayerExperiencedLoopDetected));
+
+	// Now the time the game should be rendering graphics for can be calculated as an offset
+	// from the smooth audio rendering time...
+	LookBehindLagMs = FHarmonixModule::Get().GetMeasuredVideoToAudioRenderOffsetMs();
+	VideoRenderState.TempoMapMs = AudioRenderState.TempoMapMs - LookBehindLagMs;
+	VideoRenderState.TempoMapTick = CurrentMapChain->SongMaps->MsToTick(VideoRenderState.TempoMapMs);
+	SetCurrentSongPos(ECalibratedMusicTimebase::VideoRenderTime, CalculateSongPosAtMsForLoopingOrMonotonicClock(VideoRenderState.TempoMapMs,
+		VideoRenderState.LocalTick,
+		VideoRenderSeekDetected,
+		VideoRenderLoopDetected));
+}
+
+Harmonix::FMetasoundMusicClock::EHistoryFailureType Harmonix::FMetasoundMusicClock::CalculateSmoothedTick(Metasound::FSampleCount ExpectedRenderPosSampleCount, Metasound::FSampleCount LastRenderPosSampleCount,
+	float& SmoothedLocalTick, float& SmoothedTempoMapTick, float& CurrentSpeed, HarmonixMetasound::Analysis::FMidiClockSongPositionHistory::FReadCursor& ReadCursor, float LookBehindSeconds)
+{
+	using namespace HarmonixMetasound::Analysis;
+
+	Metasound::FSampleCount LookingForSampleFrame = ExpectedRenderPosSampleCount - (LookBehindSeconds * ClockHistory->SampleRate);
+
+	int32 NumHistoryAvailable = ReadCursor.NumDataAvailable();
+	if (LookingForSampleFrame >= LastRenderPosSampleCount && NumHistoryAvailable > 1)
+	{
+		while (ReadCursor.NumDataAvailable() > 1)
+		{
+			ReadCursor.ConsumeNext();
+		}
+		NumHistoryAvailable = ReadCursor.NumDataAvailable();
+	}
+
+	if (NumHistoryAvailable == 0)
+	{
+		return EHistoryFailureType::NotEnoughDataInTheHistoryRing;
+	}
+
+	FMidiClockSongPositionHistory::FScopedItemPeekRef PeekNextRef = ReadCursor.PeekNext();
+
+	if (NumHistoryAvailable == 1 || PeekNextRef->SampleCount > LookingForSampleFrame)
+	{
+		SmoothedLocalTick = (float)PeekNextRef->UpToTick;
+		SmoothedTempoMapTick = (float)PeekNextRef->TempoMapTick;
+		CurrentSpeed = PeekNextRef->CurrentSpeed;
+		return EHistoryFailureType::None;
+	}
+
+	// OK... or same SHOULD be in the history...
+	FMidiClockSongPositionHistory::FScopedItemPeekRef PeekOneAheadRef(ReadCursor.PeekAhead(1));
+	while (PeekOneAheadRef && PeekOneAheadRef->SampleCount <= LookingForSampleFrame)
+	{
+		ReadCursor.PeekAhead(2, PeekOneAheadRef);
+		ReadCursor.PeekAhead(1, PeekNextRef);
+		ReadCursor.ConsumeNext();
+	}
+
+	// Maybe the sample BEFORE our sample is in the history, but the sample AFTER is not, so we can't lerp?
+	if (!PeekOneAheadRef)
+	{
+		return EHistoryFailureType::CaughtUpToRenderPosition;
+	}
+
+	check(LookingForSampleFrame >= PeekNextRef->SampleCount && LookingForSampleFrame < PeekOneAheadRef->SampleCount);
+
+	// We now have enough to lerp!
+	float LerpAlpha = (float)(LookingForSampleFrame - PeekNextRef->SampleCount) / (float)(PeekOneAheadRef->SampleCount - PeekNextRef->SampleCount);
+	SmoothedLocalTick = FMath::Lerp<float>((float)PeekNextRef->UpToTick, (float)PeekOneAheadRef->UpToTick, LerpAlpha);
+	SmoothedTempoMapTick = FMath::Lerp<float>((float)PeekNextRef->TempoMapTick, (float)PeekOneAheadRef->TempoMapTick, LerpAlpha);
+	CurrentSpeed = PeekNextRef->CurrentSpeed;
+
+	return EHistoryFailureType::None;
+}
diff --git a/Engine/Plugins/Runtime/Harmonix/Source/HarmonixMetasound/Private/HarmonixMetasound/MusicClock/MetasoundMusicClock.h b/Engine/Plugins/Runtime/Harmonix/Source/HarmonixMetasound/Private/HarmonixMetasound/MusicClock/MetasoundMusicClock.h
new file mode 100644
--- /dev/null
+++ b/Engine/Plugins/Runtime/Harmonix/Source/HarmonixMetasound/Private/HarmonixMetasound/MusicClock/MetasoundMusicClock.h
@@ -0,0 +1,151 @@
+// Copyright Epic Games, Inc. All Rights Reserved.
+#pragma once
+#include "Harmonix/LocalMinimumMagnitudeTracker.h"
+#include "HarmonixMetasound/MusicClock/MusicClock.h"
+#include "HarmonixMetasound/Subsystems/MidiClockUpdateSubsystem.h"
+#include "HarmonixMetasound/Components/MusicClockTypes.h"
+#include "MetasoundGeneratorHandle.h"
+#include "UObject/StrongObjectPtr.h"
+
+struct FMetasoundMusicClockSettings;
+
+namespace HarmonixMetasound::Analysis
+{
+	struct FSongMapChain;
+}
+
+namespace Harmonix
+{
+class FMetasoundMusicClock : public FMusicClock
+{
+public:
+	FMetasoundMusicClock(UObject* WorldContextObj, const FMetasoundMusicClockSettings& Settings)
+		: FMusicClock(WorldContextObj, Settings) {}
+
+	virtual TOptional<FVector> TryGetAudioSourceLocation() const override;
+
+	virtual bool CalculateSongPosWithOffset(float MsOffset, ECalibratedMusicTimebase Timebase, FMidiSongPos& OutResult) const override;
+
+	virtual void Disconnect() override;
+	virtual bool RefreshCurrentSongPos() override;
+	virtual void OnStart() override;
+	virtual void OnPause() override;
+	virtual void OnContinue() override;
+	virtual void OnStop() override;
+	virtual const ISongMapEvaluator* GetCurrentSongMapEvaluator() const override;
+
+	virtual bool LoopedThisFrame(ECalibratedMusicTimebase Timebase) const override;
+	virtual bool SeekedThisFrame(ECalibratedMusicTimebase Timebase) const override;
+
+#if !UE_BUILD_SHIPPING
+	virtual FString GetDisplayName() const override;
+#endif
+
+	bool ConnectToAudioComponentsMetasound(
+		UAudioComponent* InAudioComponent,
+		FName MetasoundOuputPinName = "MIDI Clock",
+		UMetasoundGeneratorHandle::FOnAttached::FDelegate OnGeneratorAttachedCallback = UMetasoundGeneratorHandle::FOnAttached::FDelegate::CreateLambda([](){}),
+		UMetasoundGeneratorHandle::FOnDetached::FDelegate OnGeneratorDetachedCallback = UMetasoundGeneratorHandle::FOnDetached::FDelegate::CreateLambda([](){}));
+
+	bool RunPastMusicEnd = false;
+
+protected:
+	void OnGeneratorAttached();
+	void OnGeneratorDetached();
+	void OnGraphSet();
+	void OnGeneratorIOUpdatedWithChanges(const TArray<Metasound::FVertexInterfaceChange>& VertexInterfaceChanges);
+
+private:
+	FName MetasoundOutputName;
+	// We can keep a week reference to this because our "owner" is a UClass and
+	// also has a reference to it...
+	TWeakObjectPtr<UAudioComponent> AudioComponentToWatch;
+	// We need a strong object ptr to this next thing since we will be the only one holding a reference to it...
+	TStrongObjectPtr<UMetasoundGeneratorHandle> CurrentGeneratorHandle;
+	TWeakPtr<Metasound::FMetasoundGenerator> UnderlyingGenerator;
+
+	Metasound::Frontend::FAnalyzerAddress MidiSongPosAnalyzerAddress;
+
+	UMidiClockUpdateSubsystem::FClockHistoryPtr ClockHistory;
+	HarmonixMetasound::Analysis::FMidiClockSongPositionHistory::FReadCursor SmoothedAudioRenderClockHistoryCursor;
+	HarmonixMetasound::Analysis::FMidiClockSongPositionHistory::FReadCursor SmoothedPlayerExperienceClockHistoryCursor;
+	HarmonixMetasound::Analysis::FMidiClockSongPositionHistory::FReadCursor SmoothedVideoRenderClockHistoryCursor;
+
+	TSharedPtr<const HarmonixMetasound::Analysis::FSongMapChain> CurrentMapChain;
+
+	bool bRunning = false;
+	double FreeRunStartTimeSecs = 0.0;
+	bool WasEverConnected = false;
+	float SongPosOffsetMs = 0.0f;
+	int32 LastTickSeen = 0;
+	float RenderSmoothingLagSeconds = 0.030f;
+
+	double RenderStartWallClockTimeSeconds = 0.0;
+	double LastRefreshWallClockTimeSeconds = 0.0;
+	double DeltaSecondsBetweenRefreshes = 0.0;
+
+	Metasound::FSampleCount RenderStartSampleCount {0};
+
+	static const int kFramesOfErrorHistory = 10;
+	FLocalMinimumMagnitudeTracker<double, kFramesOfErrorHistory> ErrorTracker;
+	double SyncSpeed = 1.0;
+
+	struct FPerTimebaseSmoothedClockState
+	{
+		float TempoMapMs = 0.0f;
+		float TempoMapTick = 0.0f;
+		float LocalTick = 0.0f;
+	};
+
+	FPerTimebaseSmoothedClockState AudioRenderState;
+	FPerTimebaseSmoothedClockState PlayerExperienceState;
+	FPerTimebaseSmoothedClockState VideoRenderState;
+
+	UMetasoundGeneratorHandle::FOnAttached::FDelegate OnAttachedDelegate;
+	UMetasoundGeneratorHandle::FOnDetached::FDelegate OnDetachedDelegate;
+
+	FDelegateHandle GeneratorAttachedCallbackHandle;
+	FDelegateHandle GeneratorDetachedCallbackHandle;
+	FDelegateHandle GeneratorIOUpdatedCallbackHandle;
+	FDelegateHandle GraphChangedCallbackHandle;
+
+	bool AudioRenderSeekDetected = false;
+	bool AudioRenderLoopDetected = false;
+	bool PlayerExperiencedSeekDetected = false;
+	bool PlayerExperiencedLoopDetected = false;
+	bool VideoRenderSeekDetected = false;
+	bool VideoRenderLoopDetected = false;
+
+	FMidiSongPos CalculateSongPosAtMsForLoopingOrMonotonicClock(float AbsoluteMs, float& PositionTick, bool& SeekDetected, bool& LoopDetected) const;
+	FMidiSongPos CalculateSongPosAtMsForOffsetClock(float PositionMs, float ClockTickOffsetFromDrivingClock, float& PositionTick, bool& SeekDetected) const;
+	bool CheckForSeek(float FirstTick, float NextTick, float CurrentTempo, int32 TicksPerQuarter) const;
+
+	void UpdateCurrentTicksForOffsetClock(float SmoothedTick, float SmoothedTempoMapTick);
+	void UpdateCurrentTicksForLoopingOrMonotonicClock(float SmoothedTick, float SmoothedTempoMapTick);
+
+	bool AttemptToConnectToAudioComponentsMetasound(
+		UAudioComponent* InAudioComponent,
+		UMetasoundGeneratorHandle::FOnAttached::FDelegate OnGeneratorAttachedCallback,
+		UMetasoundGeneratorHandle::FOnDetached::FDelegate OnGeneratorDetachedCallback);
+	void DetachAllCallbacks();
+	void RefreshCurrentSongPosFromWallClock();
+	void RefreshCurrentSongPosFromHistory();
+	// always go through this to change ClockHistory
+	void SetClockHistory(const UMidiClockUpdateSubsystem::FClockHistoryPtr&);
+
+	enum class EHistoryFailureType : uint8
+	{
+		None,
+		NotEnoughDataInTheHistoryRing,
+		NotEnoughHistory,
+		LookingForTimeInTheFutureOfWhatHasEvenRendered,
+		CaughtUpToRenderPosition
+	};
+
+	EHistoryFailureType CalculateSmoothedTick(Metasound::FSampleCount ExpectedRenderPosSampleCount, Metasound::FSampleCount LastRenderPosSampleCount,
+		float& SmoothedTick, float& SmoothedTempoMapTick, float& CurrentSpeed, HarmonixMetasound::Analysis::FMidiClockSongPositionHistory::FReadCursor& ReadCursor,
+		float LookBehindSeconds);
+
+	static FString HistoryFailureTypeToString(FMetasoundMusicClock::EHistoryFailureType Error);
+};
+}
\ No newline at end of file
diff --git a/Engine/Plugins/Runtime/Harmonix/Source/HarmonixMetasound/Private/HarmonixMetasound/MusicClock/MusicClock.cpp b/Engine/Plugins/Runtime/Harmonix/Source/HarmonixMetasound/Private/HarmonixMetasound/MusicClock/MusicClock.cpp
new file mode 100644
--- /dev/null
+++ b/Engine/Plugins/Runtime/Harmonix/Source/HarmonixMetasound/Private/HarmonixMetasound/MusicClock/MusicClock.cpp
@@ -0,0 +1,234 @@
+// Copyright Epic Games, Inc. All Rights Reserved.
+
+#include "HarmonixMetasound/MusicClock/MusicClock.h"
+
+#include "Engine/World.h"
+#include "HarmonixMetasound/Components/MusicClockTypes.h"
+
+void Harmonix::FMusicClock::ResetDefaultSongMaps(float DefaultTempo, int32 DefaultTimeSigNum, int32 DefaultTimeSigDenom)
+{
+	DefaultMaps.EmptyAllMaps();
+	DefaultMaps.Init(Harmonix::Midi::Constants::GTicksPerQuarterNoteInt);
+	DefaultMaps.GetTempoMap().AddTempoInfoPoint(Harmonix::Midi::Constants::BPMToMidiTempo(DefaultTempo), 0);
+	DefaultMaps.GetBarMap().AddTimeSignatureAtBarIncludingCountIn(0, DefaultTimeSigNum, DefaultTimeSigDenom);
+}
+
+void Harmonix::FMusicClock::UpdateBeatsAndBarsPerSecond()
+{
+	// Midi quarter note to beat conversion.
+	const float BPM = GetBeatsPerMinute();
+	CurrentBeatsPerSecond = (BPM / 60.0) * CurrentClockAdvanceRate;
+	CurrentBarsPerSecond = (TimeSignatureNum > 0) ? (CurrentBeatsPerSecond / TimeSignatureNum) : 0;
+}
+
+Harmonix::FMusicClock::FMusicClock(UObject* WorldContextObject, const FMusicClockSettingsBase& Settings)
+{
+	ResetDefaultSongMaps(Settings.DefaultTempo, Settings.DefaultTimeSigNumerator, Settings.DefaultTimeSigDenominator);
+	ContextObj = TWeakObjectPtr<UObject>(WorldContextObject);
+}
+
+Harmonix::FMusicClock::~FMusicClock()
+{
+}
+
+bool Harmonix::FMusicClock::CalculateSongPosWithOffset(float MsOffset, ECalibratedMusicTimebase Timebase, FMidiSongPos& OutResult) const
+{
+	const ISongMapEvaluator* SongMapEvaluator = GetCurrentSongMapEvaluator();
+	if (!SongMapEvaluator)
+	{
+		return false;
+	}
+
+	const FMidiSongPos& SongPos = GetCurrentSongPos(Timebase);
+	if (!SongPos.IsValid())
+	{
+		return false;
+	}
+
+	OutResult.SetByTime((SongPos.SecondsIncludingCountIn * 1000.0f) + MsOffset, *SongMapEvaluator);
+	return true;
+}
+
+double Harmonix::FMusicClock::GetWallClockTime() const
+{
+	if (TStrongObjectPtr<UObject> Obj = ContextObj.Pin())
+	{
+		if (Obj->GetWorld())
+		{
+			return Obj->GetWorld()->GetRealTimeSeconds();
+		}
+	}
+
+	if (GWorld)
+	{
+		return GWorld->GetRealTimeSeconds();
+	}
+	
+	return 0.0;
+}
+
+
+const FMidiSongPos& Harmonix::FMusicClock::GetCurrentSongPos(ECalibratedMusicTimebase Timebase) const
+{
+	return CurrentSongPos[GetArrayIndexForTimebase(Timebase)];
+}
+
+const FMidiSongPos& Harmonix::FMusicClock::GetPreviousSongPos(ECalibratedMusicTimebase Timebase) const
+{
+	return PrevSongPos[GetArrayIndexForTimebase(Timebase)];
+}
+
+float Harmonix::FMusicClock::GetDeltaBarF(ECalibratedMusicTimebase Timebase) const
+{
+	return DeltaBarF[GetArrayIndexForTimebase(Timebase)];
+}
+
+float Harmonix::FMusicClock::GetDeltaBeatF(ECalibratedMusicTimebase Timebase) const
+{
+	return DeltaBeatF[GetArrayIndexForTimebase(Timebase)];
+}
+
+void Harmonix::FMusicClock::Start()
+{
+	OnStart();
+
+	for (ECalibratedMusicTimebase Timebase : TEnumRange<ECalibratedMusicTimebase>())
+	{
+		CurrentSongPos[(int32)Timebase].Reset();
+	}
+
+	MusicClockState = EMusicClockState::Running;
+}
+
+void Harmonix::FMusicClock::Continue()
+{
+	if (MusicClockState != EMusicClockState::Paused)
+	{
+		return;
+	}
+	
+	OnContinue();
+	MusicClockState = EMusicClockState::Running;
+}
+
+void Harmonix::FMusicClock::Pause()
+{
+	if (MusicClockState != EMusicClockState::Running)
+	{
+		return;
+	}
+
+	OnPause();
+	MusicClockState = EMusicClockState::Paused;
+}
+
+void Harmonix::FMusicClock::Stop()
+{
+	OnStop();
+	MusicClockState = EMusicClockState::Stopped;
+
+	for (ECalibratedMusicTimebase Timebase : TEnumRange<ECalibratedMusicTimebase>())
+	{
+		CurrentSongPos[(int32)Timebase].Reset();
+	}
+}
+
+void Harmonix::FMusicClock::UpdateMusicPlaybackRate(float InTempo, float InClockAdvanceRate, int32 InTimeSigNum, int32 InTimeSigDenom)
+{
+	bool bHasAnyTimingChanges = false;
+	if (!FMath::IsNearlyEqual(Tempo, InTempo))
+	{
+		Tempo = InTempo;
+		bHasAnyTimingChanges = true;
+	}
+
+	if (!FMath::IsNearlyEqual(CurrentClockAdvanceRate, InClockAdvanceRate))
+	{
+		CurrentClockAdvanceRate = InClockAdvanceRate;
+		bHasAnyTimingChanges = true;
+	}
+	
+	if (TimeSignatureNum != InTimeSigNum || TimeSignatureDenom != InTimeSigDenom)
+	{
+		TimeSignatureNum = InTimeSigNum;
+		TimeSignatureDenom = InTimeSigDenom;
+		bHasAnyTimingChanges = true;
+	}
+
+	// only update bar/beat durations if there have been any other timing changes
+	if (bHasAnyTimingChanges)
+	{
+		UpdateBeatsAndBarsPerSecond();
+	}
+}
+
+float Harmonix::FMusicClock::GetBeatsPerMinute() const
+{
+	// if we have a tempo set, we should have a time signature.
+	ensure(TimeSignatureDenom != 0 || Tempo == 0);
+    // Tempo is MidiTempo, which is in quarter notes per minute.
+	return Tempo * TimeSignatureDenom / 4.0f; 
+}
+
+void Harmonix::FMusicClock::UpdateClockForGameFrame()
+{
+	//	Not for use outside the game thread.
+	if (ensureMsgf(
+		IsInGameThread(),
+		TEXT("%hs called from non-game thread.  This is not supported!"), __FUNCTION__) == false)
+	{
+		return;
+	}
+
+	if (GFrameCounter == LastUpdateFrame)
+	{
+		return;
+	}
+
+	//	Run the actual clock update.
+	if (GetState() == EMusicClockState::Running)
+	{
+		for (ECalibratedMusicTimebase Timebase : TEnumRange<ECalibratedMusicTimebase>())
+		{
+			PrevSongPos[(int32)Timebase] = CurrentSongPos[(int32)Timebase];
+		}
+
+		if (RefreshCurrentSongPos())
+		{
+			for (ECalibratedMusicTimebase Timebase : TEnumRange<ECalibratedMusicTimebase>())
+			{
+				DeltaBarF[(int32)Timebase] = CurrentSongPos[(int32)Timebase].BarsIncludingCountIn - PrevSongPos[(int32)Timebase].BarsIncludingCountIn;
+				DeltaBeatF[(int32)Timebase] = CurrentSongPos[(int32)Timebase].BeatsIncludingCountIn - PrevSongPos[(int32)Timebase].BeatsIncludingCountIn;
+			}
+
+			LastUpdateFrame = GFrameCounter;
+		}
+	}
+}
+
+int32 Harmonix::FMusicClock::GetArrayIndexForTimebase(ECalibratedMusicTimebase Timebase) const
+{
+	int32 TimebaseIndex = (int32)Timebase;
+	if (!ensure(TimebaseIndex >= 0 && TimebaseIndex < (int32)ECalibratedMusicTimebase::Count))
+	{
+		// minor robustness, return something even if there's an illegal value.
+		return (int32)ECalibratedMusicTimebase::VideoRenderTime;
+	}
+
+	return TimebaseIndex;
+}
+
+void Harmonix::FMusicClock::SetCurrentSongPos(ECalibratedMusicTimebase Timebase, const FMidiSongPos& InSongPos)
+{
+	CurrentSongPos[GetArrayIndexForTimebase(Timebase)] = InSongPos;
+}
+
+void Harmonix::FMusicClock::SetCurrentSongPosByTime(ECalibratedMusicTimebase Timebase, float InElapsedMs, const ISongMapEvaluator& Map)
+{
+	CurrentSongPos[GetArrayIndexForTimebase(Timebase)].SetByTime(InElapsedMs, Map);
+}
+
+void Harmonix::FMusicClock::SetCurrentSongPosByTick(ECalibratedMusicTimebase Timebase, float InTick, const ISongMapEvaluator& Map)
+{
+	CurrentSongPos[GetArrayIndexForTimebase(Timebase)].SetByTick(InTick, Map);
+}
diff --git a/Engine/Plugins/Runtime/Harmonix/Source/HarmonixMetasound/Private/HarmonixMetasound/MusicClock/RuntimeMusicClock.cpp b/Engine/Plugins/Runtime/Harmonix/Source/HarmonixMetasound/Private/HarmonixMetasound/MusicClock/RuntimeMusicClock.cpp
new file mode 100644
--- /dev/null
+++ b/Engine/Plugins/Runtime/Harmonix/Source/HarmonixMetasound/Private/HarmonixMetasound/MusicClock/RuntimeMusicClock.cpp
@@ -0,0 +1,85 @@
+// Copyright Epic Games, Inc. All Rights Reserved.
+#include "HarmonixMetasound/MusicClock/RuntimeMusicClock.h"
+#include "Harmonix.h"
+#include "HarmonixMetasound/MusicClock/WorldTimeSourceController.h"
+#include "HarmonixMidi/MidiFile.h"
+
+Harmonix::FRuntimeMusicClock::FRuntimeMusicClock(UObject* WorldContextObj, const FWallClockMusicClockSettings& Settings)
+	: FMusicClock(WorldContextObj, Settings)
+{
+	TempoMapMidi = Settings.TempoMap;
+	WorldTimeSource = MakeShared<FWorldTimeSourceController>(WorldContextObj);
+}
+
+TOptional<FVector> Harmonix::FRuntimeMusicClock::TryGetAudioSourceLocation() const
+{
+	return {};
+}
+
+void Harmonix::FRuntimeMusicClock::Disconnect()
+{
+	TempoMapMidi = nullptr;
+}
+
+bool Harmonix::FRuntimeMusicClock::RefreshCurrentSongPos()
+{
+	check(IsInGameThread());
+	WorldTimeSource->Update();
+
+	double SourceTime = WorldTimeSource->GetCurrentTime();
+	CurrentClockAdvanceRate = WorldTimeSource->GetSpeed();
+	const ISongMapEvaluator* Maps = GetCurrentSongMapEvaluator();
+	check(Maps);
+
+	// Set Experienced time here, as this code is running on the GameThread, near the same time as input is processed
+	SetCurrentSongPosByTime(ECalibratedMusicTimebase::ExperiencedTime, (float)(SourceTime * 1000.0), *Maps);
+	SetCurrentSongPosByTime(ECalibratedMusicTimebase::RawAudioRenderTime, GetCurrentSongPos(ECalibratedMusicTimebase::ExperiencedTime).SecondsIncludingCountIn * 1000.0f + FHarmonixModule::GetMeasuredUserExperienceAndReactionToAudioRenderOffsetMs(), *Maps);
+	SetCurrentSongPosByTime(ECalibratedMusicTimebase::VideoRenderTime, GetCurrentSongPos(ECalibratedMusicTimebase::RawAudioRenderTime).SecondsIncludingCountIn * 1000.0f - FHarmonixModule::GetMeasuredVideoToAudioRenderOffsetMs(), *Maps);
+	SetCurrentSongPos(ECalibratedMusicTimebase::AudioRenderTime, GetCurrentSongPos(ECalibratedMusicTimebase::RawAudioRenderTime));
+
+	const FMidiSongPos& CurrentSmoothedAudioRenderSongPos = GetCurrentSongPos(ECalibratedMusicTimebase::AudioRenderTime);
+	UpdateMusicPlaybackRate(CurrentSmoothedAudioRenderSongPos.Tempo, CurrentClockAdvanceRate, CurrentSmoothedAudioRenderSongPos.TimeSigNumerator, CurrentSmoothedAudioRenderSongPos.TimeSigDenominator);
+
+	return true;
+}
+
+void Harmonix::FRuntimeMusicClock::OnStart()
+{
+	check(IsInGameThread());
+	WorldTimeSource->Start();
+}
+
+void Harmonix::FRuntimeMusicClock::OnPause()
+{
+	check(IsInGameThread());
+	WorldTimeSource->Pause();
+}
+
+void Harmonix::FRuntimeMusicClock::OnContinue()
+{
+	check(IsInGameThread());
+	WorldTimeSource->Continue();
+}
+
+void Harmonix::FRuntimeMusicClock::OnStop()
+{
+	check(IsInGameThread());
+	WorldTimeSource->Stop();
+}
+
+const ISongMapEvaluator* Harmonix::FRuntimeMusicClock::GetCurrentSongMapEvaluator() const
+{
+	check(IsInGameThread());
+	if (TempoMapMidi.IsValid())
+	{
+		return TempoMapMidi->GetSongMaps();
+	}
+	return &DefaultMaps;
+}
+
+#if !UE_BUILD_SHIPPING
+FString Harmonix::FRuntimeMusicClock::GetDisplayName() const
+{
+	return FString::Printf(TEXT("RuntimeMusicClock: %s: %.3f bpm"), *WorldTimeSource->GetDisplayName(), GetBeatsPerMinute());
+}
+#endif
diff --git a/Engine/Plugins/Runtime/Harmonix/Source/HarmonixMetasound/Private/HarmonixMetasound/MusicClock/RuntimeMusicClock.h b/Engine/Plugins/Runtime/Harmonix/Source/HarmonixMetasound/Private/HarmonixMetasound/MusicClock/RuntimeMusicClock.h
new file mode 100644
--- /dev/null
+++ b/Engine/Plugins/Runtime/Harmonix/Source/HarmonixMetasound/Private/HarmonixMetasound/MusicClock/RuntimeMusicClock.h
@@ -0,0 +1,48 @@
+// Copyright Epic Games, Inc. All Rights Reserved.
+#pragma once
+
+#include "HarmonixMetasound/MusicClock/MusicClock.h"
+#include "HarmonixMetasound/Subsystems/MidiClockUpdateSubsystem.h"
+#include "HarmonixMetasound/Components/MusicClockTypes.h"
+#include "MetasoundGeneratorHandle.h"
+
+namespace Harmonix
+{
+	
+class FWorldTimeSourceController;
+
+/**
+ * A Music Clock that uses an external "Time Source" and maps the time to a given "Tempo Midi Map"
+ *
+ * NOTE: Currently, this clock only takes a MIDI Map and uses a "Real Time" Source.
+ * TODO: Make this usable with other Time Sources so that it's not simply a "Wall Clock"
+ */
+class FRuntimeMusicClock : public FMusicClock
+{
+public:
+	FRuntimeMusicClock(UObject* WorldContextObj, const FWallClockMusicClockSettings& Settings);
+
+	virtual TOptional<FVector> TryGetAudioSourceLocation() const override;
+
+	virtual void Disconnect() override;
+	virtual bool RefreshCurrentSongPos() override;
+	virtual void OnStart() override;
+	virtual void OnPause() override;
+	virtual void OnContinue() override;
+	virtual void OnStop() override;
+	virtual const ISongMapEvaluator* GetCurrentSongMapEvaluator() const override;
+	
+	virtual bool LoopedThisFrame(ECalibratedMusicTimebase Timebase) const override { return false; }
+	virtual bool SeekedThisFrame(ECalibratedMusicTimebase Timebase) const override { return false; }
+
+#if !UE_BUILD_SHIPPING
+	virtual FString GetDisplayName() const override;
+#endif
+
+private:
+	TWeakObjectPtr<UMidiFile> TempoMapMidi;
+
+	// TODO: Replace with generic ITimeSource once Transport Controls on the Music Clock are removed
+	TSharedPtr<FWorldTimeSourceController> WorldTimeSource;
+};
+}
\ No newline at end of file
diff --git a/Engine/Plugins/Runtime/Harmonix/Source/HarmonixMetasound/Private/HarmonixMetasound/MusicClock/WorldTimeSourceController.cpp b/Engine/Plugins/Runtime/Harmonix/Source/HarmonixMetasound/Private/HarmonixMetasound/MusicClock/WorldTimeSourceController.cpp
new file mode 100644
--- /dev/null
+++ b/Engine/Plugins/Runtime/Harmonix/Source/HarmonixMetasound/Private/HarmonixMetasound/MusicClock/WorldTimeSourceController.cpp
@@ -0,0 +1,143 @@
+// Copyright Epic Games, Inc. All Rights Reserved.
+#include "HarmonixMetasound/MusicClock/WorldTimeSourceController.h"
+
+#include "Engine/World.h"
+#include "Misc/App.h"
+
+Harmonix::FWorldTimeSourceController::FWorldTimeSourceController(UObject* InContextObj)
+	: ContextObj(InContextObj)
+{
+}
+
+FString Harmonix::FWorldTimeSourceController::GetDisplayName() const
+{
+	return TEXT("RealTimeSourceController");
+}
+
+ECalibratedMusicTimebase Harmonix::FWorldTimeSourceController::GetTimebase() const
+{
+	return ECalibratedMusicTimebase::VideoRenderTime;
+}
+
+double Harmonix::FWorldTimeSourceController::GetDeltaTime() const
+{
+	if (TStrongObjectPtr<UObject> Obj = ContextObj.Pin())
+	{
+		if (const UWorld* World = Obj->GetWorld())
+		{
+			FGameTime GameTime = World->GetTime();
+			return GameTime.GetDeltaWorldTimeSeconds();
+		}
+	}
+
+	return 0.0;
+}
+
+
+void Harmonix::FWorldTimeSourceController::Update()
+{
+	double DeltaTime = GetDeltaTime();
+	switch (TransportRequest)
+	{
+	default:
+		LatestSourceEvent = ESourceEvent::None;
+		break;
+	case ETransportRequest::Stop:
+		CurrentSourceState = ESourceState::Stopped;
+		LatestSourceEvent = ESourceEvent::Stop;
+		RunTimeSeconds = 0.0;
+		break;
+	case ETransportRequest::Play:
+		if (CurrentSourceState == ESourceState::Stopped)
+		{
+			CurrentSourceState = ESourceState::Running;
+			LatestSourceEvent = ESourceEvent::Start;
+		}
+		break;
+	case ETransportRequest::Pause:
+		if (CurrentSourceState == ESourceState::Running)
+		{
+			CurrentSourceState = ESourceState::Paused;
+			LatestSourceEvent = ESourceEvent::Pause;
+		}
+		break;
+	case ETransportRequest::Continue:
+		if (CurrentSourceState == ESourceState::Paused)
+		{
+			CurrentSourceState = ESourceState::Running;
+			LatestSourceEvent = ESourceEvent::Continue;
+		}
+		break;
+	case ETransportRequest::Seek:
+		if (CurrentSourceState != ESourceState::Stopped)
+		{
+			RunTimeSeconds = SeekRequestTime;
+			LatestSourceEvent = ESourceEvent::Seek;
+		}
+		break;
+	}
+
+	if (CurrentSourceState == ESourceState::Running)
+	{
+		RunTimeSeconds += DeltaTime;
+		if (LatestSourceEvent == ESourceEvent::None)
+		{
+			LatestSourceEvent = ESourceEvent::Advance;
+		}
+	}
+}
+
+Harmonix::ESourceState Harmonix::FWorldTimeSourceController::GetCurrentSourceState() const
+{
+	return CurrentSourceState;
+}
+
+Harmonix::ESourceEvent Harmonix::FWorldTimeSourceController::GetLatestSourceEvent() const
+{
+	return LatestSourceEvent;
+}
+
+double Harmonix::FWorldTimeSourceController::GetCurrentTime() const
+{
+	return RunTimeSeconds;
+}
+
+float Harmonix::FWorldTimeSourceController::GetSpeed() const
+{
+	if (TStrongObjectPtr<UObject> Obj = ContextObj.Pin())
+	{
+		if (const UWorld* World = Obj->GetWorld())
+		{
+			FGameTime GameTime = World->GetTime();
+			return GameTime.GetTimeDilation();
+		}
+	}
+
+	return 0.0;
+}
+
+void Harmonix::FWorldTimeSourceController::Start()
+{
+	TransportRequest = ETransportRequest::Play;
+}
+
+void Harmonix::FWorldTimeSourceController::Stop()
+{
+	TransportRequest = ETransportRequest::Stop;
+}
+
+void Harmonix::FWorldTimeSourceController::Pause()
+{
+	TransportRequest = ETransportRequest::Pause;
+}
+
+void Harmonix::FWorldTimeSourceController::Continue()
+{
+	TransportRequest = ETransportRequest::Continue;
+}
+
+void Harmonix::FWorldTimeSourceController::Seek(double InSeekTime)
+{
+	TransportRequest = ETransportRequest::Seek;
+	SeekRequestTime = InSeekTime;
+}
\ No newline at end of file
diff --git a/Engine/Plugins/Runtime/Harmonix/Source/HarmonixMetasound/Public/HarmonixMetasound/Components/MusicClockComponent.h b/Engine/Plugins/Runtime/Harmonix/Source/HarmonixMetasound/Public/HarmonixMetasound/Components/MusicClockComponent.h
--- a/Engine/Plugins/Runtime/Harmonix/Source/HarmonixMetasound/Public/HarmonixMetasound/Components/MusicClockComponent.h
+++ b/Engine/Plugins/Runtime/Harmonix/Source/HarmonixMetasound/Public/HarmonixMetasound/Components/MusicClockComponent.h
@@ -4,6 +4,7 @@
 #include "Components/ActorComponent.h"
 #include "Evaluation/IMovieSceneCustomClockSource.h"
 #include "Harmonix/MusicalTimebase.h"
+#include "HarmonixMetasound/MusicClock/MusicClockState.h"
 #include "HarmonixMidi/MidiSongPos.h"
 #include "HarmonixMidi/SongMaps.h"
 #include "MusicEnvironmentClockSource.h"
@@ -21,14 +22,6 @@ struct FMusicalTimeSpan;
 class UAudioComponent;
 class FMusicClockDriverBase;
 
-UENUM(BlueprintType)
-enum class EMusicClockState : uint8
-{
-	Stopped,
-	Paused,
-	Running,
-};
-
 UENUM(BlueprintType)
 enum class EMusicClockDriveMethod : uint8
 {
diff --git a/Engine/Plugins/Runtime/Harmonix/Source/HarmonixMetasound/Public/HarmonixMetasound/MusicClock/MusicClock.h b/Engine/Plugins/Runtime/Harmonix/Source/HarmonixMetasound/Public/HarmonixMetasound/MusicClock/MusicClock.h
new file mode 100644
--- /dev/null
+++ b/Engine/Plugins/Runtime/Harmonix/Source/HarmonixMetasound/Public/HarmonixMetasound/MusicClock/MusicClock.h
@@ -0,0 +1,104 @@
+// Copyright Epic Games, Inc. All Rights Reserved.
+#pragma once
+
+#include "Harmonix/MusicalTimebase.h"
+#include "HarmonixMetasound/MusicClock/MusicClockState.h"
+#include "HarmonixMidi/MidiSongPos.h"
+#include "HarmonixMidi/SongMaps.h"
+
+#define UE_API HARMONIXMETASOUND_API
+
+struct FMusicClockSettingsBase;
+
+namespace Harmonix
+{
+
+class FMusicClock
+{
+public:
+	UE_API FMusicClock(UObject* WorldContextObject, const FMusicClockSettingsBase& Settings);
+	UE_API virtual ~FMusicClock();
+
+	virtual TOptional<FVector> TryGetAudioSourceLocation() const = 0;
+
+	UE_API void UpdateClockForGameFrame();
+
+	UE_API virtual bool CalculateSongPosWithOffset(float MsOffset, ECalibratedMusicTimebase Timebase, FMidiSongPos& OutResult) const;
+
+	UE_API double GetWallClockTime() const;
+
+	EMusicClockState GetState() const { return MusicClockState; }
+	UE_API const FMidiSongPos& GetCurrentSongPos(ECalibratedMusicTimebase Timebase) const;
+	UE_API const FMidiSongPos& GetPreviousSongPos(ECalibratedMusicTimebase Timebase) const;
+	UE_API float GetDeltaBarF(ECalibratedMusicTimebase Timebase) const;
+	UE_API float GetDeltaBeatF(ECalibratedMusicTimebase Timebase) const;
+
+	UE_API void Start();
+	UE_API void Pause();
+	UE_API void Continue();
+	UE_API void Stop();
+	
+	virtual const ISongMapEvaluator* GetCurrentSongMapEvaluator() const = 0;
+
+	virtual bool LoopedThisFrame(ECalibratedMusicTimebase Timebase) const = 0;
+	virtual bool SeekedThisFrame(ECalibratedMusicTimebase Timebase) const = 0;
+
+	virtual void Disconnect() = 0;
+
+	UE_API void UpdateMusicPlaybackRate(float InTempo, float InClockAdvanceRate, int32 InTimeSigNum, int32 InTimeSigDenom);
+
+	// This returns our BeatsPerMinute which is our true beats, not quarter notes.
+	UE_API float GetBeatsPerMinute() const;
+
+#if !UE_BUILD_SHIPPING
+	// Short user readable description good for logging. Should include name of song or sound source or whatnot. 
+	virtual FString GetDisplayName() const = 0;
+#endif
+
+	FSongMaps DefaultMaps;
+	
+	// Note that tempo is MidiTempo, quarter notes per minute.
+	float Tempo = 0.0f;
+	int TimeSignatureNum = 0;
+	int TimeSignatureDenom = 0;
+
+	float CurrentBeatsPerSecond = 0.0f;
+	float CurrentBarsPerSecond = 0.0f;
+	float CurrentClockAdvanceRate = 1.0f;
+
+	FMidiSongPos CurrentSongPos[(int32)ECalibratedMusicTimebase::Count];
+	FMidiSongPos PrevSongPos[(int32)ECalibratedMusicTimebase::Count];
+
+	// TODO shouldn't these deltas be rolled into the MidiSongPos?
+	float DeltaBarF[(int32)ECalibratedMusicTimebase::Count];
+	float DeltaBeatF[(int32)ECalibratedMusicTimebase::Count];
+
+protected:
+
+	// TODO: Remove Transport Controls from the Clock
+	virtual void OnStart() = 0;
+	virtual void OnPause() = 0;
+	virtual void OnContinue() = 0;
+	virtual void OnStop() = 0;
+	
+	UE_API int32 GetArrayIndexForTimebase(ECalibratedMusicTimebase Timebase) const;
+
+	UE_API void SetCurrentSongPos(ECalibratedMusicTimebase Timebase, const FMidiSongPos& InSongPos);
+	UE_API void SetCurrentSongPosByTime(ECalibratedMusicTimebase Timebase, float InElapsedMs, const ISongMapEvaluator& Map);
+	UE_API void SetCurrentSongPosByTick(ECalibratedMusicTimebase Timebase, float InTick, const ISongMapEvaluator& Map);
+
+private:
+
+	UE_API void ResetDefaultSongMaps(float DefaultTempo, int32 DefaultTimeSigNum, int32 DefaultTimeSigDenom);
+	virtual bool RefreshCurrentSongPos() = 0;
+	UE_API void UpdateBeatsAndBarsPerSecond();
+
+	uint64 LastUpdateFrame = 0;
+
+	EMusicClockState MusicClockState = EMusicClockState::Stopped;
+
+	TWeakObjectPtr<UObject> ContextObj;
+};
+}
+
+#undef UE_API
\ No newline at end of file
diff --git a/Engine/Plugins/Runtime/Harmonix/Source/HarmonixMetasound/Public/HarmonixMetasound/MusicClock/MusicClockState.h b/Engine/Plugins/Runtime/Harmonix/Source/HarmonixMetasound/Public/HarmonixMetasound/MusicClock/MusicClockState.h
new file mode 100644
--- /dev/null
+++ b/Engine/Plugins/Runtime/Harmonix/Source/HarmonixMetasound/Public/HarmonixMetasound/MusicClock/MusicClockState.h
@@ -0,0 +1,12 @@
+// Copyright Epic Games, Inc. All Rights Reserved.
+#pragma once
+
+#include "MusicClockState.generated.h"
+
+UENUM(BlueprintType)
+enum class EMusicClockState : uint8
+{
+	Stopped,
+	Paused,
+	Running,
+};
diff --git a/Engine/Plugins/Runtime/Harmonix/Source/HarmonixMetasound/Public/HarmonixMetasound/MusicClock/TimeSource.h b/Engine/Plugins/Runtime/Harmonix/Source/HarmonixMetasound/Public/HarmonixMetasound/MusicClock/TimeSource.h
new file mode 100644
--- /dev/null
+++ b/Engine/Plugins/Runtime/Harmonix/Source/HarmonixMetasound/Public/HarmonixMetasound/MusicClock/TimeSource.h
@@ -0,0 +1,84 @@
+// Copyright Epic Games, Inc. All Rights Reserved.
+#pragma once
+#include "CoreTypes.h"
+#include "Harmonix/MusicalTimebase.h"
+
+namespace Harmonix
+{
+	enum class ESourceState : uint8
+	{
+		Stopped,
+		Paused,
+		Running
+	};
+
+	enum class ESourceEvent : uint8
+	{
+		None,
+		Start,
+		Advance,
+		Stop,
+		Pause,
+		Continue,
+		Seek,
+		Loop
+	};
+
+	/**
+	 * Base class for Time Sources used by the Runtime Music Clock
+	 */
+	class ITimeSource
+	{
+	public:
+		virtual ~ITimeSource() {}
+
+		/**
+		 * Get display name of this time source for debugging
+		 *
+		 * @return DisplayName
+		 */
+		virtual FString GetDisplayName() const = 0;
+
+		/** The Calibrated Timebase this Time Source represents */
+		virtual ECalibratedMusicTimebase GetTimebase() const = 0;
+
+		/**
+		 * Used to update the time source once per frame
+		 *
+		 * CurrentTime, CurrentSourceState, and LatestSourceEvent
+		 * Should all be updated & valid for this frame when Update is called
+		 */
+		virtual void Update() = 0;
+
+		/**
+		 * Get the current tme of this time source
+		 * Is updated every frame via Update
+		 * 
+		 * @return Current Time in Seconds
+		 */
+		virtual double GetCurrentTime() const = 0;
+
+		/**
+		 * Get the speed at which time is advancing
+		 * 
+		 * @return Speed
+		 */
+		virtual float GetSpeed() const = 0;
+
+		/**
+		 * Get whether the source is Running, Stopped, or Paused
+		 * 
+		 * @return The Current State of the Source
+		 */
+		virtual ESourceState GetCurrentSourceState() const = 0;
+
+		/**
+		 * Retrieve the latest "Event" from this source
+		 * Whether it "Advanced", "Stopped", "Continued", etc.
+		 * Also whether it "Seeked" or "Looped"
+		 * 
+		 * @return The Latest Event of the Source
+		 */
+		virtual ESourceEvent GetLatestSourceEvent() const = 0;
+	};
+}
diff --git a/Engine/Plugins/Runtime/Harmonix/Source/HarmonixMetasound/Public/HarmonixMetasound/MusicClock/TransportContoller.h b/Engine/Plugins/Runtime/Harmonix/Source/HarmonixMetasound/Public/HarmonixMetasound/MusicClock/TransportContoller.h
new file mode 100644
--- /dev/null
+++ b/Engine/Plugins/Runtime/Harmonix/Source/HarmonixMetasound/Public/HarmonixMetasound/MusicClock/TransportContoller.h
@@ -0,0 +1,16 @@
+// Copyright Epic Games, Inc. All Rights Reserved.
+#pragma once
+#include "CoreTypes.h"
+
+namespace Harmonix
+{
+	enum class ETransportRequest : uint8
+	{
+		None,
+		Play,
+		Pause,
+		Continue,
+		Stop,
+		Seek
+	};
+}
diff --git a/Engine/Plugins/Runtime/Harmonix/Source/HarmonixMetasound/Public/HarmonixMetasound/MusicClock/WorldTimeSourceController.h b/Engine/Plugins/Runtime/Harmonix/Source/HarmonixMetasound/Public/HarmonixMetasound/MusicClock/WorldTimeSourceController.h
new file mode 100644
--- /dev/null
+++ b/Engine/Plugins/Runtime/Harmonix/Source/HarmonixMetasound/Public/HarmonixMetasound/MusicClock/WorldTimeSourceController.h
@@ -0,0 +1,53 @@
+// Copyright Epic Games, Inc. All Rights Reserved.
+#pragma once
+
+#include "TimeSource.h"
+#include "TransportContoller.h"
+#include "UObject/WeakObjectPtr.h"
+
+namespace Harmonix
+{
+class FWorldTimeSourceController : public ITimeSource
+{
+public:
+
+	FWorldTimeSourceController(UObject* InContextObj);
+
+	virtual FString GetDisplayName() const override;
+
+	virtual ECalibratedMusicTimebase GetTimebase() const override;
+
+	virtual void Update() override;
+
+	virtual ESourceState GetCurrentSourceState() const override;
+
+	virtual ESourceEvent GetLatestSourceEvent() const override;
+
+	virtual double GetCurrentTime() const override;
+
+	virtual float GetSpeed() const override;
+
+	void Start();
+
+	void Stop();
+
+	void Pause();
+
+	void Continue();
+
+	void Seek(double InSeekTime);
+
+private:
+
+	double GetDeltaTime() const;
+
+	double RunTimeSeconds = 0.0;
+	double SeekRequestTime = 0.0;
+	ETransportRequest TransportRequest = ETransportRequest::None;
+	ESourceEvent LatestSourceEvent = ESourceEvent::None;
+	ESourceState CurrentSourceState = ESourceState::Stopped;
+
+	FWeakObjectPtr ContextObj;
+};
+
+}
