BarMap Changes

Expose TimeSignature Numerator/Denominator to blueprints as int32

Add some functionality.

#rnx

#Synced-CL 49593419

diff --git a/Engine/Plugins/Runtime/Harmonix/Source/HarmonixMidi/Private/HarmonixMidi/BarMap.cpp b/Engine/Plugins/Runtime/Harmonix/Source/HarmonixMidi/Private/HarmonixMidi/BarMap.cpp
--- a/Engine/Plugins/Runtime/Harmonix/Source/HarmonixMidi/Private/HarmonixMidi/BarMap.cpp
+++ b/Engine/Plugins/Runtime/Harmonix/Source/HarmonixMidi/Private/HarmonixMidi/BarMap.cpp
@@ -518,6 +518,11 @@ bool FBarMap::AddTimeSignatureAtBarIncludingCountIn(int32 BarIndex, int32 Numera
 	return true;
 }
 
+void FBarMap::WipeTimeSignaturePoints(int32 Tick)
+{
+	FMusicMapUtl::RemovePointsOnAndAfterTick(Points, Tick);
+}
+
 int32 FBarMap::GetNumTimeSignaturePoints() const
 {
 	return Points.Num();
diff --git a/Engine/Plugins/Runtime/Harmonix/Source/HarmonixMidi/Public/HarmonixMidi/BarMap.h b/Engine/Plugins/Runtime/Harmonix/Source/HarmonixMidi/Public/HarmonixMidi/BarMap.h
--- a/Engine/Plugins/Runtime/Harmonix/Source/HarmonixMidi/Public/HarmonixMidi/BarMap.h
+++ b/Engine/Plugins/Runtime/Harmonix/Source/HarmonixMidi/Public/HarmonixMidi/BarMap.h
@@ -92,8 +92,8 @@ public:
 		, Denominator(4)
 	{}
 	FTimeSignature(int32 InNumerator, int32 InDenominator)
-		: Numerator((int16)InNumerator)
-		, Denominator((int16)InDenominator)
+		: Numerator(InNumerator)
+		, Denominator(InDenominator)
 	{}
 
 	bool operator==(const FTimeSignature& Other) const
@@ -106,10 +106,10 @@ public:
 		return Other.Numerator != Numerator || Other.Denominator != Denominator;
 	}
 
-	UPROPERTY()
-	int16 Numerator;
-	UPROPERTY()
-	int16 Denominator;
+	UPROPERTY(EditAnywhere, BlueprintReadWrite, Category="TimeSignature", meta=(ClampMin=1, ClampMax=99))
+	int32 Numerator = 4;
+	UPROPERTY(EditAnywhere, BlueprintReadWrite, Category="TimeSignature", meta=(ClampMin=1, ClampMax=99))
+	int32 Denominator = 4;
 
 	friend FORCEINLINE uint32 GetTypeHash(const FTimeSignature& InTimeSignature)
 	{
@@ -386,6 +386,9 @@ public:
 	 */
 	UE_API bool AddTimeSignatureAtBarIncludingCountIn(int32 BarIndex, int32 InNumerator, int32 InDenominator, bool SortNow = true, bool FailOnError = true);
 
+	/** Remove every point from the given tick onward */
+	UE_API void WipeTimeSignaturePoints(int32 Tick);
+
 	UE_API int32 GetNumTimeSignaturePoints() const;
 
 	UE_API int32 GetPointIndexForTick(int32 Tick) const;
