TimeSource and WorldTimeSource changes

Add helper method to TimeSource for determining SourceEvents from State transitions

Add "StartTime" to "Start" method

#rnx
#Synced-CL 49593419

diff --git a/Engine/Plugins/Runtime/Harmonix/Source/HarmonixMetasound/Private/HarmonixMetasound/MusicClock/TimeSource.cpp b/Engine/Plugins/Runtime/Harmonix/Source/HarmonixMetasound/Private/HarmonixMetasound/MusicClock/TimeSource.cpp
new file mode 100644
--- /dev/null
+++ b/Engine/Plugins/Runtime/Harmonix/Source/HarmonixMetasound/Private/HarmonixMetasound/MusicClock/TimeSource.cpp
@@ -0,0 +1,43 @@
+// Copyright Epic Games, Inc. All Rights Reserved.
+
+#include "HarmonixMetasound/MusicClock/TimeSource.h"
+
+Harmonix::ESourceEvent Harmonix::ITimeSource::SourceEventFromStateTransition(ESourceState PrevState, ESourceState NewState)
+{
+	switch (PrevState)
+	{
+	case ESourceState::Stopped:
+		switch (NewState)
+		{
+		case ESourceState::Running:
+			return ESourceEvent::Start;
+		default:
+			return ESourceEvent::None;
+		}
+	case ESourceState::Running:
+		switch (NewState)
+		{
+		case ESourceState::Running:
+			return ESourceEvent::Advance;
+		case ESourceState::Stopped:
+			return ESourceEvent::Stop;
+		case ESourceState::Paused:
+			return ESourceEvent::Pause;
+		default:
+			return ESourceEvent::None;
+		}
+	case ESourceState::Paused:
+		switch (NewState)
+		{
+		case ESourceState::Running:
+			return ESourceEvent::Continue;
+		case ESourceState::Stopped:
+			return ESourceEvent::Stop;
+		case ESourceState::Paused:
+		default:
+			return ESourceEvent::None;
+		}
+	default:
+		return ESourceEvent::None;
+	}
+}
\ No newline at end of file
diff --git a/Engine/Plugins/Runtime/Harmonix/Source/HarmonixMetasound/Private/HarmonixMetasound/MusicClock/WorldTimeSourceController.cpp b/Engine/Plugins/Runtime/Harmonix/Source/HarmonixMetasound/Private/HarmonixMetasound/MusicClock/WorldTimeSourceController.cpp
--- a/Engine/Plugins/Runtime/Harmonix/Source/HarmonixMetasound/Private/HarmonixMetasound/MusicClock/WorldTimeSourceController.cpp
+++ b/Engine/Plugins/Runtime/Harmonix/Source/HarmonixMetasound/Private/HarmonixMetasound/MusicClock/WorldTimeSourceController.cpp
@@ -14,6 +14,11 @@ FString Harmonix::FWorldTimeSourceController::GetDisplayName() const
 	return TEXT("RealTimeSourceController");
 }
 
+TOptional<FVector> Harmonix::FWorldTimeSourceController::TryGetAudioSourceLocation() const
+{
+	return {};
+}
+
 ECalibratedMusicTimebase Harmonix::FWorldTimeSourceController::GetTimebase() const
 {
 	return ECalibratedMusicTimebase::VideoRenderTime;
@@ -36,11 +41,13 @@ double Harmonix::FWorldTimeSourceController::GetDeltaTime() const
 
 void Harmonix::FWorldTimeSourceController::Update()
 {
-	double DeltaTime = GetDeltaTime();
+	double DeltaTime = GetDeltaTime() * Speed;
+
+	LatestSourceEvent = ESourceEvent::None;
+	
 	switch (TransportRequest)
 	{
 	default:
-		LatestSourceEvent = ESourceEvent::None;
 		break;
 	case ETransportRequest::Stop:
 		CurrentSourceState = ESourceState::Stopped;
@@ -50,6 +57,7 @@ void Harmonix::FWorldTimeSourceController::Update()
 	case ETransportRequest::Play:
 		if (CurrentSourceState == ESourceState::Stopped)
 		{
+			RunTimeSeconds = SeekRequestTime;
 			CurrentSourceState = ESourceState::Running;
 			LatestSourceEvent = ESourceEvent::Start;
 		}
@@ -76,7 +84,7 @@ void Harmonix::FWorldTimeSourceController::Update()
 		}
 		break;
 	}
-
+	
 	if (CurrentSourceState == ESourceState::Running)
 	{
 		RunTimeSeconds += DeltaTime;
@@ -85,6 +93,8 @@ void Harmonix::FWorldTimeSourceController::Update()
 			LatestSourceEvent = ESourceEvent::Advance;
 		}
 	}
+
+	TransportRequest = ETransportRequest::None;
 }
 
 Harmonix::ESourceState Harmonix::FWorldTimeSourceController::GetCurrentSourceState() const
@@ -109,15 +119,21 @@ float Harmonix::FWorldTimeSourceController::GetSpeed() const
 		if (const UWorld* World = Obj->GetWorld())
 		{
 			FGameTime GameTime = World->GetTime();
-			return GameTime.GetTimeDilation();
+			return GameTime.GetTimeDilation() * Speed;
 		}
 	}
 
 	return 0.0;
 }
 
-void Harmonix::FWorldTimeSourceController::Start()
+void Harmonix::FWorldTimeSourceController::SetSpeed(float InSpeed)
+{
+	Speed = FMath::Clamp(InSpeed, 0.01f, 10.0f);
+}
+
+void Harmonix::FWorldTimeSourceController::Start(float InStartTime)
 {
+	SeekRequestTime = FMath::Max(InStartTime, 0.0f);
 	TransportRequest = ETransportRequest::Play;
 }
 
@@ -136,8 +152,8 @@ void Harmonix::FWorldTimeSourceController::Continue()
 	TransportRequest = ETransportRequest::Continue;
 }
 
-void Harmonix::FWorldTimeSourceController::Seek(double InSeekTime)
+void Harmonix::FWorldTimeSourceController::Seek(float InSeekTime)
 {
 	TransportRequest = ETransportRequest::Seek;
-	SeekRequestTime = InSeekTime;
+	SeekRequestTime = FMath::Max(InSeekTime, 0.0f);
 }
\ No newline at end of file
diff --git a/Engine/Plugins/Runtime/Harmonix/Source/HarmonixMetasound/Public/HarmonixMetasound/MusicClock/TimeSource.h b/Engine/Plugins/Runtime/Harmonix/Source/HarmonixMetasound/Public/HarmonixMetasound/MusicClock/TimeSource.h
--- a/Engine/Plugins/Runtime/Harmonix/Source/HarmonixMetasound/Public/HarmonixMetasound/MusicClock/TimeSource.h
+++ b/Engine/Plugins/Runtime/Harmonix/Source/HarmonixMetasound/Public/HarmonixMetasound/MusicClock/TimeSource.h
@@ -3,6 +3,8 @@
 #include "CoreTypes.h"
 #include "Harmonix/MusicalTimebase.h"
 
+#define UE_API HARMONIXMETASOUND_API
+
 namespace Harmonix
 {
 	enum class ESourceState : uint8
@@ -33,12 +35,22 @@ namespace Harmonix
 		virtual ~ITimeSource() {}
 
 		/**
-		 * Get display name of this time source for debugging
+		 * Determine the source event from a given state transition
 		 *
-		 * @return DisplayName
+		 * e.g. Stopped => Running, return ESourceEvent::Start
+		 * 
+		 * @param PrevState 
+		 * @param NewState 
+		 * @return The source event representing the state transition
 		 */
+		UE_API static ESourceEvent SourceEventFromStateTransition(ESourceState PrevState, ESourceState NewState);
+		
+		/** @return DisplayName of this time source for debugging */
 		virtual FString GetDisplayName() const = 0;
 
+		/** @return The location of the Audio Source for this Time Source */
+		virtual TOptional<FVector> TryGetAudioSourceLocation() const = 0;
+
 		/** The Calibrated Timebase this Time Source represents */
 		virtual ECalibratedMusicTimebase GetTimebase() const = 0;
 
@@ -82,3 +94,5 @@ namespace Harmonix
 		virtual ESourceEvent GetLatestSourceEvent() const = 0;
 	};
 }
+
+#undef UE_API
diff --git a/Engine/Plugins/Runtime/Harmonix/Source/HarmonixMetasound/Public/HarmonixMetasound/MusicClock/WorldTimeSourceController.h b/Engine/Plugins/Runtime/Harmonix/Source/HarmonixMetasound/Public/HarmonixMetasound/MusicClock/WorldTimeSourceController.h
--- a/Engine/Plugins/Runtime/Harmonix/Source/HarmonixMetasound/Public/HarmonixMetasound/MusicClock/WorldTimeSourceController.h
+++ b/Engine/Plugins/Runtime/Harmonix/Source/HarmonixMetasound/Public/HarmonixMetasound/MusicClock/WorldTimeSourceController.h
@@ -15,6 +15,8 @@ public:
 
 	virtual FString GetDisplayName() const override;
 
+	virtual TOptional<FVector> TryGetAudioSourceLocation() const override;
+
 	virtual ECalibratedMusicTimebase GetTimebase() const override;
 
 	virtual void Update() override;
@@ -27,7 +29,9 @@ public:
 
 	virtual float GetSpeed() const override;
 
-	void Start();
+	void SetSpeed(float InSpeed);
+
+	void Start(float StartTime = 0.0f);
 
 	void Stop();
 
@@ -35,12 +39,13 @@ public:
 
 	void Continue();
 
-	void Seek(double InSeekTime);
+	void Seek(float InSeekTime);
 
 private:
 
 	double GetDeltaTime() const;
 
+	float Speed = 1.0f;
 	double RunTimeSeconds = 0.0;
 	double SeekRequestTime = 0.0;
 	ETransportRequest TransportRequest = ETransportRequest::None;
