First attempt at implementation of musical timers (and tests).

Currently supports:
Pause/Resume/Remove
Looping
Offsets
Quantized Times supported
Blueprint support

Current Tests:
Musical Timing (Bar/Beat/Measure/Offset)

#tests Wrote tests, tested in pie
#rnx

#Synced-CL 50273774

diff --git a/Engine/Plugins/Runtime/Harmonix/Source/HarmonixMetasound/Private/HarmonixMetasound/Components/MusicClockComponent.cpp b/Engine/Plugins/Runtime/Harmonix/Source/HarmonixMetasound/Private/HarmonixMetasound/Components/MusicClockComponent.cpp
--- a/Engine/Plugins/Runtime/Harmonix/Source/HarmonixMetasound/Private/HarmonixMetasound/Components/MusicClockComponent.cpp
+++ b/Engine/Plugins/Runtime/Harmonix/Source/HarmonixMetasound/Private/HarmonixMetasound/Components/MusicClockComponent.cpp
@@ -1,13 +1,14 @@
 // Copyright Epic Games, Inc. All Rights Reserved.
-#include "HarmonixMetasound/Components/MusicClockComponent.h"
+#include "Components/AudioComponent.h"
+#include "Engine/World.h"
+#include "GameFramework/Actor.h"
 #include "HarmonixMetasound/Components/MetasoundMusicClockDriver.h"
+#include "HarmonixMetasound/Components/MusicClockComponent.h"
+#include "HarmonixMetasound/Components/MusicTimer.h"
 #include "HarmonixMetasound/Components/WallClockMusicClockDriver.h"
 #include "HarmonixMetasound/Subsystems/MidiClockUpdateSubsystem.h"
 #include "HarmonixMidi/MusicTimeSpan.h"
-#include "Components/AudioComponent.h"
-#include "Engine/World.h"
 #include "Templates/SharedPointer.h"
-#include "GameFramework/Actor.h"
 
 #include UE_INLINE_GENERATED_CPP_BY_NAME(MusicClockComponent)
 
@@ -18,6 +19,7 @@ UMusicClockComponent::UMusicClockComponent()
 	MakeDefaultSongMap();
 	PrimaryComponentTick.bCanEverTick = false;
 	bAutoActivate = false;
+	MusicTimerManager = CreateDefaultSubobject<UMusicTimerManager>(TEXT("MusicTimerManager_MusicClockComponent"));
 }
 
 TNotNull<UMusicClockComponent*> UMusicClockComponent::CreateMusicClockShared(UObject* WorldContextObject, bool bPersistAcrossLevelTransition)
@@ -158,6 +160,7 @@ void UMusicClockComponent::UpdateClock()
 	if (ClockDriver)
 	{
 		ClockDriver->UpdateClockForGameFrame();
+		MusicTimerManager->UpdateForGameFrame(ClockDriver->CurrentSongPos);
 
 		BroadcastSongPosChanges(TimebaseForBarAndBeatEvents);
 		BroadcastSeekLoopDetections(ECalibratedMusicTimebase::AudioRenderTime);
@@ -439,6 +442,11 @@ void UMusicClockComponent::GetCurrentTimeSignature(int& OutNumerator, int& OutDe
 	}
 }
 
+FTimeSignature UMusicClockComponent::GetCurrentTimeSignatureStruct() const
+{
+	return ClockDriver ? FTimeSignature(ClockDriver->TimeSignatureNum, ClockDriver->TimeSignatureDenom) : FTimeSignature(0, 0);
+}
+
 float UMusicClockComponent::GetCurrentBarsPerSecond() const
 {
 	if (ClockDriver)
@@ -775,6 +783,44 @@ void UMusicClockComponent::SetDisplayName(FString Name)
 }
 #endif
 
+FMusicTimerHandle UMusicClockComponent::BP_AddTimer(const FMusicTimeInterval& TimeInterval, TOptional<FMusicTimestamp> StartTime, ECalibratedMusicTimebase Timebase, bool bLooping, FOnMusicalTimerExecute TimerDelegate)
+{
+	if (!StartTime.IsSet())
+	{
+		StartTime = Quantize(GetCurrentTimestamp(Timebase), TimeInterval.Interval, EMidiFileQuantizeDirection::Down);
+	}
+
+	return MusicTimerManager->BP_AddTimer(TimeInterval, StartTime.GetValue(), Timebase, bLooping, TimerDelegate);
+}
+
+FMusicTimerHandle UMusicClockComponent::AddTimerNative(const FMusicTimeInterval& TimeInterval, TOptional<FMusicTimestamp> StartTime, ECalibratedMusicTimebase Timebase, bool bLooping, FSimpleDelegate TimerDelegate)
+{
+	if (!StartTime.IsSet())
+	{
+		StartTime = Quantize(GetCurrentTimestamp(Timebase), TimeInterval.Interval, EMidiFileQuantizeDirection::Down);
+	}
+
+	return MusicTimerManager->AddTimerNative(TimeInterval, StartTime.GetValue(), Timebase, bLooping, TimerDelegate);
+}
+
+void UMusicClockComponent::RemoveTimer(const FMusicTimerHandle& Handle)
+{
+	MusicTimerManager->RemoveTimer(Handle);
+}
+
+void UMusicClockComponent::PauseTimer(const FMusicTimerHandle& Handle, bool bPause)
+{
+	MusicTimerManager->PauseTimer(Handle, bPause);
+}
+
+FMusicTimestamp UMusicClockComponent::Quantize(const FMusicTimestamp& Time, EMidiClockSubdivisionQuantization QuantizationInterval, EMidiFileQuantizeDirection InDirection) const
+{
+	const ISongMapEvaluator& MapEvaluator = GetSongMaps();
+	float Tick = MapEvaluator.MusicTimestampToTick(Time);
+	int32 QuantizedTick = MapEvaluator.QuantizeTickToNearestSubdivision(FMath::FloorToInt32(Tick), InDirection, QuantizationInterval);
+	return MapEvaluator.TickToMusicTimestamp(QuantizedTick);
+}
+
 const FMidiSongPos& UMusicClockComponent::GetSongPos(ECalibratedMusicTimebase Timebase) const
 {
 	return GetCurrentSongPosInternal(Timebase);
@@ -1016,5 +1062,4 @@ FMusicalTime UMusicClockComponent::Quantize(const FMusicalTime& MusicalTime, int
 	Result.TicksPerBar = Result.TicksPerBeat * TimeSignature->Numerator;
 	Result.TickInBar = static_cast<int32>((Timestamp.Beat - 1.0f) * static_cast<float>(Result.TicksPerBeat));
 	return Result;
-}
-
+}
\ No newline at end of file
diff --git a/Engine/Plugins/Runtime/Harmonix/Source/HarmonixMetasound/Private/HarmonixMetasound/Components/MusicTimer.cpp b/Engine/Plugins/Runtime/Harmonix/Source/HarmonixMetasound/Private/HarmonixMetasound/Components/MusicTimer.cpp
new file mode 100644
--- /dev/null
+++ b/Engine/Plugins/Runtime/Harmonix/Source/HarmonixMetasound/Private/HarmonixMetasound/Components/MusicTimer.cpp
@@ -0,0 +1,55 @@
+// Copyright Epic Games, Inc. All Rights Reserved.
+#include "HarmonixMetasound/Components/MusicTimer.h"
+#include "HarmonixMetasound/Components/MusicTimerHandle.h"
+#include "FrameBasedMusicMap.h"
+
+DEFINE_LOG_CATEGORY(LogMusicTimer);
+
+FMusicTimer::FMusicTimer(uint64 InUniqueIdentifier, const FMusicTimeInterval& InTimeInterval, const FMusicTimestamp& InStartTime, ECalibratedMusicTimebase InTimebase, bool InbLooping, FSimpleDelegate InDelegate)
+	: Timebase(InTimebase)
+	, Delegate(InDelegate)
+	, TimeInterval(InTimeInterval)
+	, StartTime(InStartTime)
+	, UniqueIdentifier(InUniqueIdentifier)
+	, bLooping(InbLooping)
+{
+}
+
+uint64 FMusicTimer::GetUniqueIdentifier() const
+{
+	return UniqueIdentifier;
+}
+
+void FMusicTimer::Reset()
+{
+	StartTime.Reset();
+	PausedTime.Reset();
+	UniqueIdentifier = 0;
+	bLooping = false;
+}
+
+bool FMusicTimer::Update(const FMusicTimestamp& CurrentMusicalTimestamp, const FTimeSignature& TimeSignature)
+{
+	FMusicTimestamp EndTimestamp = StartTime;
+	Harmonix::IncrementTimestampByInterval(EndTimestamp, TimeInterval, TimeSignature);
+
+	FMusicTimestamp OffsetTimeStamp = EndTimestamp;
+	Harmonix::IncrementTimestampByOffset(OffsetTimeStamp, TimeInterval, TimeSignature);
+
+	if (CurrentMusicalTimestamp >= OffsetTimeStamp)
+	{
+		Delegate.ExecuteIfBound();
+		if (IsLooping())
+		{
+			// Make sure the timer start position is old position + delay. Do not use current time as we may have overshot. 
+			StartTime = EndTimestamp;
+		}
+		else
+		{
+			// We are done, remove timer.
+			return true;
+		}
+	}
+
+	return false;
+}
\ No newline at end of file
diff --git a/Engine/Plugins/Runtime/Harmonix/Source/HarmonixMetasound/Private/HarmonixMetasound/Components/MusicTimerHandle.cpp b/Engine/Plugins/Runtime/Harmonix/Source/HarmonixMetasound/Private/HarmonixMetasound/Components/MusicTimerHandle.cpp
new file mode 100644
--- /dev/null
+++ b/Engine/Plugins/Runtime/Harmonix/Source/HarmonixMetasound/Private/HarmonixMetasound/Components/MusicTimerHandle.cpp
@@ -0,0 +1,43 @@
+// Copyright Epic Games, Inc. All Rights Reserved.
+
+#include "HarmonixMetasound/Components/MusicTimerHandle.h"
+
+FMusicTimerHandle::FMusicTimerHandle(uint64 InHandle)
+	: Handle(InHandle)
+{
+}
+
+bool FMusicTimerHandle::IsValid() const
+{
+	return Handle != 0;
+}
+
+void FMusicTimerHandle::Invalidate()
+{
+	Handle = 0;
+}
+
+bool FMusicTimerHandle::operator==(const FMusicTimerHandle& Other) const
+{
+	return Handle == Other.Handle;
+}
+
+bool FMusicTimerHandle::operator!=(const FMusicTimerHandle& Other) const
+{
+	return Handle != Other.Handle;
+}
+
+FString FMusicTimerHandle::ToString() const
+{
+	return FString::Printf(TEXT("%llu"), Handle);
+}
+
+uint64 FMusicTimerHandle::GetUniqueIdentifier() const
+{
+	return Handle;
+}
+
+uint32 GetTypeHash(const FMusicTimerHandle& InHandle)
+{
+	return GetTypeHash(InHandle.Handle);
+}
\ No newline at end of file
diff --git a/Engine/Plugins/Runtime/Harmonix/Source/HarmonixMetasound/Private/HarmonixMetasound/Components/MusicTimerManager.cpp b/Engine/Plugins/Runtime/Harmonix/Source/HarmonixMetasound/Private/HarmonixMetasound/Components/MusicTimerManager.cpp
new file mode 100644
--- /dev/null
+++ b/Engine/Plugins/Runtime/Harmonix/Source/HarmonixMetasound/Private/HarmonixMetasound/Components/MusicTimerManager.cpp
@@ -0,0 +1,134 @@
+// Copyright Epic Games, Inc. All Rights Reserved.
+
+#include "HarmonixMetasound/Components/MusicClockComponent.h"
+#include "HarmonixMetasound/Components/MusicTimerManager.h"
+#include "HarmonixMetasound/Components/MusicTimerHandle.h"
+
+DEFINE_LOG_CATEGORY(LogMusicTimerManager);
+
+std::atomic<uint64> UMusicTimerManager::NextTimerId = 1;
+
+FMusicTimerHandle UMusicTimerManager::BP_AddTimer(const FMusicTimeInterval& TimeInterval, FMusicTimestamp StartTime, ECalibratedMusicTimebase Timebase, bool bLooping, FOnMusicalTimerExecute TimerDelegate)
+{
+	auto Lambda = TDelegate<void(void)>::CreateLambda([TimerDelegate]()
+	{
+		TimerDelegate.ExecuteIfBound();
+	});
+
+	return AddTimerNative(TimeInterval, StartTime, Timebase, bLooping, Lambda);
+}
+
+FMusicTimerHandle UMusicTimerManager::AddTimerNative(const FMusicTimeInterval& TimeInterval, FMusicTimestamp StartTime, ECalibratedMusicTimebase Timebase, bool bLooping, FSimpleDelegate TimerDelegate)
+{
+	FMusicTimerHandle Handle(GetUniqueTimerId());
+	PendingTimers.Add(MakeUnique<FMusicTimer>(Handle.GetUniqueIdentifier(), TimeInterval, StartTime, Timebase, bLooping, TimerDelegate));
+	return Handle;
+}
+
+void UMusicTimerManager::UpdateForGameFrame(const TStaticArray<FMidiSongPos, (int32)ECalibratedMusicTimebase::Count>& CurrentSongPositions)
+{
+	if (GFrameCounter == LastUpdateFrame)
+	{
+		return;
+	}
+
+	ActiveTimers.Append(MoveTemp(PendingTimers));
+	PendingTimers.Empty();
+
+	for (auto Iter = ActiveTimers.CreateIterator(); Iter; ++Iter)
+	{
+		TUniquePtr<FMusicTimer>& CurrentTimer = *Iter;
+		if (!CurrentTimer || !CurrentTimer->GetUniqueIdentifier())
+		{
+			Iter.RemoveCurrentSwap();
+			continue;
+		}
+
+		if (CurrentTimer->IsPaused())
+		{
+			continue;
+		}
+
+		const FMidiSongPos& MidiSongPos = CurrentSongPositions[static_cast<int32>(CurrentTimer->GetTimebase())];
+		const FMusicTimestamp& CurrentMusicalTimestamp = MidiSongPos.Timestamp;
+		FTimeSignature TimeSignature(MidiSongPos.TimeSigNumerator, MidiSongPos.TimeSigDenominator);
+		if(bool bDone = CurrentTimer->Update(CurrentMusicalTimestamp, TimeSignature))
+		{
+			Iter.RemoveCurrentSwap();
+		}
+	}
+
+	LastUpdateFrameSongPositions = CurrentSongPositions;
+	LastUpdateFrame = GFrameCounter;
+}
+
+void UMusicTimerManager::RemoveTimer(const FMusicTimerHandle& Handle)
+{
+	PendingTimers.RemoveAll([Handle](const TUniquePtr<FMusicTimer>& Timer) { return Timer && Timer->GetUniqueIdentifier() == Handle.GetUniqueIdentifier(); });
+
+	// Do not remove these immediately as they may be currently being iterated on.
+	// They will be removed as part of the update.
+	for (TUniquePtr<FMusicTimer>& ActiveTimer : ActiveTimers)
+	{
+		if (ActiveTimer && ActiveTimer->GetUniqueIdentifier() == Handle.GetUniqueIdentifier())
+		{
+			ActiveTimer->Reset();
+		}
+	}
+}
+
+void UMusicTimerManager::PauseTimer(const FMusicTimerHandle& Handle, bool bPause)
+{
+	if(LastUpdateFrame == 0)
+	{
+		UE_LOG(LogMusicTimerManager, Error, TEXT("Can not pause timers before the manager has updated at least once!"));
+		return;
+	}
+
+	auto MatchesHandleLambda = [&Handle](const TUniquePtr<FMusicTimer>& Timer) -> bool
+	{
+		return Timer && Timer->GetUniqueIdentifier() == Handle.GetUniqueIdentifier();
+	};
+
+	TUniquePtr<FMusicTimer>* FoundTimer = PendingTimers.FindByPredicate(MatchesHandleLambda);
+	if (!FoundTimer)
+	{
+		// Not in pending timers, lets see if its in the active timers.
+		FoundTimer = ActiveTimers.FindByPredicate(MatchesHandleLambda);
+	}
+
+	FMusicTimer* Timer = FoundTimer ? (*FoundTimer).Get() : nullptr;
+	if (!Timer || bPause == Timer->IsPaused())
+	{
+		return;
+	}
+
+	const FMusicTimestamp LastFrameMusicalTimestamp = LastUpdateFrameSongPositions[static_cast<int32>(Timer->GetTimebase())].Timestamp;
+	if (bPause)
+	{
+		Timer->PausedTime = LastFrameMusicalTimestamp;
+	}
+	else
+	{
+		const FMidiSongPos& LastFrameMidiSongPos = LastUpdateFrameSongPositions[static_cast<int32>(Timer->GetTimebase())];
+		FTimeSignature LastFrameTimeSignature(LastFrameMidiSongPos.TimeSigNumerator, LastFrameMidiSongPos.TimeSigDenominator);
+
+		const FMusicTimestamp& PausedTime = Timer->PausedTime.GetValue();
+		const int32 BeatsPerBar = LastFrameTimeSignature.Numerator;
+		const float CurrentAbsoluteBeats = float(LastFrameMusicalTimestamp.Bar - 1) * BeatsPerBar + (LastFrameMusicalTimestamp.Beat - 1.0f);
+		const float PausedbsoluteBeats = float(PausedTime.Bar - 1) * BeatsPerBar + (PausedTime.Beat - 1.0f);
+		const float PausedDeltaBeats = CurrentAbsoluteBeats - PausedbsoluteBeats;
+		Harmonix::IncrementTimestampByBeats(Timer->StartTime, PausedDeltaBeats, LastFrameTimeSignature);
+		Timer->PausedTime.Reset();
+	}
+}
+
+FMusicTimestamp UMusicTimerManager::Quantize(const UMusicClockComponent* MusicClock, const FMusicTimestamp& Time, EMidiClockSubdivisionQuantization QuantizationInterval, EMidiFileQuantizeDirection InDirection) const
+{
+	if(!MusicClock)
+	{
+		return FMusicTimestamp();
+	}
+
+	return MusicClock->Quantize(Time, QuantizationInterval, InDirection);
+}
diff --git a/Engine/Plugins/Runtime/Harmonix/Source/HarmonixMetasound/Private/HarmonixMetasound/DataTypes/MusicTimeInterval.cpp b/Engine/Plugins/Runtime/Harmonix/Source/HarmonixMetasound/Private/HarmonixMetasound/DataTypes/MusicTimeInterval.cpp
--- a/Engine/Plugins/Runtime/Harmonix/Source/HarmonixMetasound/Private/HarmonixMetasound/DataTypes/MusicTimeInterval.cpp
+++ b/Engine/Plugins/Runtime/Harmonix/Source/HarmonixMetasound/Private/HarmonixMetasound/DataTypes/MusicTimeInterval.cpp
@@ -3,26 +3,26 @@
 #include "HarmonixMetasound/DataTypes/MusicTimeInterval.h"
 #include "HarmonixMidi/MidiConstants.h"
 
-namespace Harmonix
+float FMusicTimeInterval::GetIntervalBeats(const FTimeSignature& TimeSignature) const
 {
-	float FMusicTimeInterval::GetIntervalBeats(const FTimeSignature& TimeSignature) const
-	{
-		return Midi::Constants::SubdivisionToBeats(Interval, TimeSignature) * IntervalMultiplier; 
-	}
+	return Harmonix::Midi::Constants::SubdivisionToBeats(Interval, TimeSignature) * IntervalMultiplier;
+}
 
-	float FMusicTimeInterval::GetOffsetBeats(const FTimeSignature& TimeSignature) const
+float FMusicTimeInterval::GetOffsetBeats(const FTimeSignature& TimeSignature) const
+{
+	const float IntervalBeats = GetIntervalBeats(TimeSignature);
+
+	if (FMath::IsNearlyZero(IntervalBeats))
 	{
-		const float IntervalBeats = GetIntervalBeats(TimeSignature);
-		
-		if (FMath::IsNearlyZero(IntervalBeats))
-		{
-			return 0.0f;
-		}
-		
-		const float OffsetBeats = Midi::Constants::SubdivisionToBeats(Offset, TimeSignature) * OffsetMultiplier;
-		return FMath::Fmod(OffsetBeats, IntervalBeats);
+		return 0.0f;
 	}
 
+	const float OffsetBeats = Harmonix::Midi::Constants::SubdivisionToBeats(Offset, TimeSignature) * OffsetMultiplier;
+	return FMath::Fmod(OffsetBeats, IntervalBeats);
+}
+
+namespace Harmonix
+{
 	void IncrementTimestampByBeats(FMusicTimestamp& Timestamp, float Beats, const FTimeSignature& TimeSignature)
 	{
 		// Add the beats
diff --git a/Engine/Plugins/Runtime/Harmonix/Source/HarmonixMetasound/Private/HarmonixMetasound/MidiOps/PulseGenerator.cpp b/Engine/Plugins/Runtime/Harmonix/Source/HarmonixMetasound/Private/HarmonixMetasound/MidiOps/PulseGenerator.cpp
--- a/Engine/Plugins/Runtime/Harmonix/Source/HarmonixMetasound/Private/HarmonixMetasound/MidiOps/PulseGenerator.cpp
+++ b/Engine/Plugins/Runtime/Harmonix/Source/HarmonixMetasound/Private/HarmonixMetasound/MidiOps/PulseGenerator.cpp
@@ -14,7 +14,7 @@ namespace Harmonix::Midi::Ops
 		Interval = NewInterval;
 
 		// Multiplier should be >= 1
-		Interval.IntervalMultiplier = FMath::Max(Interval.IntervalMultiplier, static_cast<uint16>(1));
+		Interval.IntervalMultiplier = FMath::Max(Interval.IntervalMultiplier, 1);
 	}
 
 	void FPulseGenerator::SetDutyCycle(float InDutyCycle)
diff --git a/Engine/Plugins/Runtime/Harmonix/Source/HarmonixMetasound/Public/HarmonixMetasound/Components/MusicClockComponent.h b/Engine/Plugins/Runtime/Harmonix/Source/HarmonixMetasound/Public/HarmonixMetasound/Components/MusicClockComponent.h
--- a/Engine/Plugins/Runtime/Harmonix/Source/HarmonixMetasound/Public/HarmonixMetasound/Components/MusicClockComponent.h
+++ b/Engine/Plugins/Runtime/Harmonix/Source/HarmonixMetasound/Public/HarmonixMetasound/Components/MusicClockComponent.h
@@ -4,6 +4,7 @@
 #include "Components/ActorComponent.h"
 #include "Evaluation/IMovieSceneCustomClockSource.h"
 #include "Harmonix/MusicalTimebase.h"
+#include "HarmonixMetasound/Components/MusicTimerManager.h"
 #include "HarmonixMetasound/MusicClock/MusicClockState.h"
 #include "HarmonixMidi/MidiSongPos.h"
 #include "HarmonixMidi/SongMaps.h"
@@ -21,6 +22,7 @@ namespace Metasound { class FMetasoundGenerator; }
 struct FMusicalTimeSpan;
 class UAudioComponent;
 class FMusicClockDriverBase;
+class FMusicTimer;
 
 UENUM(BlueprintType)
 enum class EMusicClockDriveMethod : uint8
@@ -142,6 +144,14 @@ public:
 	UFUNCTION(BlueprintPure, Category = "MusicClock")
 	UE_API void GetCurrentTimeSignature(int& OutNumerator, int& OutDenominator) const;
 
+	/**
+	 * Get The Current Time Signature
+	 *
+	 * 0 if the clock is stopped
+	 */
+	UFUNCTION(BlueprintPure, Category = "MusicClock")
+	UE_API FTimeSignature GetCurrentTimeSignatureStruct() const;
+
 	/**
 	 * Get The Current Bars per Second.
 	 * Based on the current tempo, advance rate and time signature
@@ -358,6 +368,40 @@ public:
 
 	UE_API const ISongMapEvaluator& GetSongMaps() const;
 
+	/**
+	 * Adds a musical timer.
+	 *
+	 * @param TimerInterval  Interval configuration for the timer.
+	 * @param StartTime      Optional start time. If unset, the current musical
+	 *                       time (floored) is used.
+	 * @param Timebase       Timebase used to evaluate the timer.
+	 * @param bLooping       Whether the timer should loop.
+	 * @param TimerDelegate  Delegate executed when the timer fires.
+	 */
+	UFUNCTION(BlueprintCallable, Category = "Music|Timer")
+	FMusicTimerHandle BP_AddTimer(const FMusicTimeInterval& TimerInterval, TOptional<FMusicTimestamp> StartTime, ECalibratedMusicTimebase Timebase, bool bLooping, FOnMusicalTimerExecute TimerDelegate);
+
+	/**
+	 * Adds a musical timer.
+	 *
+	 * @param TimerInterval  Interval configuration for the timer.
+	 * @param StartTime      Optional start time. If unset, the current musical
+	 *                       time (floored) is used.
+	 * @param Timebase       Timebase used to evaluate the timer.
+	 * @param bLooping       Whether the timer should loop.
+	 * @param TimerDelegate  Delegate executed when the timer fires.
+	 */
+	UE_API FMusicTimerHandle AddTimerNative(const FMusicTimeInterval& TimerInterval, TOptional<FMusicTimestamp> StartTime, ECalibratedMusicTimebase Timebase, bool bLooping, FSimpleDelegate TimerDelegate);
+
+	UFUNCTION(BlueprintCallable, Category = "Music|Timer")
+	UE_API void RemoveTimer(const FMusicTimerHandle& Handle);
+
+	UFUNCTION(BlueprintCallable, Category = "Music|Timer")
+	UE_API void PauseTimer(const FMusicTimerHandle& Handle, bool bPause);
+
+	UFUNCTION(BlueprintCallable, Category = "Music|Timer")
+	FMusicTimestamp Quantize(const FMusicTimestamp& Time, EMidiClockSubdivisionQuantization QuantizationInterval, EMidiFileQuantizeDirection InDirection) const;
+
 #if !UE_BUILD_SHIPPING
 	// A short string description of the clock, for debugging, like song name or animation name or whatever is driving it.
 	UE_API FString GetDisplayName() const;
@@ -405,16 +449,6 @@ public:
 
 protected:
 
-	/**
-     *
-	 * Default Clock Advance Rate
-	 * used for Wall clock with no tempo map
-	 *
-	 * NOT THE CURRENT RUNNING TEMPO
-	 */
-	UPROPERTY(EditAnywhere, BlueprintReadWrite, Category = "MusicClock|Defaults")
-	float DefaultClockAdvanceRate = 1.0f;
-	
 	/**
 	 * Default Tempo (BPM)
 	 * used for Wall clock with no tempo map
@@ -514,6 +548,8 @@ public:
 	virtual bool CanAuditionInEditor() const override { return false; }
 	// ^^^ END IMusicEnvironmentClockSource Implementation ^^^
 	
+	UE_API bool HasValidClockDriver() const { return ClockDriver.IsValid(); }
+
 private:
 	friend class  UMidiClockUpdateSubsystem;
 
@@ -529,6 +565,9 @@ private:
 	FString DisplayName;
 #endif
 	
+	UPROPERTY(Transient)
+	TObjectPtr<UMusicTimerManager> MusicTimerManager;
+
 	static TNotNull<UMusicClockComponent*> CreateMusicClockShared(UObject* WorldContextObject, bool bPersistAcrossLevelTransition);
 	
 	UE_API void CreateClockDriver();
diff --git a/Engine/Plugins/Runtime/Harmonix/Source/HarmonixMetasound/Public/HarmonixMetasound/Components/MusicClockDriverBase.h b/Engine/Plugins/Runtime/Harmonix/Source/HarmonixMetasound/Public/HarmonixMetasound/Components/MusicClockDriverBase.h
--- a/Engine/Plugins/Runtime/Harmonix/Source/HarmonixMetasound/Public/HarmonixMetasound/Components/MusicClockDriverBase.h
+++ b/Engine/Plugins/Runtime/Harmonix/Source/HarmonixMetasound/Public/HarmonixMetasound/Components/MusicClockDriverBase.h
@@ -70,8 +70,8 @@ public:
 	float CurrentBarsPerSecond = 0.0f;
 	float CurrentClockAdvanceRate = 1.0f;
 
-	FMidiSongPos CurrentSongPos[(int32)ECalibratedMusicTimebase::Count];
-	FMidiSongPos PrevSongPos[(int32)ECalibratedMusicTimebase::Count];
+	TStaticArray<FMidiSongPos, (int32)ECalibratedMusicTimebase::Count> CurrentSongPos;
+	TStaticArray<FMidiSongPos, (int32)ECalibratedMusicTimebase::Count> PrevSongPos;
 
 	// TODO shouldn't these deltas be rolled into the MidiSongPos?
 	float DeltaBarF[(int32)ECalibratedMusicTimebase::Count];
diff --git a/Engine/Plugins/Runtime/Harmonix/Source/HarmonixMetasound/Public/HarmonixMetasound/Components/MusicTimer.h b/Engine/Plugins/Runtime/Harmonix/Source/HarmonixMetasound/Public/HarmonixMetasound/Components/MusicTimer.h
new file mode 100644
--- /dev/null
+++ b/Engine/Plugins/Runtime/Harmonix/Source/HarmonixMetasound/Public/HarmonixMetasound/Components/MusicTimer.h
@@ -0,0 +1,41 @@
+// Copyright Epic Games, Inc. All Rights Reserved.
+
+#pragma once
+
+#include "Misc/MusicalTime.h"
+#include "Harmonix/MusicalTimebase.h"
+#include "HarmonixMetasound/DataTypes/MusicTimeInterval.h"
+
+DECLARE_LOG_CATEGORY_EXTERN(LogMusicTimer, Warning, All);
+
+#define UE_API HARMONIXMETASOUND_API
+
+class FMusicTimer
+{
+public:
+	friend class UMusicTimerManager;
+	UE_API FMusicTimer(uint64 InUniqueIdentifier, const FMusicTimeInterval& InTimeInterval, const FMusicTimestamp& InStartTime, ECalibratedMusicTimebase InTimebase, bool InbLooping, FSimpleDelegate InDelegate);
+
+	ECalibratedMusicTimebase GetTimebase() const { return Timebase; }
+	uint64 GetUniqueIdentifier() const;
+
+	bool IsLooping() const { return bLooping; }
+	bool IsPaused() const { return PausedTime.IsSet(); }
+
+	// Returns whether or not the timer fired its delegate
+	bool Update(const FMusicTimestamp& CurrentMusicalTimestamp, const FTimeSignature& TimeSignature);
+	void Reset();
+private:
+
+	UPROPERTY(EditAnywhere, Category = "Music|Timer")
+	ECalibratedMusicTimebase Timebase = ECalibratedMusicTimebase::VideoRenderTime;
+
+	FSimpleDelegate Delegate;
+	FMusicTimeInterval TimeInterval;
+	FMusicTimestamp StartTime;
+	TOptional<FMusicTimestamp> PausedTime;
+	uint64 UniqueIdentifier = 0;
+	bool bLooping = false;
+};
+
+#undef UE_API
\ No newline at end of file
diff --git a/Engine/Plugins/Runtime/Harmonix/Source/HarmonixMetasound/Public/HarmonixMetasound/Components/MusicTimerHandle.h b/Engine/Plugins/Runtime/Harmonix/Source/HarmonixMetasound/Public/HarmonixMetasound/Components/MusicTimerHandle.h
new file mode 100644
--- /dev/null
+++ b/Engine/Plugins/Runtime/Harmonix/Source/HarmonixMetasound/Public/HarmonixMetasound/Components/MusicTimerHandle.h
@@ -0,0 +1,36 @@
+// Copyright Epic Games, Inc. All Rights Reserved.
+
+#pragma once
+
+#include "CoreMinimal.h"
+#include "MusicTimerHandle.generated.h"
+
+#define UE_API HARMONIXMETASOUND_API
+
+/** Unique handle that can be used to distinguish timers that have identical delegates. */
+USTRUCT(BlueprintType)
+struct FMusicTimerHandle
+{
+	GENERATED_BODY()
+
+	friend class UMusicTimerManager;
+
+	UE_API FMusicTimerHandle(uint64 InHandle = 0);
+
+	UE_API bool IsValid() const;
+	UE_API void Invalidate();
+
+	UE_API FString ToString() const;
+	UE_API uint64 GetUniqueIdentifier() const;
+	UE_API friend uint32 GetTypeHash(const FMusicTimerHandle& InHandle);
+
+	bool operator==(const FMusicTimerHandle& Other) const;
+	bool operator!=(const FMusicTimerHandle& Other) const;
+
+private:
+
+	UPROPERTY(Transient)
+	uint64 Handle;
+};
+
+#undef UE_API
\ No newline at end of file
diff --git a/Engine/Plugins/Runtime/Harmonix/Source/HarmonixMetasound/Public/HarmonixMetasound/Components/MusicTimerManager.h b/Engine/Plugins/Runtime/Harmonix/Source/HarmonixMetasound/Public/HarmonixMetasound/Components/MusicTimerManager.h
new file mode 100644
--- /dev/null
+++ b/Engine/Plugins/Runtime/Harmonix/Source/HarmonixMetasound/Public/HarmonixMetasound/Components/MusicTimerManager.h
@@ -0,0 +1,75 @@
+// Copyright Epic Games, Inc. All Rights Reserved.
+
+#pragma once
+
+#include "Containers/StaticArray.h"
+#include "MusicTimer.h"
+#include "HarmonixMidi/MidiSongPos.h"
+#include "MusicTimerManager.generated.h"
+
+#define UE_API HARMONIXMETASOUND_API
+
+DECLARE_LOG_CATEGORY_EXTERN(LogMusicTimerManager, Warning, All);
+DECLARE_DYNAMIC_DELEGATE(FOnMusicalTimerExecute);
+
+enum class EMidiClockSubdivisionQuantization : uint8;
+enum class EMidiFileQuantizeDirection : uint8;
+
+struct FMusicTimerHandle;
+class UMusicClockComponent;
+
+UCLASS(Blueprintable, BlueprintType)
+class UMusicTimerManager : public UObject
+{
+	GENERATED_BODY()
+
+public:
+
+	// We need the MidiSongPosition in each calibrated time, because each timer can be operating on their own timebase.
+	void UpdateForGameFrame(const TStaticArray<FMidiSongPos, (int32)ECalibratedMusicTimebase::Count>& CurrentSongPositions);
+
+	/**
+	 * Adds a musical timer.
+	 *
+	 * @param TimerInterval  Interval configuration for the timer.
+	 * @param StartTime      Start time. You can use the Quantize function below to snap to interval boundaries.
+	 * @param Timebase       Timebase used to evaluate the timer.
+	 * @param bLooping       Whether the timer should loop.
+	 * @param TimerDelegate  Delegate executed when the timer fires.
+	 */
+	UFUNCTION(BlueprintCallable, Category = "Music|Timer")
+	FMusicTimerHandle BP_AddTimer(const FMusicTimeInterval& TimerInterval, FMusicTimestamp StartTime, ECalibratedMusicTimebase Timebase, bool bLooping, FOnMusicalTimerExecute TimerDelegate);
+
+	/**
+	 * Adds a musical timer.
+	 *
+	 * @param TimerInterval  Interval configuration for the timer.
+	 * @param StartTime      Start time. You can use the Quantize function below to snap to interval boundaries.
+	 * @param Timebase       Timebase used to evaluate the timer.
+	 * @param bLooping       Whether the timer should loop.
+	 * @param TimerDelegate  Delegate executed when the timer fires.
+	 */
+	UE_API FMusicTimerHandle AddTimerNative(const FMusicTimeInterval& TimerInterval, FMusicTimestamp StartTime, ECalibratedMusicTimebase Timebase, bool bLooping, FSimpleDelegate TimerDelegate);
+
+	UFUNCTION(BlueprintCallable, Category = "Music|Timer")
+	UE_API void RemoveTimer(const FMusicTimerHandle& Handle);
+
+	UFUNCTION(BlueprintCallable, Category = "Music|Timer")
+	UE_API void PauseTimer(const FMusicTimerHandle& Handle, bool bPause);
+
+	UFUNCTION(BlueprintCallable, Category = "Music|Timer")
+	FMusicTimestamp Quantize(const UMusicClockComponent* MusicClock, const FMusicTimestamp& Time, EMidiClockSubdivisionQuantization QuantizationInterval, EMidiFileQuantizeDirection InDirection) const;
+
+	UE_API uint64 GetUniqueTimerId() { return ++NextTimerId; }
+private:
+
+	// Used so pause doesn't need to pass in this info
+	TStaticArray<FMidiSongPos, (int32)ECalibratedMusicTimebase::Count> LastUpdateFrameSongPositions;
+
+	TArray<TUniquePtr<FMusicTimer>> ActiveTimers; // Current Timers being processed
+	TArray<TUniquePtr<FMusicTimer>> PendingTimers; // New timers added during execution
+	uint64 LastUpdateFrame = 0;
+	static std::atomic<uint64> NextTimerId;
+};
+
+#undef UE_API
\ No newline at end of file
diff --git a/Engine/Plugins/Runtime/Harmonix/Source/HarmonixMetasound/Public/HarmonixMetasound/DataTypes/MusicTimeInterval.h b/Engine/Plugins/Runtime/Harmonix/Source/HarmonixMetasound/Public/HarmonixMetasound/DataTypes/MusicTimeInterval.h
--- a/Engine/Plugins/Runtime/Harmonix/Source/HarmonixMetasound/Public/HarmonixMetasound/DataTypes/MusicTimeInterval.h
+++ b/Engine/Plugins/Runtime/Harmonix/Source/HarmonixMetasound/Public/HarmonixMetasound/DataTypes/MusicTimeInterval.h
@@ -3,21 +3,31 @@
 #pragma once
 
 #include "MidiClock.h"
+#include "MusicTimeInterval.generated.h"
 
-namespace Harmonix
+USTRUCT(BlueprintType)
+struct HARMONIXMETASOUND_API FMusicTimeInterval
 {
-	struct FMusicTimeInterval
-	{
-		EMidiClockSubdivisionQuantization Interval = EMidiClockSubdivisionQuantization::Beat;
-		EMidiClockSubdivisionQuantization Offset = EMidiClockSubdivisionQuantization::Beat;
-		uint16 IntervalMultiplier = 1;
-		uint16 OffsetMultiplier = 0;
+	GENERATED_BODY()
+
+	UPROPERTY(EditAnywhere, BlueprintReadWrite, Category = "Time Interval Settings")
+	EMidiClockSubdivisionQuantization Interval = EMidiClockSubdivisionQuantization::Beat;
+
+	UPROPERTY(EditAnywhere, BlueprintReadWrite, Category = "Time Interval Settings")
+	EMidiClockSubdivisionQuantization Offset = EMidiClockSubdivisionQuantization::Beat;
 
-		HARMONIXMETASOUND_API float GetIntervalBeats(const FTimeSignature& TimeSignature) const;
+	UPROPERTY(EditAnywhere, BlueprintReadWrite, Category = "Time Interval Settings")
+	int32 IntervalMultiplier = 1;
 
-		HARMONIXMETASOUND_API float GetOffsetBeats(const FTimeSignature& TimeSignature) const;
-	};
+	UPROPERTY(EditAnywhere, BlueprintReadWrite, Category = "Time Interval Settings")
+	int32 OffsetMultiplier = 0;
 
+	float GetIntervalBeats(const FTimeSignature& TimeSignature) const;
+	float GetOffsetBeats(const FTimeSignature& TimeSignature) const;
+};
+
+namespace Harmonix
+{
 	HARMONIXMETASOUND_API void IncrementTimestampByInterval(
 			FMusicTimestamp& Timestamp,
 			const FMusicTimeInterval& Interval,
diff --git a/Engine/Plugins/Runtime/Harmonix/Source/HarmonixMetasoundTests/Private/Components/MusicClockComponentTests.cpp b/Engine/Plugins/Runtime/Harmonix/Source/HarmonixMetasoundTests/Private/Components/MusicClockComponentTests.cpp
new file mode 100644
--- /dev/null
+++ b/Engine/Plugins/Runtime/Harmonix/Source/HarmonixMetasoundTests/Private/Components/MusicClockComponentTests.cpp
@@ -0,0 +1,233 @@
+// Copyright Epic Games, Inc. All Rights Reserved.
+
+#include "Misc/AutomationTest.h"
+
+#if WITH_DEV_AUTOMATION_TESTS
+
+#include "HarmonixMetasound/Components/MusicClockComponent.h"
+#include "HarmonixMetasound/Components/MusicTimerHandle.h"
+#include "HarmonixMetasound/DataTypes/MusicTimeInterval.h"
+#include "Tests/AutomationEditorCommon.h"
+
+namespace HarmonixMetasound::MidiStream::Tests
+{
+	class FCheckFiredLatentCommand : public IAutomationLatentCommand
+	{
+	public:
+		explicit FCheckFiredLatentCommand(FAutomationTestBase& InAutomationTestBase, TSharedRef<std::atomic<bool>, ESPMode::ThreadSafe> InFlag, bool InExpectedValue)
+			: AutomationTestBase(InAutomationTestBase), Flag(MoveTemp(InFlag))
+			, ExpectedValue(InExpectedValue)
+		{
+		}
+
+		virtual bool Update() override
+		{
+			AutomationTestBase.TestEqual(TEXT("bFired flag was unexpected value"), Flag->load(std::memory_order_acquire), ExpectedValue);
+			return true;
+		}
+
+	private:
+		FAutomationTestBase& AutomationTestBase;
+		TSharedRef<std::atomic<bool>, ESPMode::ThreadSafe> Flag;
+		bool ExpectedValue = false;
+	};
+
+
+	class FWaitUntilFireLatentCommand : public IAutomationLatentCommand
+	{
+	public:
+		FWaitUntilFireLatentCommand(TSharedRef<std::atomic<bool>, ESPMode::ThreadSafe> InFlag, float InTimeoutSeconds)
+			: Flag(MoveTemp(InFlag))
+			, StartTimeSeconds(FPlatformTime::Seconds())
+			, TimeoutSeconds(InTimeoutSeconds)
+		{
+		}
+
+		virtual bool Update() override
+		{
+			// Fired
+			if (Flag->load(std::memory_order_acquire))
+			{
+				return true;
+			}
+
+			// Timeout
+			if ((FPlatformTime::Seconds() - StartTimeSeconds) > TimeoutSeconds)
+			{
+				return true;
+			}
+
+			// Not fired yet
+			return false;
+		}
+
+	private:
+		TSharedRef<std::atomic<bool>, ESPMode::ThreadSafe> Flag;
+		double StartTimeSeconds = 0.0;
+		float TimeoutSeconds = 0.0f;
+	};
+
+	struct FMusicClockComponentHolder
+	{
+		TObjectPtr<UWorld> World = nullptr;
+		TObjectPtr<AActor> Owner = nullptr;
+		TObjectPtr<UMusicClockComponent> MusicClockComponent = nullptr;
+		FMusicTimerHandle TimerHandle;
+	};
+
+	static FMusicClockComponentHolder CreateMusicClock(FAutomationTestBase& AutomationTestBase)
+	{
+		FMusicClockComponentHolder ComponentHolder;
+
+		ComponentHolder.World = FAutomationEditorCommonUtils::CreateNewMap();
+		AutomationTestBase.TestNotNull("World created", ComponentHolder.World.Get());
+
+		ComponentHolder.Owner = ComponentHolder.World->SpawnActor<AActor>();
+		AutomationTestBase.TestNotNull("Owner actor spawned", ComponentHolder.Owner.Get());
+
+		bool bStart = true;
+		bool bPersistAcrossLevelTransitions = false;
+		UMidiFile* TempoMap = nullptr;
+		ComponentHolder.MusicClockComponent = UMusicClockComponent::CreateWallClockDrivenMusicClock(ComponentHolder.World, TempoMap, bStart, bPersistAcrossLevelTransitions);
+
+		return ComponentHolder;
+	}
+
+	FMusicClockComponentHolder StartTimer(FAutomationTestBase& AutomationTestBase, TSharedRef<std::atomic<bool>, ESPMode::ThreadSafe> bFired, FMusicTimeInterval Interval)
+	{
+		FMusicClockComponentHolder MusicClockComponentHolder = CreateMusicClock(AutomationTestBase);
+		UMusicClockComponent* MusicClockComponent = MusicClockComponentHolder.MusicClockComponent;
+
+		const ECalibratedMusicTimebase Timebase = ECalibratedMusicTimebase::RawAudioRenderTime;
+		const FMusicTimestamp StartTime = MusicClockComponent->GetCurrentTimestamp(Timebase);
+		const bool bLooping = false;
+
+		auto Callback = TDelegate<void()>::CreateLambda([MusicClockComponent, StartTime, Interval, Timebase, &AutomationTestBase, bFired]()
+		{
+			if (const FTimeSignature* TimeSignature = MusicClockComponent->GetSongMaps().GetTimeSignatureAtTick(MusicClockComponent->GetTicksIncludingCountIn(Timebase)))
+			{
+				// Check what time we should have ended on
+				FMusicTimestamp EndTime = StartTime;
+				Harmonix::IncrementTimestampByInterval(EndTime, Interval, *TimeSignature);
+				const FMusicTimestamp CurrentTime = MusicClockComponent->GetCurrentTimestamp(Timebase);
+				AutomationTestBase.TestTrue("Timer triggered at the correct time", CurrentTime.Bar == EndTime.Bar && FMath::IsNearlyEqual(CurrentTime.Beat, EndTime.Beat, 0.5f));
+			}
+
+			bFired->store(true, std::memory_order_release);
+		});
+
+		MusicClockComponentHolder.TimerHandle = MusicClockComponent->AddTimerNative(Interval, StartTime, Timebase, bLooping, Callback);
+		AutomationTestBase.TestTrue("Valid Timer Handle", MusicClockComponentHolder.TimerHandle.IsValid());
+
+		return MusicClockComponentHolder;
+	}
+
+	IMPLEMENT_SIMPLE_AUTOMATION_TEST(FMusicClockTimerBeatTest, "Harmonix.Metasound.Components.MusicClockComponent.Timer.TimerBeat", EAutomationTestFlags::EditorContext | EAutomationTestFlags::EngineFilter)
+	bool FMusicClockTimerBeatTest::RunTest(const FString&)
+	{
+		FMusicTimeInterval Interval;
+		Interval.Interval = EMidiClockSubdivisionQuantization::Beat;
+		Interval.IntervalMultiplier = 25;
+
+		TSharedRef<std::atomic<bool>> bFired = MakeShared<std::atomic<bool>, ESPMode::ThreadSafe>(false);
+		FMusicClockComponentHolder Refs = StartTimer(*this, bFired, Interval);
+
+		// Wait until the callback fires
+		const float Timeout = 50.0f;
+		AddCommand(new FWaitUntilFireLatentCommand(bFired, Timeout));
+		AddCommand(new FCheckFiredLatentCommand(*this, bFired, true));
+
+		// Latent command, test is not over.
+		return true;
+	}
+
+	IMPLEMENT_SIMPLE_AUTOMATION_TEST(FMusicClockTimerBarTest, "Harmonix.Metasound.Components.MusicClockComponent.Timer.TimerBar", EAutomationTestFlags::EditorContext | EAutomationTestFlags::EngineFilter)
+	bool FMusicClockTimerBarTest::RunTest(const FString&)
+	{
+		FMusicTimeInterval Interval;
+		Interval.Interval = EMidiClockSubdivisionQuantization::Bar;
+		Interval.IntervalMultiplier = 25;
+
+		TSharedRef<std::atomic<bool>> bFired = MakeShared<std::atomic<bool>, ESPMode::ThreadSafe>(false);
+		FMusicClockComponentHolder Refs = StartTimer(*this, bFired, Interval);
+
+		// Wait until the callback fires
+		const float Timeout = 50.0f;
+		AddCommand(new FWaitUntilFireLatentCommand(bFired, Timeout));
+		AddCommand(new FCheckFiredLatentCommand(*this, bFired, true));
+
+		// Latent command, test is not over.
+		return true;
+	}
+
+	IMPLEMENT_SIMPLE_AUTOMATION_TEST(FMusicClockTimerNoteTest, "Harmonix.Metasound.Components.MusicClockComponent.Timer.TimerNote", EAutomationTestFlags::EditorContext | EAutomationTestFlags::EngineFilter)
+	bool FMusicClockTimerNoteTest::RunTest(const FString&)
+	{
+		FMusicTimeInterval Interval;
+		Interval.Interval = EMidiClockSubdivisionQuantization::DottedEighthNote;
+		Interval.IntervalMultiplier = 25;
+
+		TSharedRef<std::atomic<bool>> bFired = MakeShared<std::atomic<bool>, ESPMode::ThreadSafe>(false);
+		FMusicClockComponentHolder Refs = StartTimer(*this, bFired, Interval);
+
+		// Wait until the callback fires
+		const float Timeout = 50.0f;
+		AddCommand(new FWaitUntilFireLatentCommand(bFired, Timeout));
+		AddCommand(new FCheckFiredLatentCommand(*this, bFired, true));
+
+		// Latent command, test is not over.
+		return true;
+	}
+
+	IMPLEMENT_SIMPLE_AUTOMATION_TEST(FMusicClockTimerOffsetTest, "Harmonix.Metasound.Components.MusicClockComponent.Timer.TimerOffsets", EAutomationTestFlags::EditorContext | EAutomationTestFlags::EngineFilter)
+	bool FMusicClockTimerOffsetTest::RunTest(const FString&)
+	{
+		FMusicTimeInterval Interval;
+		Interval.Interval = EMidiClockSubdivisionQuantization::Bar;
+		Interval.IntervalMultiplier = 15;
+		Interval.Offset = EMidiClockSubdivisionQuantization::Beat;
+		Interval.OffsetMultiplier = 3;
+
+		TSharedRef<std::atomic<bool>> bFired = MakeShared<std::atomic<bool>, ESPMode::ThreadSafe>(false);
+		FMusicClockComponentHolder Refs = StartTimer(*this, bFired, Interval);
+
+		// Wait until the callback fires
+		const float Timeout = 50.0f;
+		AddCommand(new FWaitUntilFireLatentCommand(bFired, Timeout));
+		AddCommand(new FCheckFiredLatentCommand(*this, bFired, true));
+
+		// Latent command, test is not over.
+		return true;
+	}
+
+	IMPLEMENT_SIMPLE_AUTOMATION_TEST(FMusicClockTimerPauseTest, "Harmonix.Metasound.Components.MusicClockComponent.Timer.TimerPause", EAutomationTestFlags::EditorContext | EAutomationTestFlags::EngineFilter)
+	bool FMusicClockTimerPauseTest::RunTest(const FString&)
+	{
+		FMusicTimeInterval Interval;
+		Interval.Interval = EMidiClockSubdivisionQuantization::Bar;
+		Interval.IntervalMultiplier = 15;
+
+		TSharedRef<std::atomic<bool>> bFired = MakeShared<std::atomic<bool>, ESPMode::ThreadSafe>(false);
+
+		// Start timer and pause it
+		FMusicClockComponentHolder Refs = StartTimer(*this, bFired, Interval);
+		Refs.MusicClockComponent->PauseTimer(Refs.TimerHandle, true);
+
+		// Timer is paused, we should timeout. 
+		const float Timeout = 10.0f;
+		AddCommand(new FWaitUntilFireLatentCommand(bFired, Timeout));
+		AddCommand(new FCheckFiredLatentCommand(*this, bFired, false));
+
+		// Unpause timer and wait to see if it resumes and fires
+		Refs.MusicClockComponent->PauseTimer(Refs.TimerHandle, false);
+		bFired->store(false, std::memory_order_relaxed);
+
+		AddCommand(new FWaitUntilFireLatentCommand(bFired, Timeout));
+		AddCommand(new FCheckFiredLatentCommand(*this, bFired, true));
+
+		// Latent command, test is not over.
+		return true;
+	}
+}
+
+#endif
diff --git a/Engine/Plugins/Runtime/Harmonix/Source/HarmonixMidi/Public/HarmonixMidi/SongMaps.h b/Engine/Plugins/Runtime/Harmonix/Source/HarmonixMidi/Public/HarmonixMidi/SongMaps.h
--- a/Engine/Plugins/Runtime/Harmonix/Source/HarmonixMidi/Public/HarmonixMidi/SongMaps.h
+++ b/Engine/Plugins/Runtime/Harmonix/Source/HarmonixMidi/Public/HarmonixMidi/SongMaps.h
@@ -25,7 +25,7 @@ namespace HarmonixMetasound
 	class FMidiClock;
 }
 
-UENUM()
+UENUM(BlueprintType)
 enum class EMidiFileQuantizeDirection : uint8
 {
 	Nearest = 0,
