[Backout] - CL50369713
#rnx
Original CL Desc
-----------------------------------------------------------------
First attempt at implementation of musical timers (and tests).

Currently supports:
Pause/Resume/Remove
Looping
Offsets
Quantized Times supported
Blueprint support

Current Tests:
Musical Timing (Bar/Beat/Measure/Offset)

#tests Wrote tests, tested in pie
#rnx

diff --git a/Engine/Plugins/Runtime/Harmonix/Source/HarmonixMetasound/Private/HarmonixMetasound/Components/MusicClockComponent.cpp b/Engine/Plugins/Runtime/Harmonix/Source/HarmonixMetasound/Private/HarmonixMetasound/Components/MusicClockComponent.cpp
--- a/Engine/Plugins/Runtime/Harmonix/Source/HarmonixMetasound/Private/HarmonixMetasound/Components/MusicClockComponent.cpp
+++ b/Engine/Plugins/Runtime/Harmonix/Source/HarmonixMetasound/Private/HarmonixMetasound/Components/MusicClockComponent.cpp
@@ -1,14 +1,13 @@
 // Copyright Epic Games, Inc. All Rights Reserved.
-#include "Components/AudioComponent.h"
-#include "Engine/World.h"
-#include "GameFramework/Actor.h"
-#include "HarmonixMetasound/Components/MetasoundMusicClockDriver.h"
 #include "HarmonixMetasound/Components/MusicClockComponent.h"
-#include "HarmonixMetasound/Components/MusicTimer.h"
+#include "HarmonixMetasound/Components/MetasoundMusicClockDriver.h"
 #include "HarmonixMetasound/Components/WallClockMusicClockDriver.h"
 #include "HarmonixMetasound/Subsystems/MidiClockUpdateSubsystem.h"
 #include "HarmonixMidi/MusicTimeSpan.h"
+#include "Components/AudioComponent.h"
+#include "Engine/World.h"
 #include "Templates/SharedPointer.h"
+#include "GameFramework/Actor.h"
 
 #include UE_INLINE_GENERATED_CPP_BY_NAME(MusicClockComponent)
 
@@ -19,7 +18,6 @@ UMusicClockComponent::UMusicClockComponent()
 	MakeDefaultSongMap();
 	PrimaryComponentTick.bCanEverTick = false;
 	bAutoActivate = false;
-	MusicTimerManager = CreateDefaultSubobject<UMusicTimerManager>(TEXT("MusicTimerManager_MusicClockComponent"));
 }
 
 TNotNull<UMusicClockComponent*> UMusicClockComponent::CreateMusicClockShared(UObject* WorldContextObject, bool bPersistAcrossLevelTransition)
@@ -160,7 +158,6 @@ void UMusicClockComponent::UpdateClock()
 	if (ClockDriver)
 	{
 		ClockDriver->UpdateClockForGameFrame();
-		MusicTimerManager->UpdateForGameFrame(ClockDriver->CurrentSongPos);
 
 		BroadcastSongPosChanges(TimebaseForBarAndBeatEvents);
 		BroadcastSeekLoopDetections(ECalibratedMusicTimebase::AudioRenderTime);
@@ -442,11 +439,6 @@ void UMusicClockComponent::GetCurrentTimeSignature(int& OutNumerator, int& OutDe
 	}
 }
 
-FTimeSignature UMusicClockComponent::GetCurrentTimeSignatureStruct() const
-{
-	return ClockDriver ? FTimeSignature(ClockDriver->TimeSignatureNum, ClockDriver->TimeSignatureDenom) : FTimeSignature(0, 0);
-}
-
 float UMusicClockComponent::GetCurrentBarsPerSecond() const
 {
 	if (ClockDriver)
@@ -783,44 +775,6 @@ void UMusicClockComponent::SetDisplayName(FString Name)
 }
 #endif
 
-FMusicTimerHandle UMusicClockComponent::BP_AddTimer(const FMusicTimeInterval& TimeInterval, TOptional<FMusicTimestamp> StartTime, ECalibratedMusicTimebase Timebase, bool bLooping, FOnMusicalTimerExecute TimerDelegate)
-{
-	if (!StartTime.IsSet())
-	{
-		StartTime = Quantize(GetCurrentTimestamp(Timebase), TimeInterval.Interval, EMidiFileQuantizeDirection::Down);
-	}
-
-	return MusicTimerManager->BP_AddTimer(TimeInterval, StartTime.GetValue(), Timebase, bLooping, TimerDelegate);
-}
-
-FMusicTimerHandle UMusicClockComponent::AddTimerNative(const FMusicTimeInterval& TimeInterval, TOptional<FMusicTimestamp> StartTime, ECalibratedMusicTimebase Timebase, bool bLooping, FSimpleDelegate TimerDelegate)
-{
-	if (!StartTime.IsSet())
-	{
-		StartTime = Quantize(GetCurrentTimestamp(Timebase), TimeInterval.Interval, EMidiFileQuantizeDirection::Down);
-	}
-
-	return MusicTimerManager->AddTimerNative(TimeInterval, StartTime.GetValue(), Timebase, bLooping, TimerDelegate);
-}
-
-void UMusicClockComponent::RemoveTimer(const FMusicTimerHandle& Handle)
-{
-	MusicTimerManager->RemoveTimer(Handle);
-}
-
-void UMusicClockComponent::PauseTimer(const FMusicTimerHandle& Handle, bool bPause)
-{
-	MusicTimerManager->PauseTimer(Handle, bPause);
-}
-
-FMusicTimestamp UMusicClockComponent::Quantize(const FMusicTimestamp& Time, EMidiClockSubdivisionQuantization QuantizationInterval, EMidiFileQuantizeDirection InDirection) const
-{
-	const ISongMapEvaluator& MapEvaluator = GetSongMaps();
-	float Tick = MapEvaluator.MusicTimestampToTick(Time);
-	int32 QuantizedTick = MapEvaluator.QuantizeTickToNearestSubdivision(FMath::FloorToInt32(Tick), InDirection, QuantizationInterval);
-	return MapEvaluator.TickToMusicTimestamp(QuantizedTick);
-}
-
 const FMidiSongPos& UMusicClockComponent::GetSongPos(ECalibratedMusicTimebase Timebase) const
 {
 	return GetCurrentSongPosInternal(Timebase);
@@ -1062,4 +1016,5 @@ FMusicalTime UMusicClockComponent::Quantize(const FMusicalTime& MusicalTime, int
 	Result.TicksPerBar = Result.TicksPerBeat * TimeSignature->Numerator;
 	Result.TickInBar = static_cast<int32>((Timestamp.Beat - 1.0f) * static_cast<float>(Result.TicksPerBeat));
 	return Result;
-}
\ No newline at end of file
+}
+
diff --git a/Engine/Plugins/Runtime/Harmonix/Source/HarmonixMetasound/Private/HarmonixMetasound/Components/MusicTimer.cpp b/Engine/Plugins/Runtime/Harmonix/Source/HarmonixMetasound/Private/HarmonixMetasound/Components/MusicTimer.cpp
deleted file mode 100644
diff --git a/Engine/Plugins/Runtime/Harmonix/Source/HarmonixMetasound/Private/HarmonixMetasound/Components/MusicTimerHandle.cpp b/Engine/Plugins/Runtime/Harmonix/Source/HarmonixMetasound/Private/HarmonixMetasound/Components/MusicTimerHandle.cpp
deleted file mode 100644
diff --git a/Engine/Plugins/Runtime/Harmonix/Source/HarmonixMetasound/Private/HarmonixMetasound/Components/MusicTimerManager.cpp b/Engine/Plugins/Runtime/Harmonix/Source/HarmonixMetasound/Private/HarmonixMetasound/Components/MusicTimerManager.cpp
deleted file mode 100644
diff --git a/Engine/Plugins/Runtime/Harmonix/Source/HarmonixMetasound/Private/HarmonixMetasound/DataTypes/MusicTimeInterval.cpp b/Engine/Plugins/Runtime/Harmonix/Source/HarmonixMetasound/Private/HarmonixMetasound/DataTypes/MusicTimeInterval.cpp
--- a/Engine/Plugins/Runtime/Harmonix/Source/HarmonixMetasound/Private/HarmonixMetasound/DataTypes/MusicTimeInterval.cpp
+++ b/Engine/Plugins/Runtime/Harmonix/Source/HarmonixMetasound/Private/HarmonixMetasound/DataTypes/MusicTimeInterval.cpp
@@ -3,26 +3,26 @@
 #include "HarmonixMetasound/DataTypes/MusicTimeInterval.h"
 #include "HarmonixMidi/MidiConstants.h"
 
-float FMusicTimeInterval::GetIntervalBeats(const FTimeSignature& TimeSignature) const
-{
-	return Harmonix::Midi::Constants::SubdivisionToBeats(Interval, TimeSignature) * IntervalMultiplier;
-}
-
-float FMusicTimeInterval::GetOffsetBeats(const FTimeSignature& TimeSignature) const
+namespace Harmonix
 {
-	const float IntervalBeats = GetIntervalBeats(TimeSignature);
-
-	if (FMath::IsNearlyZero(IntervalBeats))
+	float FMusicTimeInterval::GetIntervalBeats(const FTimeSignature& TimeSignature) const
 	{
-		return 0.0f;
+		return Midi::Constants::SubdivisionToBeats(Interval, TimeSignature) * IntervalMultiplier; 
 	}
 
-	const float OffsetBeats = Harmonix::Midi::Constants::SubdivisionToBeats(Offset, TimeSignature) * OffsetMultiplier;
-	return FMath::Fmod(OffsetBeats, IntervalBeats);
-}
+	float FMusicTimeInterval::GetOffsetBeats(const FTimeSignature& TimeSignature) const
+	{
+		const float IntervalBeats = GetIntervalBeats(TimeSignature);
+		
+		if (FMath::IsNearlyZero(IntervalBeats))
+		{
+			return 0.0f;
+		}
+		
+		const float OffsetBeats = Midi::Constants::SubdivisionToBeats(Offset, TimeSignature) * OffsetMultiplier;
+		return FMath::Fmod(OffsetBeats, IntervalBeats);
+	}
 
-namespace Harmonix
-{
 	void IncrementTimestampByBeats(FMusicTimestamp& Timestamp, float Beats, const FTimeSignature& TimeSignature)
 	{
 		// Add the beats
diff --git a/Engine/Plugins/Runtime/Harmonix/Source/HarmonixMetasound/Private/HarmonixMetasound/MidiOps/PulseGenerator.cpp b/Engine/Plugins/Runtime/Harmonix/Source/HarmonixMetasound/Private/HarmonixMetasound/MidiOps/PulseGenerator.cpp
--- a/Engine/Plugins/Runtime/Harmonix/Source/HarmonixMetasound/Private/HarmonixMetasound/MidiOps/PulseGenerator.cpp
+++ b/Engine/Plugins/Runtime/Harmonix/Source/HarmonixMetasound/Private/HarmonixMetasound/MidiOps/PulseGenerator.cpp
@@ -14,7 +14,7 @@ namespace Harmonix::Midi::Ops
 		Interval = NewInterval;
 
 		// Multiplier should be >= 1
-		Interval.IntervalMultiplier = FMath::Max(Interval.IntervalMultiplier, 1);
+		Interval.IntervalMultiplier = FMath::Max(Interval.IntervalMultiplier, static_cast<uint16>(1));
 	}
 
 	void FPulseGenerator::SetDutyCycle(float InDutyCycle)
diff --git a/Engine/Plugins/Runtime/Harmonix/Source/HarmonixMetasound/Public/HarmonixMetasound/Components/MusicClockComponent.h b/Engine/Plugins/Runtime/Harmonix/Source/HarmonixMetasound/Public/HarmonixMetasound/Components/MusicClockComponent.h
--- a/Engine/Plugins/Runtime/Harmonix/Source/HarmonixMetasound/Public/HarmonixMetasound/Components/MusicClockComponent.h
+++ b/Engine/Plugins/Runtime/Harmonix/Source/HarmonixMetasound/Public/HarmonixMetasound/Components/MusicClockComponent.h
@@ -4,7 +4,6 @@
 #include "Components/ActorComponent.h"
 #include "Evaluation/IMovieSceneCustomClockSource.h"
 #include "Harmonix/MusicalTimebase.h"
-#include "HarmonixMetasound/Components/MusicTimerManager.h"
 #include "HarmonixMetasound/MusicClock/MusicClockState.h"
 #include "HarmonixMidi/MidiSongPos.h"
 #include "HarmonixMidi/SongMaps.h"
@@ -22,7 +21,6 @@ namespace Metasound { class FMetasoundGenerator; }
 struct FMusicalTimeSpan;
 class UAudioComponent;
 class FMusicClockDriverBase;
-class FMusicTimer;
 
 UENUM(BlueprintType)
 enum class EMusicClockDriveMethod : uint8
@@ -144,14 +142,6 @@ public:
 	UFUNCTION(BlueprintPure, Category = "MusicClock")
 	UE_API void GetCurrentTimeSignature(int& OutNumerator, int& OutDenominator) const;
 
-	/**
-	 * Get The Current Time Signature
-	 *
-	 * 0 if the clock is stopped
-	 */
-	UFUNCTION(BlueprintPure, Category = "MusicClock")
-	UE_API FTimeSignature GetCurrentTimeSignatureStruct() const;
-
 	/**
 	 * Get The Current Bars per Second.
 	 * Based on the current tempo, advance rate and time signature
@@ -368,40 +358,6 @@ public:
 
 	UE_API const ISongMapEvaluator& GetSongMaps() const;
 
-	/**
-	 * Adds a musical timer.
-	 *
-	 * @param TimerInterval  Interval configuration for the timer.
-	 * @param StartTime      Optional start time. If unset, the current musical
-	 *                       time (floored) is used.
-	 * @param Timebase       Timebase used to evaluate the timer.
-	 * @param bLooping       Whether the timer should loop.
-	 * @param TimerDelegate  Delegate executed when the timer fires.
-	 */
-	UFUNCTION(BlueprintCallable, Category = "Music|Timer")
-	FMusicTimerHandle BP_AddTimer(const FMusicTimeInterval& TimerInterval, TOptional<FMusicTimestamp> StartTime, ECalibratedMusicTimebase Timebase, bool bLooping, FOnMusicalTimerExecute TimerDelegate);
-
-	/**
-	 * Adds a musical timer.
-	 *
-	 * @param TimerInterval  Interval configuration for the timer.
-	 * @param StartTime      Optional start time. If unset, the current musical
-	 *                       time (floored) is used.
-	 * @param Timebase       Timebase used to evaluate the timer.
-	 * @param bLooping       Whether the timer should loop.
-	 * @param TimerDelegate  Delegate executed when the timer fires.
-	 */
-	UE_API FMusicTimerHandle AddTimerNative(const FMusicTimeInterval& TimerInterval, TOptional<FMusicTimestamp> StartTime, ECalibratedMusicTimebase Timebase, bool bLooping, FSimpleDelegate TimerDelegate);
-
-	UFUNCTION(BlueprintCallable, Category = "Music|Timer")
-	UE_API void RemoveTimer(const FMusicTimerHandle& Handle);
-
-	UFUNCTION(BlueprintCallable, Category = "Music|Timer")
-	UE_API void PauseTimer(const FMusicTimerHandle& Handle, bool bPause);
-
-	UFUNCTION(BlueprintCallable, Category = "Music|Timer")
-	FMusicTimestamp Quantize(const FMusicTimestamp& Time, EMidiClockSubdivisionQuantization QuantizationInterval, EMidiFileQuantizeDirection InDirection) const;
-
 #if !UE_BUILD_SHIPPING
 	// A short string description of the clock, for debugging, like song name or animation name or whatever is driving it.
 	UE_API FString GetDisplayName() const;
@@ -449,6 +405,16 @@ public:
 
 protected:
 
+	/**
+     *
+	 * Default Clock Advance Rate
+	 * used for Wall clock with no tempo map
+	 *
+	 * NOT THE CURRENT RUNNING TEMPO
+	 */
+	UPROPERTY(EditAnywhere, BlueprintReadWrite, Category = "MusicClock|Defaults")
+	float DefaultClockAdvanceRate = 1.0f;
+	
 	/**
 	 * Default Tempo (BPM)
 	 * used for Wall clock with no tempo map
@@ -548,8 +514,6 @@ public:
 	virtual bool CanAuditionInEditor() const override { return false; }
 	// ^^^ END IMusicEnvironmentClockSource Implementation ^^^
 	
-	UE_API bool HasValidClockDriver() const { return ClockDriver.IsValid(); }
-
 private:
 	friend class  UMidiClockUpdateSubsystem;
 
@@ -565,9 +529,6 @@ private:
 	FString DisplayName;
 #endif
 	
-	UPROPERTY(Transient)
-	TObjectPtr<UMusicTimerManager> MusicTimerManager;
-
 	static TNotNull<UMusicClockComponent*> CreateMusicClockShared(UObject* WorldContextObject, bool bPersistAcrossLevelTransition);
 	
 	UE_API void CreateClockDriver();
diff --git a/Engine/Plugins/Runtime/Harmonix/Source/HarmonixMetasound/Public/HarmonixMetasound/Components/MusicClockDriverBase.h b/Engine/Plugins/Runtime/Harmonix/Source/HarmonixMetasound/Public/HarmonixMetasound/Components/MusicClockDriverBase.h
--- a/Engine/Plugins/Runtime/Harmonix/Source/HarmonixMetasound/Public/HarmonixMetasound/Components/MusicClockDriverBase.h
+++ b/Engine/Plugins/Runtime/Harmonix/Source/HarmonixMetasound/Public/HarmonixMetasound/Components/MusicClockDriverBase.h
@@ -70,8 +70,8 @@ public:
 	float CurrentBarsPerSecond = 0.0f;
 	float CurrentClockAdvanceRate = 1.0f;
 
-	TStaticArray<FMidiSongPos, (int32)ECalibratedMusicTimebase::Count> CurrentSongPos;
-	TStaticArray<FMidiSongPos, (int32)ECalibratedMusicTimebase::Count> PrevSongPos;
+	FMidiSongPos CurrentSongPos[(int32)ECalibratedMusicTimebase::Count];
+	FMidiSongPos PrevSongPos[(int32)ECalibratedMusicTimebase::Count];
 
 	// TODO shouldn't these deltas be rolled into the MidiSongPos?
 	float DeltaBarF[(int32)ECalibratedMusicTimebase::Count];
diff --git a/Engine/Plugins/Runtime/Harmonix/Source/HarmonixMetasound/Public/HarmonixMetasound/Components/MusicTimer.h b/Engine/Plugins/Runtime/Harmonix/Source/HarmonixMetasound/Public/HarmonixMetasound/Components/MusicTimer.h
deleted file mode 100644
diff --git a/Engine/Plugins/Runtime/Harmonix/Source/HarmonixMetasound/Public/HarmonixMetasound/Components/MusicTimerHandle.h b/Engine/Plugins/Runtime/Harmonix/Source/HarmonixMetasound/Public/HarmonixMetasound/Components/MusicTimerHandle.h
deleted file mode 100644
diff --git a/Engine/Plugins/Runtime/Harmonix/Source/HarmonixMetasound/Public/HarmonixMetasound/Components/MusicTimerManager.h b/Engine/Plugins/Runtime/Harmonix/Source/HarmonixMetasound/Public/HarmonixMetasound/Components/MusicTimerManager.h
deleted file mode 100644
diff --git a/Engine/Plugins/Runtime/Harmonix/Source/HarmonixMetasound/Public/HarmonixMetasound/DataTypes/MusicTimeInterval.h b/Engine/Plugins/Runtime/Harmonix/Source/HarmonixMetasound/Public/HarmonixMetasound/DataTypes/MusicTimeInterval.h
--- a/Engine/Plugins/Runtime/Harmonix/Source/HarmonixMetasound/Public/HarmonixMetasound/DataTypes/MusicTimeInterval.h
+++ b/Engine/Plugins/Runtime/Harmonix/Source/HarmonixMetasound/Public/HarmonixMetasound/DataTypes/MusicTimeInterval.h
@@ -3,31 +3,21 @@
 #pragma once
 
 #include "MidiClock.h"
-#include "MusicTimeInterval.generated.h"
 
-USTRUCT(BlueprintType)
-struct HARMONIXMETASOUND_API FMusicTimeInterval
+namespace Harmonix
 {
-	GENERATED_BODY()
-
-	UPROPERTY(EditAnywhere, BlueprintReadWrite, Category = "Time Interval Settings")
-	EMidiClockSubdivisionQuantization Interval = EMidiClockSubdivisionQuantization::Beat;
-
-	UPROPERTY(EditAnywhere, BlueprintReadWrite, Category = "Time Interval Settings")
-	EMidiClockSubdivisionQuantization Offset = EMidiClockSubdivisionQuantization::Beat;
+	struct FMusicTimeInterval
+	{
+		EMidiClockSubdivisionQuantization Interval = EMidiClockSubdivisionQuantization::Beat;
+		EMidiClockSubdivisionQuantization Offset = EMidiClockSubdivisionQuantization::Beat;
+		uint16 IntervalMultiplier = 1;
+		uint16 OffsetMultiplier = 0;
 
-	UPROPERTY(EditAnywhere, BlueprintReadWrite, Category = "Time Interval Settings")
-	int32 IntervalMultiplier = 1;
+		HARMONIXMETASOUND_API float GetIntervalBeats(const FTimeSignature& TimeSignature) const;
 
-	UPROPERTY(EditAnywhere, BlueprintReadWrite, Category = "Time Interval Settings")
-	int32 OffsetMultiplier = 0;
+		HARMONIXMETASOUND_API float GetOffsetBeats(const FTimeSignature& TimeSignature) const;
+	};
 
-	float GetIntervalBeats(const FTimeSignature& TimeSignature) const;
-	float GetOffsetBeats(const FTimeSignature& TimeSignature) const;
-};
-
-namespace Harmonix
-{
 	HARMONIXMETASOUND_API void IncrementTimestampByInterval(
 			FMusicTimestamp& Timestamp,
 			const FMusicTimeInterval& Interval,
diff --git a/Engine/Plugins/Runtime/Harmonix/Source/HarmonixMetasoundTests/Private/Components/MusicClockComponentTests.cpp b/Engine/Plugins/Runtime/Harmonix/Source/HarmonixMetasoundTests/Private/Components/MusicClockComponentTests.cpp
deleted file mode 100644
diff --git a/Engine/Plugins/Runtime/Harmonix/Source/HarmonixMidi/Public/HarmonixMidi/SongMaps.h b/Engine/Plugins/Runtime/Harmonix/Source/HarmonixMidi/Public/HarmonixMidi/SongMaps.h
--- a/Engine/Plugins/Runtime/Harmonix/Source/HarmonixMidi/Public/HarmonixMidi/SongMaps.h
+++ b/Engine/Plugins/Runtime/Harmonix/Source/HarmonixMidi/Public/HarmonixMidi/SongMaps.h
@@ -25,7 +25,7 @@ namespace HarmonixMetasound
 	class FMidiClock;
 }
 
-UENUM(BlueprintType)
+UENUM()
 enum class EMidiFileQuantizeDirection : uint8
 {
 	Nearest = 0,
