HarmonixMidi

Check MidiMsg is SystemMessage when converting it to a VoiceId

#rnx
#Synced-CL 50324507

diff --git a/Engine/Plugins/Runtime/Harmonix/Source/HarmonixMidi/Public/HarmonixMidi/MidiMsg.h b/Engine/Plugins/Runtime/Harmonix/Source/HarmonixMidi/Public/HarmonixMidi/MidiMsg.h
--- a/Engine/Plugins/Runtime/Harmonix/Source/HarmonixMidi/Public/HarmonixMidi/MidiMsg.h
+++ b/Engine/Plugins/Runtime/Harmonix/Source/HarmonixMidi/Public/HarmonixMidi/MidiMsg.h
@@ -54,6 +54,7 @@ public:
 	bool  IsTempo() const { return Type == EType::Tempo; }
 	bool  IsTimeSignature() const { return Type == EType::TimeSig; }
 	bool  IsLyric() const { return Type == EType::Text && TextType == Harmonix::Midi::Constants::GMeta_Lyric; }
+	bool  IsSystemMessage() const { return Type == EType::Std && Harmonix::Midi::Constants::IsSystem(Status); }
 
 	/** Construct a standard (std) midi message */
 	UE_API FMidiMsg(uint8 InStatus, uint8 InData1, uint8 InData2);
diff --git a/Engine/Plugins/Runtime/Harmonix/Source/HarmonixMidi/Public/HarmonixMidi/MidiVoiceId.h b/Engine/Plugins/Runtime/Harmonix/Source/HarmonixMidi/Public/HarmonixMidi/MidiVoiceId.h
--- a/Engine/Plugins/Runtime/Harmonix/Source/HarmonixMidi/Public/HarmonixMidi/MidiVoiceId.h
+++ b/Engine/Plugins/Runtime/Harmonix/Source/HarmonixMidi/Public/HarmonixMidi/MidiVoiceId.h
@@ -30,7 +30,7 @@ public:
 	{
 		static_assert(FMidiVoiceGeneratorBase::IdWidth <= 20);
 		uint8 NotePart = (FromMessage.IsNoteOn() || FromMessage.IsNoteOff()) ? FromMessage.GetStdData1() : 0;
-		uint8 ChPart = (FromMessage.IsStd()) ? FromMessage.GetStdChannel() : 0;
+		uint8 ChPart = FromMessage.IsStd() && !FromMessage.IsSystemMessage() ? FromMessage.GetStdChannel() : 0;
 		Id = GeneratorId | ((uint32)(ChPart & 0xF) << 8) | (uint32)NotePart;
 	}
 
