Fixed an infinite loop that could occur in the MusicTimeInterval IncrementTimestampByBeats function if the passed in timestamp was 0;.
Made MusicClockComponent tests more robust and solved test failures on horde.

#rnx
#tests Made sure tests passed in editor

#Synced-CL 50386834

#Synced-CL 50386834

diff --git a/Engine/Plugins/Runtime/Harmonix/Source/HarmonixMetasound/Private/HarmonixMetasound/Components/MusicTimerManager.cpp b/Engine/Plugins/Runtime/Harmonix/Source/HarmonixMetasound/Private/HarmonixMetasound/Components/MusicTimerManager.cpp
--- a/Engine/Plugins/Runtime/Harmonix/Source/HarmonixMetasound/Private/HarmonixMetasound/Components/MusicTimerManager.cpp
+++ b/Engine/Plugins/Runtime/Harmonix/Source/HarmonixMetasound/Private/HarmonixMetasound/Components/MusicTimerManager.cpp
@@ -79,12 +79,6 @@ void UMusicTimerManager::RemoveTimer(const FMusicTimerHandle& Handle)
 
 void UMusicTimerManager::PauseTimer(const FMusicTimerHandle& Handle, bool bPause)
 {
-	if(LastUpdateFrame == 0)
-	{
-		UE_LOG(LogMusicTimerManager, Error, TEXT("Can not pause timers before the manager has updated at least once!"));
-		return;
-	}
-
 	auto MatchesHandleLambda = [&Handle](const TUniquePtr<FMusicTimer>& Timer) -> bool
 	{
 		return Timer && Timer->GetUniqueIdentifier() == Handle.GetUniqueIdentifier();
diff --git a/Engine/Plugins/Runtime/Harmonix/Source/HarmonixMetasound/Private/HarmonixMetasound/DataTypes/MusicTimeInterval.cpp b/Engine/Plugins/Runtime/Harmonix/Source/HarmonixMetasound/Private/HarmonixMetasound/DataTypes/MusicTimeInterval.cpp
--- a/Engine/Plugins/Runtime/Harmonix/Source/HarmonixMetasound/Private/HarmonixMetasound/DataTypes/MusicTimeInterval.cpp
+++ b/Engine/Plugins/Runtime/Harmonix/Source/HarmonixMetasound/Private/HarmonixMetasound/DataTypes/MusicTimeInterval.cpp
@@ -36,6 +36,12 @@ namespace Harmonix
 {
 	void IncrementTimestampByBeats(FMusicTimestamp& Timestamp, float Beats, const FTimeSignature& TimeSignature)
 	{
+
+        if(TimeSignature.Numerator == 0)
+        {
+            return;
+        }
+
 		// Add the beats
 		Timestamp.Beat += Beats;
 
diff --git a/Engine/Plugins/Runtime/Harmonix/Source/HarmonixMetasoundTests/Private/Components/MusicClockComponentTests.cpp b/Engine/Plugins/Runtime/Harmonix/Source/HarmonixMetasoundTests/Private/Components/MusicClockComponentTests.cpp
--- a/Engine/Plugins/Runtime/Harmonix/Source/HarmonixMetasoundTests/Private/Components/MusicClockComponentTests.cpp
+++ b/Engine/Plugins/Runtime/Harmonix/Source/HarmonixMetasoundTests/Private/Components/MusicClockComponentTests.cpp
@@ -3,7 +3,7 @@
 #include "Misc/AutomationTest.h"
 
 #if WITH_DEV_AUTOMATION_TESTS
-
+#include "Delegates/Delegate.h"
 #include "HarmonixMetasound/Components/MusicClockComponent.h"
 #include "HarmonixMetasound/Components/MusicTimerHandle.h"
 #include "HarmonixMetasound/DataTypes/MusicTimeInterval.h"
@@ -11,33 +11,12 @@
 
 namespace HarmonixMetasound::MidiStream::Tests
 {
-	class FCheckFiredLatentCommand : public IAutomationLatentCommand
-	{
-	public:
-		explicit FCheckFiredLatentCommand(FAutomationTestBase& InAutomationTestBase, TSharedRef<std::atomic<bool>, ESPMode::ThreadSafe> InFlag, bool InExpectedValue)
-			: AutomationTestBase(InAutomationTestBase), Flag(MoveTemp(InFlag))
-			, ExpectedValue(InExpectedValue)
-		{
-		}
-
-		virtual bool Update() override
-		{
-			AutomationTestBase.TestEqual(TEXT("bFired flag was unexpected value"), Flag->load(std::memory_order_acquire), ExpectedValue);
-			return true;
-		}
-
-	private:
-		FAutomationTestBase& AutomationTestBase;
-		TSharedRef<std::atomic<bool>, ESPMode::ThreadSafe> Flag;
-		bool ExpectedValue = false;
-	};
-
-
 	class FWaitUntilFireLatentCommand : public IAutomationLatentCommand
 	{
 	public:
-		FWaitUntilFireLatentCommand(TSharedRef<std::atomic<bool>, ESPMode::ThreadSafe> InFlag, float InTimeoutSeconds)
-			: Flag(MoveTemp(InFlag))
+		FWaitUntilFireLatentCommand(TSharedRef<std::atomic<bool>, ESPMode::ThreadSafe> InFlag, float InTimeoutSeconds, TOptional<FSimpleDelegate> InDelegate)
+			: Delegate(InDelegate)
+			, Flag(MoveTemp(InFlag))
 			, StartTimeSeconds(FPlatformTime::Seconds())
 			, TimeoutSeconds(InTimeoutSeconds)
 		{
@@ -45,82 +24,154 @@ namespace HarmonixMetasound::MidiStream::Tests
 
 		virtual bool Update() override
 		{
+			bool bDone = false;
+
 			// Fired
 			if (Flag->load(std::memory_order_acquire))
 			{
-				return true;
+				bDone = true;
 			}
 
 			// Timeout
 			if ((FPlatformTime::Seconds() - StartTimeSeconds) > TimeoutSeconds)
 			{
-				return true;
+				bDone = true;
+			}
+
+			if(bDone && Delegate.IsSet())
+			{
+				Delegate.GetValue().ExecuteIfBound();
 			}
 
-			// Not fired yet
-			return false;
+			return bDone;
 		}
 
 	private:
+		TOptional<FSimpleDelegate> Delegate;
 		TSharedRef<std::atomic<bool>, ESPMode::ThreadSafe> Flag;
 		double StartTimeSeconds = 0.0;
 		float TimeoutSeconds = 0.0f;
 	};
 
-	struct FMusicClockComponentHolder
+	class FMusicClockTestHelper
 	{
-		TObjectPtr<UWorld> World = nullptr;
-		TObjectPtr<AActor> Owner = nullptr;
-		TObjectPtr<UMusicClockComponent> MusicClockComponent = nullptr;
-		FMusicTimerHandle TimerHandle;
-	};
+	public:
+		FMusicClockTestHelper(FAutomationTestBase* InAutomationTestBase)
+			: AutomationTestBase(InAutomationTestBase)
+		{
 
-	static FMusicClockComponentHolder CreateMusicClock(FAutomationTestBase& AutomationTestBase)
-	{
-		FMusicClockComponentHolder ComponentHolder;
+		}
 
-		ComponentHolder.World = FAutomationEditorCommonUtils::CreateNewMap();
-		AutomationTestBase.TestNotNull("World created", ComponentHolder.World.Get());
+		void CreateMusicClock()
+		{
+			if(!AutomationTestBase)
+			{
+				return;
+			}
 
-		ComponentHolder.Owner = ComponentHolder.World->SpawnActor<AActor>();
-		AutomationTestBase.TestNotNull("Owner actor spawned", ComponentHolder.Owner.Get());
+			World = FAutomationEditorCommonUtils::CreateNewMap();
+			AutomationTestBase->TestNotNull("World created", World.Get());
 
-		bool bStart = true;
-		bool bPersistAcrossLevelTransitions = false;
-		UMidiFile* TempoMap = nullptr;
-		ComponentHolder.MusicClockComponent = UMusicClockComponent::CreateWallClockDrivenMusicClock(ComponentHolder.World, TempoMap, bStart, bPersistAcrossLevelTransitions);
+			Owner = World->SpawnActor<AActor>();
+			AutomationTestBase->TestNotNull("Owner actor spawned", Owner.Get());
 
-		return ComponentHolder;
-	}
+			bool bStart = true;
+			bool bPersistAcrossLevelTransitions = false;
+			UMidiFile* TempoMap = nullptr;
+			MusicClockComponent = UMusicClockComponent::CreateWallClockDrivenMusicClock(World, TempoMap, bStart, bPersistAcrossLevelTransitions);
+		}
 
-	FMusicClockComponentHolder StartTimer(FAutomationTestBase& AutomationTestBase, TSharedRef<std::atomic<bool>, ESPMode::ThreadSafe> bFired, FMusicTimeInterval Interval)
-	{
-		FMusicClockComponentHolder MusicClockComponentHolder = CreateMusicClock(AutomationTestBase);
-		UMusicClockComponent* MusicClockComponent = MusicClockComponentHolder.MusicClockComponent;
+		void StartTimer(const FMusicTimestamp& InStartTime, const FMusicTimeInterval& InInterval)
+		{
+			if (!MusicClockComponent || !AutomationTestBase)
+			{
+				return;
+			}
+
+			Interval = InInterval;
+			auto OnTimerFinished = TDelegate<void()>::CreateLambda([this]()
+			{
+				TimerTriggered->store(true, std::memory_order_release);
+			});
+
+			StartTime = InStartTime;
+			TimerHandle = MusicClockComponent->AddTimerNative(Interval, StartTime, Timebase, bLooping, OnTimerFinished);
+			AutomationTestBase->TestTrue("Valid Timer Handle", TimerHandle.IsValid());
+		}
+
+		FMusicTimestamp GetExpectedEndTime() const
+		{
+			FMusicTimestamp EndTime = StartTime;
+			Harmonix::IncrementTimestampByInterval(EndTime, Interval, TimeSignature);
+
+			FMusicTimestamp OffsetTimeStamp = EndTime;
+			Harmonix::IncrementTimestampByOffset(OffsetTimeStamp, Interval, TimeSignature);
+
+			return OffsetTimeStamp;
+		}
 
-		const ECalibratedMusicTimebase Timebase = ECalibratedMusicTimebase::RawAudioRenderTime;
-		const FMusicTimestamp StartTime = MusicClockComponent->GetCurrentTimestamp(Timebase);
-		const bool bLooping = false;
+		FMusicTimestamp GetStartTime() const
+		{
+			return StartTime;
+		}
 
-		auto Callback = TDelegate<void()>::CreateLambda([MusicClockComponent, StartTime, Interval, Timebase, &AutomationTestBase, bFired]()
+		FMusicTimestamp GetCurrentTimestamp() const
 		{
-			if (const FTimeSignature* TimeSignature = MusicClockComponent->GetSongMaps().GetTimeSignatureAtTick(MusicClockComponent->GetTicksIncludingCountIn(Timebase)))
+			if (!MusicClockComponent)
 			{
-				// Check what time we should have ended on
-				FMusicTimestamp EndTime = StartTime;
-				Harmonix::IncrementTimestampByInterval(EndTime, Interval, *TimeSignature);
-				const FMusicTimestamp CurrentTime = MusicClockComponent->GetCurrentTimestamp(Timebase);
-				AutomationTestBase.TestTrue("Timer triggered at the correct time", CurrentTime.Bar == EndTime.Bar && FMath::IsNearlyEqual(CurrentTime.Beat, EndTime.Beat, 0.5f));
+				return FMusicTimestamp();
 			}
 
-			bFired->store(true, std::memory_order_release);
-		});
+			return MusicClockComponent->GetCurrentTimestamp(Timebase);
+		}
 
-		MusicClockComponentHolder.TimerHandle = MusicClockComponent->AddTimerNative(Interval, StartTime, Timebase, bLooping, Callback);
-		AutomationTestBase.TestTrue("Valid Timer Handle", MusicClockComponentHolder.TimerHandle.IsValid());
+		void WaitForTimer(TOptional<FSimpleDelegate> OnTimerFinished)
+		{
+			if (!AutomationTestBase)
+			{
+				return;
+			}
 
-		return MusicClockComponentHolder;
-	}
+			AutomationTestBase->AddCommand(new FWaitUntilFireLatentCommand(TimerTriggered, TimeoutInSeconds, OnTimerFinished));
+		}
+
+		bool WasFired() const
+		{
+			return TimerTriggered->load(std::memory_order_acquire);
+		}
+
+		FMusicTimerHandle GetTimerHandle() const
+		{
+			return TimerHandle;
+		}
+
+		FTimeSignature GetTimeSignature() const
+		{
+			return TimeSignature;
+		}
+
+		TObjectPtr<UMusicClockComponent> GetMusicClockComponent() const
+		{
+			return MusicClockComponent;
+		}
+
+		int32 TimeoutInSeconds = 15;
+		bool bLooping = false;
+
+	private:
+		FMusicTimeInterval Interval;
+		ECalibratedMusicTimebase Timebase = ECalibratedMusicTimebase::RawAudioRenderTime;
+		FTimeSignature TimeSignature;
+		FMusicTimestamp StartTime;
+
+		FMusicTimerHandle TimerHandle;
+		TObjectPtr<UWorld> World = nullptr;
+		TObjectPtr<AActor> Owner = nullptr;
+		TObjectPtr<UMusicClockComponent> MusicClockComponent = nullptr;
+		FAutomationTestBase* AutomationTestBase = nullptr;
+
+		TSharedRef<std::atomic<bool>, ESPMode::ThreadSafe> TimerTriggered = MakeShared<std::atomic<bool>, ESPMode::ThreadSafe>(false);
+	};
 
 	IMPLEMENT_SIMPLE_AUTOMATION_TEST(FMusicClockTimerBeatTest, "Harmonix.Metasound.Components.MusicClockComponent.Timer.TimerBeat", EAutomationTestFlags::EditorContext | EAutomationTestFlags::EngineFilter)
 	bool FMusicClockTimerBeatTest::RunTest(const FString&)
@@ -129,13 +180,19 @@ namespace HarmonixMetasound::MidiStream::Tests
 		Interval.Interval = EMidiClockSubdivisionQuantization::Beat;
 		Interval.IntervalMultiplier = 25;
 
-		TSharedRef<std::atomic<bool>> bFired = MakeShared<std::atomic<bool>, ESPMode::ThreadSafe>(false);
-		FMusicClockComponentHolder Refs = StartTimer(*this, bFired, Interval);
+		TSharedRef<FMusicClockTestHelper> TestHelper = MakeShared<FMusicClockTestHelper>(this);
+		TestHelper->CreateMusicClock();
+		TestHelper->StartTimer(TestHelper->GetCurrentTimestamp(), Interval);
 
-		// Wait until the callback fires
-		const float Timeout = 50.0f;
-		AddCommand(new FWaitUntilFireLatentCommand(bFired, Timeout));
-		AddCommand(new FCheckFiredLatentCommand(*this, bFired, true));
+		auto OnTestFinished = TDelegate<void()>::CreateLambda([this, TestHelper]()
+		{
+			const FMusicTimestamp CurrentTime = TestHelper->GetCurrentTimestamp();
+			const FMusicTimestamp ExpectedEndTime = TestHelper->GetExpectedEndTime();
+			TestTrue("Timer callback was called", TestHelper->WasFired());
+			TestTrue("Timer triggered at the correct time", CurrentTime.Bar == ExpectedEndTime.Bar && FMath::IsNearlyEqual(CurrentTime.Beat, ExpectedEndTime.Beat, 0.5f));
+		});
+
+		TestHelper->WaitForTimer(OnTestFinished);
 
 		// Latent command, test is not over.
 		return true;
@@ -146,15 +203,21 @@ namespace HarmonixMetasound::MidiStream::Tests
 	{
 		FMusicTimeInterval Interval;
 		Interval.Interval = EMidiClockSubdivisionQuantization::Bar;
-		Interval.IntervalMultiplier = 25;
+		Interval.IntervalMultiplier = 10;
 
-		TSharedRef<std::atomic<bool>> bFired = MakeShared<std::atomic<bool>, ESPMode::ThreadSafe>(false);
-		FMusicClockComponentHolder Refs = StartTimer(*this, bFired, Interval);
+		TSharedRef<FMusicClockTestHelper> TestHelper = MakeShared<FMusicClockTestHelper>(this);
+		TestHelper->CreateMusicClock();
+		TestHelper->StartTimer(TestHelper->GetCurrentTimestamp(), Interval);
 
-		// Wait until the callback fires
-		const float Timeout = 50.0f;
-		AddCommand(new FWaitUntilFireLatentCommand(bFired, Timeout));
-		AddCommand(new FCheckFiredLatentCommand(*this, bFired, true));
+		auto OnTestFinished = TDelegate<void()>::CreateLambda([this, TestHelper]()
+		{
+			const FMusicTimestamp CurrentTime = TestHelper->GetCurrentTimestamp();
+			const FMusicTimestamp ExpectedEndTime = TestHelper->GetExpectedEndTime();
+			TestTrue("Timer callback was called", TestHelper->WasFired());
+			TestTrue("Timer triggered at the correct time", CurrentTime.Bar == ExpectedEndTime.Bar && FMath::IsNearlyEqual(CurrentTime.Beat, ExpectedEndTime.Beat, 0.5f));
+		});
+
+		TestHelper->WaitForTimer(OnTestFinished);
 
 		// Latent command, test is not over.
 		return true;
@@ -167,13 +230,19 @@ namespace HarmonixMetasound::MidiStream::Tests
 		Interval.Interval = EMidiClockSubdivisionQuantization::DottedEighthNote;
 		Interval.IntervalMultiplier = 25;
 
-		TSharedRef<std::atomic<bool>> bFired = MakeShared<std::atomic<bool>, ESPMode::ThreadSafe>(false);
-		FMusicClockComponentHolder Refs = StartTimer(*this, bFired, Interval);
+		TSharedRef<FMusicClockTestHelper> TestHelper = MakeShared<FMusicClockTestHelper>(this);
+		TestHelper->CreateMusicClock();
+		TestHelper->StartTimer(TestHelper->GetCurrentTimestamp(), Interval);
+
+		auto OnTestFinished = TDelegate<void()>::CreateLambda([this, TestHelper]()
+		{
+			const FMusicTimestamp CurrentTime = TestHelper->GetCurrentTimestamp();
+			const FMusicTimestamp ExpectedEndTime = TestHelper->GetExpectedEndTime();
+			TestTrue("Timer callback was called", TestHelper->WasFired());
+			TestTrue("Timer triggered at the correct time", CurrentTime.Bar == ExpectedEndTime.Bar && FMath::IsNearlyEqual(CurrentTime.Beat, ExpectedEndTime.Beat, 0.5f));
+		});
 
-		// Wait until the callback fires
-		const float Timeout = 50.0f;
-		AddCommand(new FWaitUntilFireLatentCommand(bFired, Timeout));
-		AddCommand(new FCheckFiredLatentCommand(*this, bFired, true));
+		TestHelper->WaitForTimer(OnTestFinished);
 
 		// Latent command, test is not over.
 		return true;
@@ -184,17 +253,23 @@ namespace HarmonixMetasound::MidiStream::Tests
 	{
 		FMusicTimeInterval Interval;
 		Interval.Interval = EMidiClockSubdivisionQuantization::Bar;
-		Interval.IntervalMultiplier = 15;
+		Interval.IntervalMultiplier = 5;
 		Interval.Offset = EMidiClockSubdivisionQuantization::Beat;
 		Interval.OffsetMultiplier = 3;
 
-		TSharedRef<std::atomic<bool>> bFired = MakeShared<std::atomic<bool>, ESPMode::ThreadSafe>(false);
-		FMusicClockComponentHolder Refs = StartTimer(*this, bFired, Interval);
+		TSharedRef<FMusicClockTestHelper> TestHelper = MakeShared<FMusicClockTestHelper>(this);
+		TestHelper->CreateMusicClock();
+		TestHelper->StartTimer(TestHelper->GetCurrentTimestamp(), Interval);
 
-		// Wait until the callback fires
-		const float Timeout = 50.0f;
-		AddCommand(new FWaitUntilFireLatentCommand(bFired, Timeout));
-		AddCommand(new FCheckFiredLatentCommand(*this, bFired, true));
+		auto OnTestFinished = TDelegate<void()>::CreateLambda([this, TestHelper]()
+		{
+			const FMusicTimestamp CurrentTime = TestHelper->GetCurrentTimestamp();
+			const FMusicTimestamp ExpectedEndTime = TestHelper->GetExpectedEndTime();
+			TestTrue("Timer callback was called", TestHelper->WasFired());
+			TestTrue("Timer triggered at the correct time", CurrentTime.Bar == ExpectedEndTime.Bar && FMath::IsNearlyEqual(CurrentTime.Beat, ExpectedEndTime.Beat, 0.5f));
+		});
+
+		TestHelper->WaitForTimer(OnTestFinished);
 
 		// Latent command, test is not over.
 		return true;
@@ -205,25 +280,45 @@ namespace HarmonixMetasound::MidiStream::Tests
 	{
 		FMusicTimeInterval Interval;
 		Interval.Interval = EMidiClockSubdivisionQuantization::Bar;
-		Interval.IntervalMultiplier = 15;
+		Interval.IntervalMultiplier = 5;
 
-		TSharedRef<std::atomic<bool>> bFired = MakeShared<std::atomic<bool>, ESPMode::ThreadSafe>(false);
+		TSharedRef<FMusicClockTestHelper> TestHelper = MakeShared<FMusicClockTestHelper>(this);
+		TestHelper->CreateMusicClock();
+		const FMusicTimestamp StartTime = TestHelper->GetCurrentTimestamp();
+		TestHelper->StartTimer(StartTime, Interval);
+		TestHelper->GetMusicClockComponent()->PauseTimer(TestHelper->GetTimerHandle(), true);
 
-		// Start timer and pause it
-		FMusicClockComponentHolder Refs = StartTimer(*this, bFired, Interval);
-		Refs.MusicClockComponent->PauseTimer(Refs.TimerHandle, true);
+		TSharedRef<FMusicTimestamp> PauseTime = MakeShared<FMusicTimestamp>();
+		auto OnTimerTriggered = TDelegate<void()>::CreateLambda([this, TestHelper, PauseTime]()
+		{
+			TestTrue("Timer callback was called", TestHelper->WasFired());
+
+			// Add the time we were paused to the expected end time
+			const FMusicTimestamp CurrentTime = TestHelper->GetCurrentTimestamp();
+			const FMusicTimestamp StartTime = TestHelper->GetStartTime();
+			FMusicTimestamp ExpectedEndTime = TestHelper->GetExpectedEndTime();
+			const FTimeSignature TimeSignature = TestHelper->GetTimeSignature();
+			const int32 BeatsPerBar = TimeSignature.Numerator;
+			const float StartTimeAbsoluteBeats = float(StartTime.Bar - 1) * BeatsPerBar + (StartTime.Beat - 1.0f);
+			const float PausedbsoluteBeats = float(PauseTime->Bar - 1) * BeatsPerBar + (PauseTime->Beat - 1.0f);
+			const float PausedDeltaBeats = PausedbsoluteBeats - StartTimeAbsoluteBeats;
+			Harmonix::IncrementTimestampByBeats(ExpectedEndTime, PausedDeltaBeats, TimeSignature);
+
+			TestTrue("Timer triggered at the correct time", CurrentTime.Bar == ExpectedEndTime.Bar && FMath::IsNearlyEqual(CurrentTime.Beat, ExpectedEndTime.Beat, 0.5f));
+		});
 
-		// Timer is paused, we should timeout. 
-		const float Timeout = 10.0f;
-		AddCommand(new FWaitUntilFireLatentCommand(bFired, Timeout));
-		AddCommand(new FCheckFiredLatentCommand(*this, bFired, false));
+		auto OnTimerTimeout = TDelegate<void()>::CreateLambda([this, TestHelper, OnTimerTriggered, PauseTime]()
+		{
+			// Timer was paused so should not have fired.
+			TestFalse("Timer callback was called", TestHelper->WasFired());
 
-		// Unpause timer and wait to see if it resumes and fires
-		Refs.MusicClockComponent->PauseTimer(Refs.TimerHandle, false);
-		bFired->store(false, std::memory_order_relaxed);
+			// Unpause
+			TestHelper->GetMusicClockComponent()->PauseTimer(TestHelper->GetTimerHandle(), false);
+			*PauseTime = TestHelper->GetCurrentTimestamp();
+			TestHelper->WaitForTimer(OnTimerTriggered);
+		});
 
-		AddCommand(new FWaitUntilFireLatentCommand(bFired, Timeout));
-		AddCommand(new FCheckFiredLatentCommand(*this, bFired, true));
+		TestHelper->WaitForTimer(OnTimerTimeout);
 
 		// Latent command, test is not over.
 		return true;
