Remove UE_API from inline method.

#rnx

#Synced-CL 50453631
#Synced-CL 50453631

diff --git a/Engine/Plugins/Runtime/Harmonix/Source/HarmonixMetasound/Public/HarmonixMetasound/Components/MusicClockComponent.h b/Engine/Plugins/Runtime/Harmonix/Source/HarmonixMetasound/Public/HarmonixMetasound/Components/MusicClockComponent.h
--- a/Engine/Plugins/Runtime/Harmonix/Source/HarmonixMetasound/Public/HarmonixMetasound/Components/MusicClockComponent.h
+++ b/Engine/Plugins/Runtime/Harmonix/Source/HarmonixMetasound/Public/HarmonixMetasound/Components/MusicClockComponent.h
@@ -548,7 +548,7 @@ public:
 	virtual bool CanAuditionInEditor() const override { return false; }
 	// ^^^ END IMusicEnvironmentClockSource Implementation ^^^
 	
-	UE_API bool HasValidClockDriver() const { return ClockDriver.IsValid(); }
+	bool HasValidClockDriver() const { return ClockDriver.IsValid(); }
 
 private:
 	friend class  UMidiClockUpdateSubsystem;
