Increased timeout for MusicClockComponentTests as they are timing out on build machines.

#rnx

#Synced-CL 50513505
#Synced-CL 50513505

diff --git a/Engine/Plugins/Runtime/Harmonix/Source/HarmonixMetasoundTests/Private/Components/MusicClockComponentTests.cpp b/Engine/Plugins/Runtime/Harmonix/Source/HarmonixMetasoundTests/Private/Components/MusicClockComponentTests.cpp
--- a/Engine/Plugins/Runtime/Harmonix/Source/HarmonixMetasoundTests/Private/Components/MusicClockComponentTests.cpp
+++ b/Engine/Plugins/Runtime/Harmonix/Source/HarmonixMetasoundTests/Private/Components/MusicClockComponentTests.cpp
@@ -155,7 +155,7 @@ namespace HarmonixMetasound::MidiStream::Tests
 			return MusicClockComponent;
 		}
 
-		int32 TimeoutInSeconds = 15;
+		int32 TimeoutInSeconds = 60;
 		bool bLooping = false;
 
 	private:
