Track Relative Render Cost of the Fusion Sampler

#rnx

#Synced-CL 50667372

diff --git a/Engine/Plugins/Runtime/Harmonix/Source/HarmonixMetasound/Private/HarmonixMetasound/Nodes/FusionSamplerNode.cpp b/Engine/Plugins/Runtime/Harmonix/Source/HarmonixMetasound/Private/HarmonixMetasound/Nodes/FusionSamplerNode.cpp
--- a/Engine/Plugins/Runtime/Harmonix/Source/HarmonixMetasound/Private/HarmonixMetasound/Nodes/FusionSamplerNode.cpp
+++ b/Engine/Plugins/Runtime/Harmonix/Source/HarmonixMetasound/Private/HarmonixMetasound/Nodes/FusionSamplerNode.cpp
@@ -142,6 +142,7 @@ namespace HarmonixMetasound
 			, EnableMTRenderingInPin(Args.InMTRenderingEnable)
 			, VoicePoolName(Args.InVoicePoolName)
 			, UsePitchShifters(Args.InUsePitchShifters)
+			, CostReporter(Args.Params.GraphRenderCost ? Args.Params.GraphRenderCost->AddNode(Args.Params.Node.GetInstanceID(), Args.Params.Environment) : FNodeRenderCost{})
 			, OutputChannelLayout(Args.InOutputChannelLayout)
 		{
 			Reset(Args.Params);
@@ -229,6 +230,7 @@ namespace HarmonixMetasound
 		const bool UsePitchShifters;
 
 		//** DATA
+		FNodeRenderCost CostReporter;
 		int32 FramesPerBlock = 0;
 		int32 CurrentTrackNumber = 0;
 		bool MadeAudioLastFrame = false;
@@ -580,11 +582,13 @@ namespace HarmonixMetasound
 			{
 				TRACE_CPUPROFILER_EVENT_SCOPE_STR("FusionAsyncRender");
 				DoRender();
+				CostReporter.SetRenderCost(GetNumVoicesInUse());
 			});
 		}
 		else
 		{
 			DoRender();
+			CostReporter.SetRenderCost(GetNumVoicesInUse());
 		}
 	}
 
