[Harmonix]

MidiTrack AddText.

Perform a case sensitive comparison to add unique strings. "AddUnique" for strings defaults to "IgnoreCase" behavior, which is not what you might expect.

#rnx

#Synced-CL 51211507

diff --git a/Engine/Plugins/Runtime/Harmonix/Source/HarmonixMidi/Public/HarmonixMidi/MidiTrack.h b/Engine/Plugins/Runtime/Harmonix/Source/HarmonixMidi/Public/HarmonixMidi/MidiTrack.h
--- a/Engine/Plugins/Runtime/Harmonix/Source/HarmonixMidi/Public/HarmonixMidi/MidiTrack.h
+++ b/Engine/Plugins/Runtime/Harmonix/Source/HarmonixMidi/Public/HarmonixMidi/MidiTrack.h
@@ -68,7 +68,25 @@ public:
 		return Strings[Index];
 	}
 
-	uint16 AddText(const FString& Str) { return (uint16)Strings.AddUnique(Str); }
+	int32 AddText(const FString& Str)
+	{
+		int32 Index = INDEX_NONE;
+		for (int32 Idx = 0; Idx < Strings.Num(); ++Idx)
+		{
+			if (Strings[Idx].Equals(Str, ESearchCase::CaseSensitive))
+			{
+				Index = Idx;
+				break;
+			}
+		}
+		
+		if (Index != INDEX_NONE)
+		{
+			return Index;
+		}
+		
+		return Strings.Add(Str);
+	}
 
 	FMidiTextRepository* GetTextRepository() { return &Strings; }
 	const FMidiTextRepository* GetTextRepository() const { return &Strings; }
