PR #14213: Fix: Crash in MidiStreamWriter for Text MidiEventType

Additionally, make sure text event indices match string repository indices. and ensure strings match.

#rnx
#Synced-CL 51211507

diff --git a/Engine/Plugins/Runtime/Harmonix/Source/HarmonixMetasound/Private/HarmonixMetasound/MidiOps/MidiStreamWriter.cpp b/Engine/Plugins/Runtime/Harmonix/Source/HarmonixMetasound/Private/HarmonixMetasound/MidiOps/MidiStreamWriter.cpp
--- a/Engine/Plugins/Runtime/Harmonix/Source/HarmonixMetasound/Private/HarmonixMetasound/MidiOps/MidiStreamWriter.cpp
+++ b/Engine/Plugins/Runtime/Harmonix/Source/HarmonixMetasound/Private/HarmonixMetasound/MidiOps/MidiStreamWriter.cpp
@@ -37,8 +37,40 @@ namespace Harmonix::Midi::Ops
 				const int32 OffsetTick = MidiEvent.CurrentMidiTick - FirstTickToProcess;
 				const int32 MidiTick = NextWriteTick + OffsetTick;
 				FMidiTrack& MidiTrack = MidiTracks.FindOrAdd(MidiEvent.TrackIndex);
-				MidiTrack.AddEvent(FMidiEvent(MidiTick, MidiEvent.MidiMessage));
-				AddedEvents = true;
+				if (MidiEvent.MidiMessage.IsText())
+				{
+					uint16 TextIndex = MidiEvent.MidiMessage.GetTextIndex();
+					const FString* EventText = InStream.GetMidiTrackText(MidiEvent.TrackIndex, TextIndex);
+					if (EventText == nullptr)
+					{
+						continue;
+					}
+
+					// make sure the text index is the same in the TextRepository as the event
+					// this can potentially leave empty strings in the repository, which isn't optimal
+					FMidiTextRepository& TextRepository = *MidiTrack.GetTextRepository();
+					if (!TextRepository.IsValidIndex(TextIndex))
+					{
+						TextRepository.SetNum(TextIndex + 1);
+						TextRepository[TextIndex] = *EventText;
+					}
+
+					// found a text entry mismatch - should we overrwrite or ignore?
+					if (!ensure(TextRepository[TextIndex] == *EventText))
+					{
+						TextRepository[TextIndex] = *EventText;
+					}
+
+					// only add text events we have strings for
+					MidiTrack.AddEvent(FMidiEvent(MidiTick, MidiEvent.MidiMessage));
+					AddedEvents = true;	
+				}
+				else
+				{
+					MidiTrack.AddEvent(FMidiEvent(MidiTick, MidiEvent.MidiMessage));
+					AddedEvents = true;						
+				}
+
 			}
 		}
 		NextWriteTick += LastTickToProcess - FirstTickToProcess + 1;
