Fixup API macro usage

#rnx

#Synced-CL 51577459

diff --git a/Engine/Plugins/Compositing/OpenColorIO/Source/OpenColorIO/Public/OpenColorIOConfiguration.h b/Engine/Plugins/Compositing/OpenColorIO/Source/OpenColorIO/Public/OpenColorIOConfiguration.h
--- a/Engine/Plugins/Compositing/OpenColorIO/Source/OpenColorIO/Public/OpenColorIOConfiguration.h
+++ b/Engine/Plugins/Compositing/OpenColorIO/Source/OpenColorIO/Public/OpenColorIOConfiguration.h
@@ -139,7 +139,7 @@ public:
 
 	/** Get the configuration file path. */
 	UFUNCTION(BlueprintGetter)
-	UE_API const FFilePath& GetConfigurationFile() const { return ConfigurationFile; }
+	const FFilePath& GetConfigurationFile() const { return ConfigurationFile; }
 
 	/** Set the configuration file path. */
 	UFUNCTION(BlueprintSetter)
diff --git a/Engine/Plugins/Runtime/CommonUI/Source/CommonUI/Public/CommonButtonBase.h b/Engine/Plugins/Runtime/CommonUI/Source/CommonUI/Public/CommonButtonBase.h
--- a/Engine/Plugins/Runtime/CommonUI/Source/CommonUI/Public/CommonButtonBase.h
+++ b/Engine/Plugins/Runtime/CommonUI/Source/CommonUI/Public/CommonButtonBase.h
@@ -254,13 +254,13 @@ public:
 	COMMONUI_API bool IsPressed() const;
 
 	COMMONUI_API void SetMinDesiredHeight(int32 InMinHeight);
-	COMMONUI_API int32 GetMinHeight() const { return MinHeight; }
+	int32 GetMinHeight() const { return MinHeight; }
 	COMMONUI_API void SetMinDesiredWidth(int32 InMinWidth);
-	COMMONUI_API int32 GetMinWidth() const { return MinWidth; }
+	int32 GetMinWidth() const { return MinWidth; }
 	COMMONUI_API void SetMaxDesiredHeight(int32 InMaxHeight);
-	COMMONUI_API int32 GetMaxHeight() const { return MaxHeight; }
+	int32 GetMaxHeight() const { return MaxHeight; }
 	COMMONUI_API void SetMaxDesiredWidth(int32 InMaxWidth);
-	COMMONUI_API int32 GetMaxWidth() const { return MaxWidth; }
+	int32 GetMaxWidth() const { return MaxWidth; }
 
 	inline TSharedPtr<class SCommonButton> GetCommonButton() const { return MyCommonButton; }
 
diff --git a/Engine/Plugins/Runtime/Harmonix/Source/HarmonixMetasound/Public/HarmonixMetasound/Components/MusicTimerManager.h b/Engine/Plugins/Runtime/Harmonix/Source/HarmonixMetasound/Public/HarmonixMetasound/Components/MusicTimerManager.h
--- a/Engine/Plugins/Runtime/Harmonix/Source/HarmonixMetasound/Public/HarmonixMetasound/Components/MusicTimerManager.h
+++ b/Engine/Plugins/Runtime/Harmonix/Source/HarmonixMetasound/Public/HarmonixMetasound/Components/MusicTimerManager.h
@@ -60,7 +60,7 @@ public:
 	UFUNCTION(BlueprintCallable, Category = "Music|Timer")
 	FMusicTimestamp Quantize(const UMusicClockComponent* MusicClock, const FMusicTimestamp& Time, EMidiClockSubdivisionQuantization QuantizationInterval, EMidiFileQuantizeDirection InDirection) const;
 
-	UE_API uint64 GetUniqueTimerId() { return ++NextTimerId; }
+	uint64 GetUniqueTimerId() { return ++NextTimerId; }
 private:
 
 	// Used so pause doesn't need to pass in this info
@@ -72,4 +72,4 @@ private:
 	static std::atomic<uint64> NextTimerId;
 };
 
-#undef UE_API
\ No newline at end of file
+#undef UE_API
diff --git a/Engine/Plugins/Runtime/Harmonix/Source/HarmonixMetasound/Public/HarmonixMetasound/DataTypes/MusicTimeInterval.h b/Engine/Plugins/Runtime/Harmonix/Source/HarmonixMetasound/Public/HarmonixMetasound/DataTypes/MusicTimeInterval.h
--- a/Engine/Plugins/Runtime/Harmonix/Source/HarmonixMetasound/Public/HarmonixMetasound/DataTypes/MusicTimeInterval.h
+++ b/Engine/Plugins/Runtime/Harmonix/Source/HarmonixMetasound/Public/HarmonixMetasound/DataTypes/MusicTimeInterval.h
@@ -5,13 +5,15 @@
 #include "MidiClock.h"
 #include "MusicTimeInterval.generated.h"
 
-USTRUCT(BlueprintType)
-struct HARMONIXMETASOUND_API FMusicTimeInterval
+#define UE_API HARMONIXMETASOUND_API
+
+USTRUCT(MinimalAPI, BlueprintType)
+struct FMusicTimeInterval
 {
 	GENERATED_BODY()
 
-	FMusicTimeInterval();
-	FMusicTimeInterval(EMidiClockSubdivisionQuantization InInterval, EMidiClockSubdivisionQuantization InOffset, int32 InIntervalMultiplier, int32 InOffsetMultiplier);
+	UE_API FMusicTimeInterval();
+	UE_API FMusicTimeInterval(EMidiClockSubdivisionQuantization InInterval, EMidiClockSubdivisionQuantization InOffset, int32 InIntervalMultiplier, int32 InOffsetMultiplier);
 
 	UPROPERTY(EditAnywhere, BlueprintReadWrite, Category = "Time Interval Settings")
 	EMidiClockSubdivisionQuantization Interval = EMidiClockSubdivisionQuantization::Beat;
@@ -25,8 +27,8 @@ struct HARMONIXMETASOUND_API FMusicTimeInterval
 	UPROPERTY(EditAnywhere, BlueprintReadWrite, Category = "Time Interval Settings")
 	int32 OffsetMultiplier = 0;
 
-	float GetIntervalBeats(const FTimeSignature& TimeSignature) const;
-	float GetOffsetBeats(const FTimeSignature& TimeSignature) const;
+	UE_API float GetIntervalBeats(const FTimeSignature& TimeSignature) const;
+	UE_API float GetOffsetBeats(const FTimeSignature& TimeSignature) const;
 };
 
 namespace Harmonix
@@ -46,3 +48,5 @@ namespace Harmonix
 		float Beats,
 		const FTimeSignature& TimeSignature);
 }
+
+#undef UE_API
diff --git a/Engine/Source/Runtime/Core/Public/Misc/QueuedThreadPoolWrapper.h b/Engine/Source/Runtime/Core/Public/Misc/QueuedThreadPoolWrapper.h
--- a/Engine/Source/Runtime/Core/Public/Misc/QueuedThreadPoolWrapper.h
+++ b/Engine/Source/Runtime/Core/Public/Misc/QueuedThreadPoolWrapper.h
@@ -69,7 +69,7 @@ public:
 	/**
 	 *  Return the estimated memory currently used by this thread pool, if supported, or -1 if not.
 	 */
-	CORE_API virtual int64 GetEstimatedMemory() const
+	virtual int64 GetEstimatedMemory() const
 	{
 		return -1ll;
 	}
