[Harmonix] Add MsToBeat() method that returns zero indexed beat, add comments about the asymmetrical one indexed beats returned from "FractionalBeatAt" methods.

#Synced-CL 51837090

diff --git a/Engine/Plugins/Runtime/Harmonix/Source/HarmonixMetasound/Private/HarmonixMetasound/Components/MusicClockComponent.cpp b/Engine/Plugins/Runtime/Harmonix/Source/HarmonixMetasound/Private/HarmonixMetasound/Components/MusicClockComponent.cpp
--- a/Engine/Plugins/Runtime/Harmonix/Source/HarmonixMetasound/Private/HarmonixMetasound/Components/MusicClockComponent.cpp
+++ b/Engine/Plugins/Runtime/Harmonix/Source/HarmonixMetasound/Private/HarmonixMetasound/Components/MusicClockComponent.cpp
@@ -654,6 +654,12 @@ float UMusicClockComponent::BeatToMs(float Beat) const
 	return GetSongMaps().GetMsAtBeat(Beat);
 }
 
+float UMusicClockComponent::MsToBeat(float Ms) const
+{
+	// "BeatAt" is one indexed, so we have to subtract one.
+	return GetSongMaps().GetFractionalBeatAtMs(Ms) - 1.0f;
+}
+
 float UMusicClockComponent::GetMsPerBeatAtMs(float Ms) const
 {
 	return GetSongMaps().GetMsPerBeatAtMs(Ms);
diff --git a/Engine/Plugins/Runtime/Harmonix/Source/HarmonixMetasound/Public/HarmonixMetasound/Components/MusicClockComponent.h b/Engine/Plugins/Runtime/Harmonix/Source/HarmonixMetasound/Public/HarmonixMetasound/Components/MusicClockComponent.h
--- a/Engine/Plugins/Runtime/Harmonix/Source/HarmonixMetasound/Public/HarmonixMetasound/Components/MusicClockComponent.h
+++ b/Engine/Plugins/Runtime/Harmonix/Source/HarmonixMetasound/Public/HarmonixMetasound/Components/MusicClockComponent.h
@@ -324,9 +324,14 @@ public:
 	UFUNCTION(BlueprintPure, Category = "Tick")
 	UE_API float TickToMs(float Tick) const;
 
+	/* Convert zero indexed fractional beat to Ms */ 
 	UFUNCTION(BlueprintPure, Category = "Beat")
 	UE_API float BeatToMs(float Beat) const;
 
+	/* Convert Ms to zero indexed fractional beat */
+	UFUNCTION(BlueprintPure, Category = "Beat")
+	UE_API float MsToBeat(float Ms) const;
+
 	UFUNCTION(BlueprintPure, Category = "Beat")
 	UE_API float GetMsPerBeatAtMs(float Ms) const;
 
diff --git a/Engine/Plugins/Runtime/Harmonix/Source/HarmonixMidi/Public/HarmonixMidi/SongMaps.h b/Engine/Plugins/Runtime/Harmonix/Source/HarmonixMidi/Public/HarmonixMidi/SongMaps.h
--- a/Engine/Plugins/Runtime/Harmonix/Source/HarmonixMidi/Public/HarmonixMidi/SongMaps.h
+++ b/Engine/Plugins/Runtime/Harmonix/Source/HarmonixMidi/Public/HarmonixMidi/SongMaps.h
@@ -98,12 +98,12 @@ public:
 
 	// beats
 	UE_API const FBeatMapPoint* GetBeatAtMs(float Ms) const;
-	UE_API float GetMsAtBeat(float Beat) const;
+	UE_API float GetMsAtBeat(float Beat) const; // Beat is zero indexed.
 	UE_API const FBeatMapPoint* GetBeatAtTick(int32 Tick) const;
 	UE_API float GetMsPerBeatAtMs(float Ms) const;
 	UE_API float GetMsPerBeatAtTick(int32 Tick) const;
-	UE_API float GetFractionalBeatAtMs(float Ms) const;
-	UE_API float GetFractionalBeatAtTick(float Tick) const;
+	UE_API float GetFractionalBeatAtMs(float Ms) const; // Returns ONE indexed beat, beware.
+	UE_API float GetFractionalBeatAtTick(float Tick) const; // Returns ONE indexed beat, beware.
 	UE_API int32 GetBeatIndexAtMs(float Ms) const;
 	UE_API int32 GetBeatIndexAtTick(int32 Tick) const;
 	UE_API EMusicalBeatType GetBeatTypeAtMs(float Ms) const;
