Use PitchTime vs StreamTime to adjust for latency of pitch detection.

Start implement "Timestamping". Timestamping happens on the analysis thread (This is wrong!) Close enough for now. Will move it to the Capture Thread in next CL

#rnx
#Synced-CL 51952668

#Synced-CL 51952668

diff --git a/Engine/Plugins/Runtime/Harmonix/Source/HarmonixMetasound/Private/HarmonixMetasound/Subsystems/MidiClockUpdateSubsystem.cpp b/Engine/Plugins/Runtime/Harmonix/Source/HarmonixMetasound/Private/HarmonixMetasound/Subsystems/MidiClockUpdateSubsystem.cpp
--- a/Engine/Plugins/Runtime/Harmonix/Source/HarmonixMetasound/Private/HarmonixMetasound/Subsystems/MidiClockUpdateSubsystem.cpp
+++ b/Engine/Plugins/Runtime/Harmonix/Source/HarmonixMetasound/Private/HarmonixMetasound/Subsystems/MidiClockUpdateSubsystem.cpp
@@ -19,6 +19,8 @@ namespace MidiClockUpdateSubsystem
 FCriticalSection UMidiClockUpdateSubsystem::ClockHistoryMapLocker;
 TMap<uint32, TWeakPtr<HarmonixMetasound::Analysis::FMidiClockSongPositionHistory>> UMidiClockUpdateSubsystem::ClockHistories;
 
+double UMidiClockUpdateSubsystem::UpdateTimeSeconds = 0.0;
+
 void UMidiClockUpdateSubsystem::Initialize(FSubsystemCollectionBase& Collection)
 {
 	EngineSamplingInputDelegate = FCoreDelegates::OnSamplingInput.AddUObject(this, &UMidiClockUpdateSubsystem::CoreDelegatesSamplingInput);
@@ -51,6 +53,11 @@ void UMidiClockUpdateSubsystem::StopTrackingMusicClockComponent(UMusicClockCompo
 	}
 }
 
+double UMidiClockUpdateSubsystem::GetUpdateTimeSeconds()
+{
+	return UpdateTimeSeconds;
+}
+
 uint32 UMidiClockUpdateSubsystem::MakeMidiSongPosAnalyzerAddressHash(const Metasound::Frontend::FAnalyzerAddress& ForAddress)
 {
 	uint32 AddressHash = GetTypeHashHelper(ForAddress.AnalyzerMemberName);
@@ -115,6 +122,7 @@ void UMidiClockUpdateSubsystem::UpdateUMusicClockComponents()
 {
 	TRACE_CPUPROFILER_EVENT_SCOPE(UpdateUMusicClockComponents);
 
+	UpdateTimeSeconds = FPlatformTime::Seconds();
 	using IteratorType = TArray<TWeakObjectPtr<UMusicClockComponent>>::TIterator;
 	for (IteratorType ClockIterator = TrackedMusicClockComponents.CreateIterator(); ClockIterator; ++ClockIterator)
 	{
diff --git a/Engine/Plugins/Runtime/Harmonix/Source/HarmonixMetasound/Public/HarmonixMetasound/Subsystems/MidiClockUpdateSubsystem.h b/Engine/Plugins/Runtime/Harmonix/Source/HarmonixMetasound/Public/HarmonixMetasound/Subsystems/MidiClockUpdateSubsystem.h
--- a/Engine/Plugins/Runtime/Harmonix/Source/HarmonixMetasound/Public/HarmonixMetasound/Subsystems/MidiClockUpdateSubsystem.h
+++ b/Engine/Plugins/Runtime/Harmonix/Source/HarmonixMetasound/Public/HarmonixMetasound/Subsystems/MidiClockUpdateSubsystem.h
@@ -47,6 +47,8 @@ public:
 	static UE_API void TrackMusicClockComponent(UMusicClockComponent* Clock);
 	static UE_API void StopTrackingMusicClockComponent(UMusicClockComponent* Clock);
 
+	static UE_API double GetUpdateTimeSeconds();
+
 	using FClockHistoryPtr = TSharedPtr<HarmonixMetasound::Analysis::FMidiClockSongPositionHistory>;
 	static UE_API FClockHistoryPtr GetOrCreateClockHistory(const Metasound::Frontend::FAnalyzerAddress& ForAddress);
 
@@ -78,6 +80,7 @@ public:
 private:
 
 	mutable FCriticalSection TrackedMidiClocksMutex;
+	static double UpdateTimeSeconds;
 
 	TArray<TWeakObjectPtr<UMusicClockComponent>> TrackedMusicClockComponents;
 	UE_API void TrackMusicClockComponentImpl(UMusicClockComponent* Clock);
