Harmonix "Audio Data Renderer"

Use SoundWaveProxyPlayer instead of SoundWaveProxyReader

#rnx

#Synced-CL 52025364

diff --git a/Engine/Plugins/Runtime/Harmonix/Source/HarmonixDsp/Private/HarmonixDsp/AudioData/StreamingAudioRendererV2.cpp b/Engine/Plugins/Runtime/Harmonix/Source/HarmonixDsp/Private/HarmonixDsp/AudioData/StreamingAudioRendererV2.cpp
--- a/Engine/Plugins/Runtime/Harmonix/Source/HarmonixDsp/Private/HarmonixDsp/AudioData/StreamingAudioRendererV2.cpp
+++ b/Engine/Plugins/Runtime/Harmonix/Source/HarmonixDsp/Private/HarmonixDsp/AudioData/StreamingAudioRendererV2.cpp
@@ -1,16 +1,12 @@
 // Copyright Epic Games, Inc. All Rights Reserved.
 #include "HarmonixDsp/AudioData/StreamingAudioRendererV2.h"
-
-#include "HarmonixDsp/FusionSampler/FusionSampler.h"
-
 #include "HarmonixDsp/AudioUtility.h"
+#include "HarmonixDsp/FusionSampler/FusionSampler.h"
 #include "HarmonixDsp/GainMatrix.h"
 #include "HarmonixDsp/PannerDetails.h"
-
-#include "Sound/SoundWave.h"
-#include "Sound/SoundWaveProxyReader.h"
-
 #include "Math/UnrealMathUtility.h"
+#include "Sound/SoundWave.h"
+#include "Sound/SoundWaveProxyPlayer.h"
 
 DEFINE_LOG_CATEGORY(LogHarmonixStreamingAudioRendererV2);
 
@@ -24,8 +20,7 @@ FStreamingAudioRendererV2::~FStreamingAudioRendererV2()
 
 void FStreamingAudioRendererV2::Reset()
 {
-	WaveProxyReader.Reset();
-	InterleavedCircularBuffer.SetNum(0);
+	WavePlayer.Reset();
 	Shifter = nullptr;
 }
 
@@ -33,27 +28,18 @@ void FStreamingAudioRendererV2::SetAudioData(TSharedRef<FSoundWaveProxy> InSound
 {
 	SoundWaveProxy = InSoundWave.ToSharedPtr();
 
-	WaveProxyReader.Reset();
-	WaveProxyReader = CreateProxyReader(InSoundWave);
-
-	check(WaveProxyReader.IsValid());
+	WavePlayer.Reset();
+	WavePlayer = CreateWavePlayer(InSoundWave);
 
+	check(WavePlayer->IsPlayerValid())
+	
 	int32 WaveProxyNumChannels = SoundWaveProxy->GetNumChannels();
-
-	int32 DecodeBufferSize = WaveProxyNumChannels * DeinterleaveBlockSizeInFrames;
-	DecodeBuffer.Reset(DecodeBufferSize);
-	DecodeBuffer.AddUninitialized(DecodeBufferSize);
-
+	
 	NumDeinterleaveChannels = WaveProxyNumChannels;
 
 	LastLoopFrameCache.SetNumUninitialized(NumDeinterleaveChannels);
 	bLastLoopFrameCached = false;
 
-	// interleaved circular buffer for streaming source audio data
-	// make it larger than then amount we decode each block to avoid underruns
-	int32 FrameCapacity = DecodeBufferSize * 2;
-	InterleavedCircularBuffer.SetCapacity(FrameCapacity);
-
 	TrackChannelInfo = InSettings.TrackChannelInfo;
 	Shifter = InSettings.Shifter;
 	MySampler = InSettings.Sampler;
@@ -76,13 +62,13 @@ void FStreamingAudioRendererV2::MigrateToSampler(const FFusionSampler* InSampler
 
 void FStreamingAudioRendererV2::SetFrame(uint32 InFrameNum)
 {
-	if (!WaveProxyReader->HasFailed())
+	if (WavePlayer->IsPlayerValid())
 	{
 		SeekSourceAudioToFrame(InFrameNum);
 	}
 }
 
-double FStreamingAudioRendererV2::Render(TAudioBuffer<float>& OutBuffer, double InPos, int32 InMaxFrame, double InResampleInc, double InPitchShift, double InSpeed, bool MaintainPitchWhenSpeedChanges, bool InShouldHonorLoopPoints, const FGainMatrix& InGain)
+double FStreamingAudioRendererV2::Render(TAudioBuffer<float>& OutBuffer, double InPos, int32 InMaxFrame, double InResampleInc, double InPitchShift, double InSpeed, bool MaintainPitchWhenSpeedChanges, const FGainMatrix& InGain)
 {
 	if (!SoundWaveProxy || SoundWaveProxy->GetNumFrames() == 0)
 	{
@@ -97,13 +83,13 @@ double FStreamingAudioRendererV2::Render(TAudioBuffer<float>& OutBuffer, double
 		{
 			Increment *= InSpeed;
 		}
-		return RenderInternal(OutBuffer, InPos, InMaxFrame, Increment, InShouldHonorLoopPoints, InGain);
+		return RenderInternal(OutBuffer, InPos, InMaxFrame, Increment, InGain);
 	}
 
-	return Shifter->Render(OutBuffer, InPos, InMaxFrame, InResampleInc, InPitchShift, InSpeed, MaintainPitchWhenSpeedChanges, InShouldHonorLoopPoints, InGain);
+	return Shifter->Render(OutBuffer, InPos, InMaxFrame, InResampleInc, InPitchShift, InSpeed, MaintainPitchWhenSpeedChanges, InGain);
 }
 
-double FStreamingAudioRendererV2::RenderInternal(TAudioBuffer<float>& OutBuffer, double InPos, int32 InMaxFrame, double InInc, bool InShouldHonorLoopPoints, const FGainMatrix& InGain)
+double FStreamingAudioRendererV2::RenderInternal(TAudioBuffer<float>& OutBuffer, double InPos, int32 InMaxFrame, double InInc, const FGainMatrix& InGain)
 {
 	check(OutBuffer.GetNumValidFrames() <= AudioRendering::kMicroSliceSize);
 
@@ -116,25 +102,25 @@ double FStreamingAudioRendererV2::RenderInternal(TAudioBuffer<float>& OutBuffer,
 	uint32 NumOutFrames = OutBuffer.GetNumValidFrames();
 
 	FLerpData LerpArray[AudioRendering::kMicroSliceSize];
-	InPos = CalculateLerpData(LerpArray, AudioRendering::kMicroSliceSize, NumOutFrames, InPos, InMaxFrame, InShouldHonorLoopPoints, InInc);
+	InPos = CalculateLerpData(LerpArray, AudioRendering::kMicroSliceSize, NumOutFrames, InPos, InMaxFrame, InInc);
 
 	if (TrackChannelInfo && TrackChannelInfo->Num() > 0)
 	{
-		RenderMultiChannelRoutedUnshifted(OutBuffer, LerpArray, NumOutFrames, InGain, InInc, InShouldHonorLoopPoints);
+		RenderMultiChannelRoutedUnshifted(OutBuffer, LerpArray, NumOutFrames, InGain, InInc);
 	}
 	else if (SoundWaveProxy->GetNumChannels() <= 2)
 	{
-		RenderSimpleUnshifted(OutBuffer, LerpArray, NumOutFrames, InGain, InInc, InShouldHonorLoopPoints);
+		RenderSimpleUnshifted(OutBuffer, LerpArray, NumOutFrames, InGain, InInc);
 	}
 	else
 	{
-		RenderMultiChannelUnshifted(OutBuffer, LerpArray, NumOutFrames, InGain, InInc, InShouldHonorLoopPoints);
+		RenderMultiChannelUnshifted(OutBuffer, LerpArray, NumOutFrames, InGain, InInc);
 	}
 
 	return InPos;
 }
 
-double FStreamingAudioRendererV2::RenderUnshifted(TAudioBuffer<float>& OutBuffer, double InPos, int32 InMaxFrame, double InInc, bool InShouldHonorLoopPoints, const FGainMatrix& InGain)
+double FStreamingAudioRendererV2::RenderUnshifted(TAudioBuffer<float>& OutBuffer, double InPos, int32 InMaxFrame, double InInc, const FGainMatrix& InGain)
 {
 	if (!SoundWaveProxy)
 	{
@@ -142,7 +128,7 @@ double FStreamingAudioRendererV2::RenderUnshifted(TAudioBuffer<float>& OutBuffer
 		return InPos;
 	}
 
-	return RenderInternal(OutBuffer, InPos, InMaxFrame, InInc, InShouldHonorLoopPoints, InGain);
+	return RenderInternal(OutBuffer, InPos, InMaxFrame, InInc, InGain);
 }
 
 int32 FStreamingAudioRendererV2::CalculateNumFramesNeeded(const FLerpData* LerpData, int32 NumPoints)
@@ -177,7 +163,7 @@ int32 FStreamingAudioRendererV2::CalculateNumFramesNeeded(const FLerpData* LerpD
 	return NumNeeded;
 }
 
-void FStreamingAudioRendererV2::RenderSimpleUnshifted(TAudioBuffer<float>& OutBuffer, const FLerpData* LerpArray, uint32 InNumFrames, const FGainMatrix& InGain, double InInc, bool InShouldHonorLoopPoints)
+void FStreamingAudioRendererV2::RenderSimpleUnshifted(TAudioBuffer<float>& OutBuffer, const FLerpData* LerpArray, uint32 InNumFrames, const FGainMatrix& InGain, double InInc)
 {
 	check(InNumFrames <= AudioRendering::kMicroSliceSize);
 
@@ -199,7 +185,7 @@ void FStreamingAudioRendererV2::RenderSimpleUnshifted(TAudioBuffer<float>& OutBu
 	}
 	
 	WorkBuffer.SetNum(NumSamplesNeeded);
-	GenerateSourceAudio(StartFrameIndex, WorkBuffer, InShouldHonorLoopPoints);
+	GenerateSourceAudio(StartFrameIndex, WorkBuffer);
 
 	int32 PosA;
 	int32 PosB;
@@ -233,7 +219,7 @@ void FStreamingAudioRendererV2::RenderSimpleUnshifted(TAudioBuffer<float>& OutBu
 	}
 }
 
-void FStreamingAudioRendererV2::RenderMultiChannelUnshifted(TAudioBuffer<float>& OutBuffer, const FLerpData* LerpArray, uint32 InNumFrames, const FGainMatrix& InGain, double InInc, bool InShouldHonorLoopPoints)
+void FStreamingAudioRendererV2::RenderMultiChannelUnshifted(TAudioBuffer<float>& OutBuffer, const FLerpData* LerpArray, uint32 InNumFrames, const FGainMatrix& InGain, double InInc)
 {
 	check(InNumFrames <= AudioRendering::kMicroSliceSize);
 
@@ -260,7 +246,7 @@ void FStreamingAudioRendererV2::RenderMultiChannelUnshifted(TAudioBuffer<float>&
 	}
 	
 	WorkBuffer.SetNum(NumSamplesNeeded);
-	GenerateSourceAudio(StartFrameIndex, WorkBuffer, InShouldHonorLoopPoints);
+	GenerateSourceAudio(StartFrameIndex, WorkBuffer);
 
 	int32 PosA;
 	int32 PosB;
@@ -306,7 +292,7 @@ void FStreamingAudioRendererV2::RenderMultiChannelUnshifted(TAudioBuffer<float>&
 	}
 }
 
-void FStreamingAudioRendererV2::RenderMultiChannelRoutedUnshifted(TAudioBuffer<float>& OutBuffer, const FLerpData* LerpArray, uint32 InNumFrames, const FGainMatrix& InGain, double InInc, bool InShouldHonorLoopPoints)
+void FStreamingAudioRendererV2::RenderMultiChannelRoutedUnshifted(TAudioBuffer<float>& OutBuffer, const FLerpData* LerpArray, uint32 InNumFrames, const FGainMatrix& InGain, double InInc)
 {
 	check(InNumFrames <= AudioRendering::kMicroSliceSize);
 
@@ -333,7 +319,7 @@ void FStreamingAudioRendererV2::RenderMultiChannelRoutedUnshifted(TAudioBuffer<f
 	}
 	
 	WorkBuffer.SetNum(NumSamplesNeeded);
-	GenerateSourceAudio(StartFrameIndex, WorkBuffer, InShouldHonorLoopPoints);
+	GenerateSourceAudio(StartFrameIndex, WorkBuffer);
 
 	int32 PosA;
 	int32 PosB;
@@ -413,7 +399,7 @@ void FStreamingAudioRendererV2::RenderMultiChannelRoutedUnshifted(TAudioBuffer<f
 
 void FStreamingAudioRendererV2::SeekSourceAudioToFrame(uint32 FrameIdx)
 {
-	check(!WaveProxyReader->HasFailed());
+	check(WavePlayer->IsPlayerValid());
 	
 	// no need to seek
 	uint32 SourceFrameIndex = GetSourceAudioFrameIndex();
@@ -422,7 +408,7 @@ void FStreamingAudioRendererV2::SeekSourceAudioToFrame(uint32 FrameIdx)
 	{
 		return;
 	}
-	uint32 FramesInWave = (uint32)WaveProxyReader->GetNumFramesInWave();
+	uint32 FramesInWave = (uint32)WavePlayer->GetSourceNumFrames();
 	if (FrameIdx >= FramesInWave)
 	{
 		// passed the end
@@ -436,80 +422,44 @@ void FStreamingAudioRendererV2::SeekSourceAudioToFrame(uint32 FrameIdx)
 		{
 			return;
 		}
-
-		int32 NumFramesAhead = FrameIdx - SourceFrameIndex;
-		int32 NumSourceFramesAvailable = InterleavedCircularBuffer.Num() / NumDeinterleaveChannels;
-
-		if (NumFramesAhead <= NumSourceFramesAvailable)
-		{
-			InterleavedCircularBuffer.Pop(NumFramesAhead * NumDeinterleaveChannels);
-			check(GetSourceAudioFrameIndex() == FrameIdx);
-			return;
-		}
 	}
 	// do the actual seek
 
 
 	// do the actual seek
-	if (WaveProxyReader->SeekToFrame(FrameIdx))
+	if (WavePlayer->SeekToFrame(FrameIdx))
 	{
 		// at this point, we should be synced up!!
-		ensure(WaveProxyReader->GetFrameIndex() == FrameIdx);
-		InterleavedCircularBuffer.SetNum(0);
+		ensure(WavePlayer->GetCurrentPlaybackFrame() == FrameIdx);
 	}
 
-	check(GetSourceAudioFrameIndex() == FrameIdx);
-}
-
-void FStreamingAudioRendererV2::DecodeSourceAudio(Audio::TCircularAudioBuffer<float>& OutBuffer)
-{
-	check(!WaveProxyReader->HasFailed());
-	
-	const int32 NumSamplesToGenerate = DeinterleaveBlockSizeInFrames * WaveProxyReader->GetNumChannels();
-	check(NumSamplesToGenerate == DecodeBuffer.Num());
-	
-	if (OutBuffer.Remainder() == 0)
-	{
-		return;
-	}
-	else if (OutBuffer.Remainder() < (uint32)DecodeBuffer.Num())
-	{
-		Audio::FAlignedFloatBuffer RemainderBuffer;
-		RemainderBuffer.SetNum(OutBuffer.Remainder());
-		int32 NumPopped = WaveProxyReader->PopAudio(RemainderBuffer);
-		OutBuffer.Push(RemainderBuffer.GetData(), NumPopped);
-	}
-	else
-	{
-		int32 NumPopped = WaveProxyReader->PopAudio(DecodeBuffer);
-		OutBuffer.Push(DecodeBuffer.GetData(), NumPopped);
-	}
+	checkf(GetSourceAudioFrameIndex() == FrameIdx, TEXT("Failed to seek. SourceFrameIdx = %d, ExpectedFrameIdx = %d"), GetSourceAudioFrameIndex(), FrameIdx);
 }
 
 uint32 FStreamingAudioRendererV2::GetSourceAudioFrameIndex()
 {
-	check(!WaveProxyReader->HasFailed());
-	int32 ReaderFrameIndex = WaveProxyReader->GetFrameIndex();
-	int32 ReaderFramesInWave = WaveProxyReader->GetNumFramesInWave();
-	if (ReaderFramesInWave < ReaderFrameIndex)
+	check(WavePlayer->IsPlayerValid());
+	int32 PlaybackFrame = WavePlayer->GetCurrentPlaybackFrame();
+	int32 ReaderFramesInWave = WavePlayer->GetSourceNumFrames();
+	if (ReaderFramesInWave < PlaybackFrame)
 	{
-		ReaderFrameIndex = ReaderFramesInWave;
+		PlaybackFrame = ReaderFramesInWave;
 	}
 
-	int32 NumDecodedSamples = InterleavedCircularBuffer.Num();
-	int32 NumDecodedFrames = NumDecodedSamples / NumDeinterleaveChannels;
-
-	return ReaderFrameIndex - NumDecodedFrames;
+	return PlaybackFrame;
 
 }
 
-TUniquePtr<FSoundWaveProxyReader> FStreamingAudioRendererV2::CreateProxyReader(TSharedRef<FSoundWaveProxy> WaveProxy)
+TUniquePtr<FSoundWaveProxyPlayer> FStreamingAudioRendererV2::CreateWavePlayer(TSharedRef<FSoundWaveProxy> WaveProxy)
 {
-	FSoundWaveProxyReader::FSettings WaveReaderSettings;
-	WaveReaderSettings.MaxDecodeSizeInFrames = MaxDecodeSizeInFrames;
-	WaveReaderSettings.StartTimeInSeconds = 0.0f;
-	WaveReaderSettings.bMaintainAudioSync = true;
-	return FSoundWaveProxyReader::Create(WaveProxy, WaveReaderSettings);	
+	float SampleRate = WaveProxy->GetSampleRate();
+	int32 NumChannels = WaveProxy->GetNumChannels();
+	FSoundWaveProxyPlayer::FSettings PlayerSettings(SampleRate, NumChannels);
+	PlayerSettings.MaxDecodeSizeFrames = MaxDecodeSizeInFrames;
+	PlayerSettings.bMaintainAudioSync = true;
+	TUniquePtr<FSoundWaveProxyPlayer> NewWavePlayer = FSoundWaveProxyPlayer::Create(PlayerSettings);
+	NewWavePlayer->SetSoundWave(WaveProxy);
+	return MoveTemp(NewWavePlayer);
 }
 
 bool FStreamingAudioRendererV2::HasLoopSection() const
@@ -526,16 +476,16 @@ uint32 FStreamingAudioRendererV2::GetLoopStartFrame() const
 
 uint32 FStreamingAudioRendererV2::GetLoopEndFrame() const
 {
-	check(HasLoopSection())
+	check(HasLoopSection());
 	const FSoundWaveCuePoint& LoopRegion = SoundWaveProxy->GetLoopRegions()[0];
 	return LoopRegion.FramePosition + LoopRegion.FrameLength;
 }
 
-void FStreamingAudioRendererV2::GenerateSourceAudio(uint32 StartFrameIndex, Audio::FAlignedFloatBuffer& OutAudio, bool bHonorLoopRegion)
+void FStreamingAudioRendererV2::GenerateSourceAudio(uint32 StartFrameIndex, Audio::FAlignedFloatBuffer& OutAudio)
 {
 	check(OutAudio.Num() % NumDeinterleaveChannels == 0);
 
-	if (bHonorLoopRegion && HasLoopSection())
+	if (HasLoopSection())
 	{
 		uint32 NumFramesRequested = OutAudio.Num() / NumDeinterleaveChannels;
 		uint32 FirstFrameInLoop = GetLoopStartFrame();
@@ -545,7 +495,7 @@ void FStreamingAudioRendererV2::GenerateSourceAudio(uint32 StartFrameIndex, Audi
 
 		if (!ensure(StartFrameIndex <= LastFrameInLoop))
 		{
-			GenerateSourceAudioInternal(StartFrameIndex, OutAudio.GetData(), OutAudio.Num());
+			GenerateSourceAudioInternal(StartFrameIndex, OutAudio);
 			return;
 		}
 
@@ -573,7 +523,7 @@ void FStreamingAudioRendererV2::GenerateSourceAudio(uint32 StartFrameIndex, Audi
 					// Hmmmm.
 					if(ensure(GetSourceAudioFrameIndex() == StartFrameIndex))
 					{
-						GenerateSourceAudioInternal(StartFrameIndex, OutAudio.GetData(), NumSamples);
+						GenerateSourceAudioInternal(StartFrameIndex, OutAudio);
 					}
 					else
 					{
@@ -586,7 +536,7 @@ void FStreamingAudioRendererV2::GenerateSourceAudio(uint32 StartFrameIndex, Audi
 			}
 			else
 			{
-				GenerateSourceAudioInternal(StartFrameIndex, OutAudio.GetData(), NumSamples);
+				GenerateSourceAudioInternal(StartFrameIndex, OutAudio);
 			}
 
 			/// cache loop end frame...
@@ -602,56 +552,47 @@ void FStreamingAudioRendererV2::GenerateSourceAudio(uint32 StartFrameIndex, Audi
 			int32 SampleOffset = NumSamples;
 			NumFrames = NumFramesRequested - NumFrames;
 			NumSamples = NumFrames * NumDeinterleaveChannels;
-			GenerateSourceAudioInternal(FirstFrameInLoop, OutAudio.GetData() + SampleOffset, NumSamples);
+			GenerateSourceAudioInternal(FirstFrameInLoop, MakeArrayView(OutAudio.GetData() + SampleOffset, NumSamples));
 		}
 		else
 		{
-			GenerateSourceAudioInternal(StartFrameIndex, OutAudio.GetData(), OutAudio.Num());
+			GenerateSourceAudioInternal(StartFrameIndex, OutAudio);
 		}
 	}
 	else
 	{
-		GenerateSourceAudioInternal(StartFrameIndex, OutAudio.GetData(), OutAudio.Num());
+		GenerateSourceAudioInternal(StartFrameIndex, OutAudio);
 	}
 }
 
-void FStreamingAudioRendererV2::GenerateSourceAudioInternal(uint32 StartFrameIndex, float* OutAudioData, uint32 NumSamples)
+void FStreamingAudioRendererV2::GenerateSourceAudioInternal(uint32 StartFrameIndex, TArrayView<float> OutAudio)
 {
-	int32 NumSamplesRequested = NumSamples;
+	int32 NumSamplesRequested = OutAudio.Num();
 	int32 BufferIdx = 0;
 	
-	while (NumSamplesRequested > 0 && !WaveProxyReader->HasFailed())
+	while (NumSamplesRequested > 0 && WavePlayer->IsPlayerValid())
 	{
 		if (GetSourceAudioFrameIndex() != StartFrameIndex)
 		{
 			SeekSourceAudioToFrame(StartFrameIndex);
 		}
-		
-		check(InterleavedCircularBuffer.Num() % NumDeinterleaveChannels == 0);
-
-		int32 NumSamplesAvailable = InterleavedCircularBuffer.Num();
-
-		if (NumSamplesRequested > NumSamplesAvailable)
-		{
-			DecodeSourceAudio(InterleavedCircularBuffer);
-		}
 
-		int32 NumSamplesToRead = FMath::Min((int32)InterleavedCircularBuffer.Num(), NumSamplesRequested);
-
-		if (NumSamplesToRead == 0)
+		int32 NumFramesRead = WavePlayer->GenerateSourceAudio(OutAudio.Slice(BufferIdx, OutAudio.Num() - BufferIdx));
+		int32 NumSamplesRead = WavePlayer->GetOutputNumChannels() * NumFramesRead;
+		
+		if (NumSamplesRead == 0)
 		{
 			break;
 		}
 
-		InterleavedCircularBuffer.Peek(OutAudioData + BufferIdx, NumSamplesToRead);
-		NumSamplesRequested -= NumSamplesToRead;
-		BufferIdx += NumSamplesToRead;
-		StartFrameIndex += (uint32)(NumSamplesToRead / NumDeinterleaveChannels);
+		NumSamplesRequested -= NumSamplesRead;
+		BufferIdx += NumSamplesRead;
+		StartFrameIndex += (uint32)(NumSamplesRead / NumDeinterleaveChannels);
 	}
 
 	if (NumSamplesRequested > 0)
 	{
-		FMemory::Memzero(&OutAudioData[BufferIdx], sizeof(float) * NumSamplesRequested);
+		FMemory::Memzero(OutAudio.GetData() + BufferIdx, sizeof(float) * NumSamplesRequested);
 
 		UE_LOG(LogHarmonixStreamingAudioRendererV2, Verbose, TEXT("%s: Failed to generated samples: StartFrameIndex: %d, NumSamplesRequested: %d"), 
 			*SoundWaveProxy->GetFName().ToString(), StartFrameIndex, NumSamplesRequested);
diff --git a/Engine/Plugins/Runtime/Harmonix/Source/HarmonixDsp/Private/HarmonixDsp/AudioData/StreamingAudioRendererV2.h b/Engine/Plugins/Runtime/Harmonix/Source/HarmonixDsp/Private/HarmonixDsp/AudioData/StreamingAudioRendererV2.h
--- a/Engine/Plugins/Runtime/Harmonix/Source/HarmonixDsp/Private/HarmonixDsp/AudioData/StreamingAudioRendererV2.h
+++ b/Engine/Plugins/Runtime/Harmonix/Source/HarmonixDsp/Private/HarmonixDsp/AudioData/StreamingAudioRendererV2.h
@@ -14,7 +14,7 @@
 #define UE_API HARMONIXDSP_API
 
 class FSoundWaveProxy;
-class FSoundWaveProxyReader;
+class FSoundWaveProxyPlayer;
 class FFusionSampler;
 class IStretcherAndPitchShifter;
 
@@ -38,29 +38,27 @@ public:
 	UE_API virtual void SetFrame(uint32 InFrameNum) override;
 
 	UE_API virtual double Render(TAudioBuffer<float>& OutBuffer, double InPos, int32 InMaxFrame, double InResampleInc, double InPitchShift, double InSpeed,
-		bool MaintainPitchWhenSpeedChanges, bool InShouldHonorLoopPoints, const FGainMatrix& InGain) override;
+		bool MaintainPitchWhenSpeedChanges, const FGainMatrix& InGain) override;
 
-	UE_API double RenderInternal(TAudioBuffer<float>& OutBuffer, double Pos, int32 MaxFrame, double Inc, bool ShouldHonorLoopPoints, const FGainMatrix& Gain);
+	UE_API double RenderInternal(TAudioBuffer<float>& OutBuffer, double Pos, int32 MaxFrame, double Inc, const FGainMatrix& Gain);
 
 	UE_API virtual double RenderUnshifted(TAudioBuffer<float>& OutBuffer, double InPos, int32 InMaxFrame, double InInc,
-		bool InShouldHonorLoopPoints, const FGainMatrix& InGain) override;
+		const FGainMatrix& InGain) override;
 
-	UE_API void GenerateSourceAudio(uint32 StartFrame, Audio::FAlignedFloatBuffer& OutAudio, bool bHonorLoopRegion = false);
+	UE_API void GenerateSourceAudio(uint32 StartFrame, Audio::FAlignedFloatBuffer& OutAudio);
 
 private:
 
-	void RenderSimpleUnshifted(TAudioBuffer<float>& OutBuffer, const FLerpData* LerpArray, uint32 InNumFrames, const FGainMatrix& InGain, double InInc, bool ShouldHonorLoopPoints);
-	void RenderMultiChannelUnshifted(TAudioBuffer<float>& OutBuffer, const FLerpData* LerpArray, uint32 InNumFrames, const FGainMatrix& InGain, double InInc, bool ShouldHonorLoopPoints);
-	void RenderMultiChannelRoutedUnshifted(TAudioBuffer<float>& OutBuffer, const FLerpData* LerpArray, uint32 InNumFrames, const FGainMatrix& InGain, double InInc, bool ShouldHonorLoopPoints);
+	void RenderSimpleUnshifted(TAudioBuffer<float>& OutBuffer, const FLerpData* LerpArray, uint32 InNumFrames, const FGainMatrix& InGain, double InInc);
+	void RenderMultiChannelUnshifted(TAudioBuffer<float>& OutBuffer, const FLerpData* LerpArray, uint32 InNumFrames, const FGainMatrix& InGain, double InInc);
+	void RenderMultiChannelRoutedUnshifted(TAudioBuffer<float>& OutBuffer, const FLerpData* LerpArray, uint32 InNumFrames, const FGainMatrix& InGain, double InInc);
 
 	void SeekSourceAudioToFrame(uint32 FrameIdx);
-	void DecodeSourceAudio(Audio::TCircularAudioBuffer<float>& OutBuffer);
-
-	void GenerateSourceAudioInternal(uint32 StartFrameIndex, float* OutAudioData, uint32 NumSamples);
+	void GenerateSourceAudioInternal(uint32 StartFrameIndex, TArrayView<float> OutAudio);
 
 	uint32 GetSourceAudioFrameIndex();
 
-	static TUniquePtr<FSoundWaveProxyReader> CreateProxyReader(TSharedRef<FSoundWaveProxy> WaveProxy);
+	static TUniquePtr<FSoundWaveProxyPlayer> CreateWavePlayer(TSharedRef<FSoundWaveProxy> WaveProxy);
 
 	bool HasLoopSection() const;
 	uint32 GetLoopStartFrame() const;
@@ -75,17 +73,14 @@ private:
 	// it will get loaded on construction 
 	TSharedPtr<FSoundWaveProxy> SoundWaveProxy;
 
-	TUniquePtr<FSoundWaveProxyReader> WaveProxyReader;
+	TUniquePtr<FSoundWaveProxyPlayer> WavePlayer;
 
-	Audio::FAlignedFloatBuffer DecodeBuffer;
 	Audio::FAlignedFloatBuffer WorkBuffer;
 
-	Audio::TCircularAudioBuffer<float> InterleavedCircularBuffer;
 	int32 NumDeinterleaveChannels;
 	Audio::FAlignedFloatBuffer LastLoopFrameCache;
 	bool bLastLoopFrameCached = false;
 
-
 	// needs to be small enough to avoid audio artifacts and syncing issues
 	// but also large enough that we're decoding multiple times per block
 	static constexpr int32 DeinterleaveBlockSizeInFrames = 256;
diff --git a/Engine/Plugins/Runtime/Harmonix/Source/HarmonixDsp/Private/HarmonixDsp/AudioDataRenderer.cpp b/Engine/Plugins/Runtime/Harmonix/Source/HarmonixDsp/Private/HarmonixDsp/AudioDataRenderer.cpp
--- a/Engine/Plugins/Runtime/Harmonix/Source/HarmonixDsp/Private/HarmonixDsp/AudioDataRenderer.cpp
+++ b/Engine/Plugins/Runtime/Harmonix/Source/HarmonixDsp/Private/HarmonixDsp/AudioDataRenderer.cpp
@@ -3,7 +3,7 @@
 #include "HarmonixDsp/AudioDataRenderer.h"
 #include "Sound/SoundWave.h"
 
-double IAudioDataRenderer::CalculateLerpData(FLerpData* LerpArray, int32 InNumPoints, uint32 InNumOutFrames, double InPos, int32 InMaxFrame, bool InHonorLoopPoints, double InInc) const
+double IAudioDataRenderer::CalculateLerpData(FLerpData* LerpArray, int32 InNumPoints, uint32 InNumOutFrames, double InPos, int32 InMaxFrame, double InInc) const
 {
 	check(InNumPoints >= (int32)InNumOutFrames);
 	FMemory::Memset(LerpArray, 0, sizeof(FLerpData) * InNumPoints);
@@ -24,14 +24,12 @@ double IAudioDataRenderer::CalculateLerpData(FLerpData* LerpArray, int32 InNumPo
 	uint32 LoopStartFrameIndex = 0;
 	uint32 LoopEndFrameIndex = 0;
 	bool IsLooping = false;
-	bool HasTailSection = false;
 	if (SoundWaveProxy->GetLoopRegions().Num() > 0)
 	{
+		IsLooping = true;
 		const FSoundWaveCuePoint& LoopRegion = SoundWaveProxy->GetLoopRegions()[0];
 		LoopStartFrameIndex = LoopRegion.FramePosition;
 		LoopEndFrameIndex = LoopRegion.FramePosition + LoopRegion.FrameLength;
-		HasTailSection = LoopEndFrameIndex != SoundWaveProxy->GetNumFrames() - 1;
-		IsLooping = (!HasTailSection || (HasTailSection && InHonorLoopPoints));
 	}
 
 
diff --git a/Engine/Plugins/Runtime/Harmonix/Source/HarmonixDsp/Private/HarmonixDsp/FusionSampler/FusionVoice.cpp b/Engine/Plugins/Runtime/Harmonix/Source/HarmonixDsp/Private/HarmonixDsp/FusionSampler/FusionVoice.cpp
--- a/Engine/Plugins/Runtime/Harmonix/Source/HarmonixDsp/Private/HarmonixDsp/FusionSampler/FusionVoice.cpp
+++ b/Engine/Plugins/Runtime/Harmonix/Source/HarmonixDsp/Private/HarmonixDsp/FusionSampler/FusionVoice.cpp
@@ -728,14 +728,6 @@ uint32 FFusionVoice::Process(uint32 InSliceIndex, uint32 InSubsliceIndex, float*
 
 	OutputBuffer.SetAliasedChannelData(OutData, InNumChannels);
 
-	// we should loop (if possible) as long as we are not currently releasing the voice
-	// (This seem hinky to me. Often looping samples DO NOT have a nice exit 
-	// portion after the loop end. In these cases the sample will just stop short. - Buzz)
-	bool IsInReleaseStage = (AdsrVolume.GetStage() == EAdsrStage::Release);
-
-	// Forcing this to true so that Adsr release actually has samples to work on!
-	bool ShouldHonorLoopPoints = true; // !IsInReleaseStage; 
-
 	//--------------------------------
 	// update filter coefs
 	//--------------------------------
@@ -893,7 +885,6 @@ uint32 FFusionVoice::Process(uint32 InSliceIndex, uint32 InSubsliceIndex, float*
 					PitchThisSlice,
 					static_cast<double>(RenderSpeed),
 					MaintainPitchWhenSpeedChanges,
-					ShouldHonorLoopPoints,
 					WorkingGain);
 
 				ActiveRenderer->SetFrame((uint32)SamplePos);
diff --git a/Engine/Plugins/Runtime/Harmonix/Source/HarmonixDsp/Private/HarmonixDsp/Stretchers/SmbPitchShifter.cpp b/Engine/Plugins/Runtime/Harmonix/Source/HarmonixDsp/Private/HarmonixDsp/Stretchers/SmbPitchShifter.cpp
--- a/Engine/Plugins/Runtime/Harmonix/Source/HarmonixDsp/Private/HarmonixDsp/Stretchers/SmbPitchShifter.cpp
+++ b/Engine/Plugins/Runtime/Harmonix/Source/HarmonixDsp/Private/HarmonixDsp/Stretchers/SmbPitchShifter.cpp
@@ -225,9 +225,9 @@ void FSmbPitchShifter::Reset()
 	}
 }
 
-double FSmbPitchShifter::Render(TAudioBuffer<float>& OutputData, double InPos, int32 InMaxFrame, double ResampleInc, double InPitchShift, double InSpeed, bool MaintainPitchWhenSpeedChanges, bool ShouldHonorLoopPoints, const FGainMatrix& InGain)
+double FSmbPitchShifter::Render(TAudioBuffer<float>& OutputData, double InPos, int32 InMaxFrame, double ResampleInc, double InPitchShift, double InSpeed, bool MaintainPitchWhenSpeedChanges, const FGainMatrix& InGain)
 {
-	double NewPos = MyAudioRenderer->RenderUnshifted(OutputData, InPos, InMaxFrame, ResampleInc * InSpeed, ShouldHonorLoopPoints, InGain);
+	double NewPos = MyAudioRenderer->RenderUnshifted(OutputData, InPos, InMaxFrame, ResampleInc * InSpeed, InGain);
 	float ActualPitchShift = MaintainPitchWhenSpeedChanges ? (float)(InPitchShift / InSpeed) : InPitchShift;
 	int32 NumOutSamples = OutputData.GetNumValidFrames();
 	int32 NumOutChannels = OutputData.GetNumValidChannels();
diff --git a/Engine/Plugins/Runtime/Harmonix/Source/HarmonixDsp/Private/HarmonixDsp/Stretchers/SmbPitchShifter.h b/Engine/Plugins/Runtime/Harmonix/Source/HarmonixDsp/Private/HarmonixDsp/Stretchers/SmbPitchShifter.h
--- a/Engine/Plugins/Runtime/Harmonix/Source/HarmonixDsp/Private/HarmonixDsp/Stretchers/SmbPitchShifter.h
+++ b/Engine/Plugins/Runtime/Harmonix/Source/HarmonixDsp/Private/HarmonixDsp/Stretchers/SmbPitchShifter.h
@@ -106,7 +106,7 @@ public:
 		float* InLeftData, float* InRightData,
 		float* OutLeftData, float* OutRightData) override;
 
-	UE_API virtual double Render(TAudioBuffer<float>& OutputData, double Pos, int32 MaxFrame, double ResampleInc, double PitchShift, double Speed, bool MaintainPitchWhenSpeedChanges, bool ShouldHonorLoopPoints, const FGainMatrix& Gain) override;
+	UE_API virtual double Render(TAudioBuffer<float>& OutputData, double Pos, int32 MaxFrame, double ResampleInc, double PitchShift, double Speed, bool MaintainPitchWhenSpeedChanges, const FGainMatrix& Gain) override;
 
 	UE_API virtual bool TakeInput(TAudioBuffer<float>& InBuffer) override;
 	UE_API virtual bool InputSilence(int32 NumFrames) override;
diff --git a/Engine/Plugins/Runtime/Harmonix/Source/HarmonixDsp/Public/HarmonixDsp/AudioDataRenderer.h b/Engine/Plugins/Runtime/Harmonix/Source/HarmonixDsp/Public/HarmonixDsp/AudioDataRenderer.h
--- a/Engine/Plugins/Runtime/Harmonix/Source/HarmonixDsp/Public/HarmonixDsp/AudioDataRenderer.h
+++ b/Engine/Plugins/Runtime/Harmonix/Source/HarmonixDsp/Public/HarmonixDsp/AudioDataRenderer.h
@@ -40,10 +40,9 @@ public:
 	virtual void SetFrame(uint32 InFrameNum) = 0;
 
 	virtual double Render(TAudioBuffer<float>& OutBuffer, double InPos, int32 InMaxFrame, double InResampleInc, double InPitchShift, double InSpeed,
-		bool MaintainPitchWhenSpeedChanges, bool InShouldHonorLoopPoints, const FGainMatrix& InGain) = 0;
+		bool MaintainPitchWhenSpeedChanges, const FGainMatrix& InGain) = 0;
 
-	virtual double RenderUnshifted(TAudioBuffer<float>& OutBuffer, double InPos, int32 InMaxFrame, double InInc,
-		bool InShouldHonorLoopPoints, const FGainMatrix& InGain) = 0;
+	virtual double RenderUnshifted(TAudioBuffer<float>& OutBuffer, double InPos, int32 InMaxFrame, double InInc, const FGainMatrix& InGain) = 0;
 
 protected:
 
@@ -60,5 +59,5 @@ protected:
 		float WeightB;
 	};
 
-	double CalculateLerpData(FLerpData* LerpArray, int32 InNumPoints, uint32 InNumOutFrames, double InPos, int32 InMaxFrame, bool InHonorLoopPoints, double Inc) const;
+	double CalculateLerpData(FLerpData* LerpArray, int32 InNumPoints, uint32 InNumOutFrames, double InPos, int32 InMaxFrame, double Inc) const;
 };
diff --git a/Engine/Plugins/Runtime/Harmonix/Source/HarmonixDsp/Public/HarmonixDsp/StretcherAndPitchShifter.h b/Engine/Plugins/Runtime/Harmonix/Source/HarmonixDsp/Public/HarmonixDsp/StretcherAndPitchShifter.h
--- a/Engine/Plugins/Runtime/Harmonix/Source/HarmonixDsp/Public/HarmonixDsp/StretcherAndPitchShifter.h
+++ b/Engine/Plugins/Runtime/Harmonix/Source/HarmonixDsp/Public/HarmonixDsp/StretcherAndPitchShifter.h
@@ -61,7 +61,7 @@ public:
 		SetupAndResetImpl(ShifterSettings);
 	}
 
-	virtual double Render(TAudioBuffer<float>& OutputData, double Pos, int32 MaxFrame, double ResampleInc, double PitchShift, double Speed, bool MaintainPitchWhenSpeedChanges, bool ShouldHonorLoopPoints, const FGainMatrix& Gain) = 0;
+	virtual double Render(TAudioBuffer<float>& OutputData, double Pos, int32 MaxFrame, double ResampleInc, double PitchShift, double Speed, bool MaintainPitchWhenSpeedChanges, const FGainMatrix& Gain) = 0;
 
 	virtual void StereoPitchShift(float InPitchShift, int32 NumSampsToProcess, float* InLeftData, float* InRightData, float* OutLeftData, float* OutRightData) = 0;
 
