Added GetMsAtFractionalBarIncludingCountIn to ISongMapEvaluator so code can avoid dealing directly with ticks.

#tests Confirmed that GetBarIncludingCountInAtTick returns the same value as TickToMs(GetBarMap().FractionalBarIncludingCountInToTick)

diff --git a/Engine/Plugins/Runtime/Harmonix/Source/HarmonixMidi/Private/HarmonixMidi/SongMaps.cpp b/Engine/Plugins/Runtime/Harmonix/Source/HarmonixMidi/Private/HarmonixMidi/SongMaps.cpp
--- a/Engine/Plugins/Runtime/Harmonix/Source/HarmonixMidi/Private/HarmonixMidi/SongMaps.cpp
+++ b/Engine/Plugins/Runtime/Harmonix/Source/HarmonixMidi/Private/HarmonixMidi/SongMaps.cpp
@@ -1002,6 +1002,12 @@ float ISongMapEvaluator::GetBarIncludingCountInAtTick(float Tick) const
 	return GetBarMap().TickToFractionalBarIncludingCountIn(Tick);
 }
 
+float ISongMapEvaluator::GetMsAtFractionalBarIncludingCountIn(float BarIncludingCountIn) const
+{
+	const float Tick = GetBarMap().FractionalBarIncludingCountInToTick(BarIncludingCountIn);
+	return TickToMs(Tick);
+}
+
 float ISongMapEvaluator::GetMsPerBarAtMs(float Ms) const
 {
 	float Tick = MsToTick(Ms);
diff --git a/Engine/Plugins/Runtime/Harmonix/Source/HarmonixMidi/Public/HarmonixMidi/SongMaps.h b/Engine/Plugins/Runtime/Harmonix/Source/HarmonixMidi/Public/HarmonixMidi/SongMaps.h
--- a/Engine/Plugins/Runtime/Harmonix/Source/HarmonixMidi/Public/HarmonixMidi/SongMaps.h
+++ b/Engine/Plugins/Runtime/Harmonix/Source/HarmonixMidi/Public/HarmonixMidi/SongMaps.h
@@ -128,6 +128,7 @@ public:
 	UE_API int32 GetTimeSignatureChangePointTick(int32 PointIndex);
 	UE_API float GetBarIncludingCountInAtMs(float Ms) const;
 	UE_API float GetBarIncludingCountInAtTick(float Tick) const;
+	UE_API float GetMsAtFractionalBarIncludingCountIn(float BarIncludingCountIn) const; // BarIncludingCountIn is zero based
 	UE_API float GetMsPerBarAtMs(float Ms) const;
 	UE_API float GetMsPerBarAtTick(float Tick) const;
 	UE_API bool BarMapIsEmpty() const;
