Add missing QueryInterface and GetInterfaceId to TAudioRenderableProxy

CL 51640962 updated IMPL_AUDIOPROXY_CLASS to add GetInterfaceId() and
QueryInterface() for safe proxy downcasting via IProxyData::GetAs<>().
However, TAudioRenderableProxy in the Harmonix plugin has a manual
expansion of this macro (it cannot use the macro directly because it
delegates GetAudioProxyTypeName() to its DATA_STRUCT template parameter).
This manual expansion was not updated, so it was missing the two new
methods. Without QueryInterface(), any call to GetAs<TAudioRenderableProxy<...>>()
falls through to the base IProxyData::QueryInterface() which returns nullptr,
causing a checkf crash.

The fix adds GetInterfaceId() and QueryInterface() to the manual expansion,
mirroring the macro pattern but adapted for the template context.
#rnx
#Synced-CL 52358587

diff --git a/Engine/Plugins/Runtime/Harmonix/Source/Harmonix/Public/Harmonix/AudioRenderableProxy.h b/Engine/Plugins/Runtime/Harmonix/Source/Harmonix/Public/Harmonix/AudioRenderableProxy.h
--- a/Engine/Plugins/Runtime/Harmonix/Source/Harmonix/Public/Harmonix/AudioRenderableProxy.h
+++ b/Engine/Plugins/Runtime/Harmonix/Source/Harmonix/Public/Harmonix/AudioRenderableProxy.h
@@ -491,6 +491,15 @@ namespace Harmonix
 		{
 			return DATA_STRUCT::GetAudioProxyTypeName();
 		}
+		static FName GetInterfaceId() { return GetAudioProxyTypeName(); }
+		virtual void* QueryInterface(const FName InterfaceId) override
+		{
+			if (InterfaceId == GetAudioProxyTypeName())
+			{
+				return this;
+			}
+			return Audio::IProxyData::QueryInterface(InterfaceId);
+		}
 		static constexpr bool bWasAudioProxyClassImplemented = true;
 		friend class ::Audio::IProxyData;
 		friend class ::Audio::TProxyData<TAudioRenderableProxy<DATA_STRUCT, QUEUE_TYPE>>;
