[Backout] - CL52458046
#rnx
Original CL Desc
-----------------------------------------------------------------
Harmonix "Audio Data Renderer"

Use SoundWaveProxyPlayer instead of SoundWaveProxyReader

#rnx

#Synced-CL 52509571

diff --git a/Engine/Plugins/Runtime/Harmonix/Source/HarmonixDsp/Private/HarmonixDsp/AudioData/StreamingAudioRendererV2.cpp b/Engine/Plugins/Runtime/Harmonix/Source/HarmonixDsp/Private/HarmonixDsp/AudioData/StreamingAudioRendererV2.cpp
--- a/Engine/Plugins/Runtime/Harmonix/Source/HarmonixDsp/Private/HarmonixDsp/AudioData/StreamingAudioRendererV2.cpp
+++ b/Engine/Plugins/Runtime/Harmonix/Source/HarmonixDsp/Private/HarmonixDsp/AudioData/StreamingAudioRendererV2.cpp
@@ -1,15 +1,21 @@
 // Copyright Epic Games, Inc. All Rights Reserved.
 #include "HarmonixDsp/AudioData/StreamingAudioRendererV2.h"
-#include "HarmonixDsp/AudioUtility.h"
+
 #include "HarmonixDsp/FusionSampler/FusionSampler.h"
+
+#include "HarmonixDsp/AudioUtility.h"
 #include "HarmonixDsp/GainMatrix.h"
 #include "HarmonixDsp/PannerDetails.h"
-#include "Math/UnrealMathUtility.h"
+
 #include "Sound/SoundWave.h"
-#include "Sound/SoundWaveProxyPlayer.h"
+#include "Sound/SoundWaveProxyReader.h"
+
+#include "Math/UnrealMathUtility.h"
 
 DEFINE_LOG_CATEGORY(LogHarmonixStreamingAudioRendererV2);
 
+PRAGMA_DISABLE_DEPRECATION_WARNINGS
+
 FStreamingAudioRendererV2::FStreamingAudioRendererV2()
 {
 }
@@ -20,7 +26,8 @@ FStreamingAudioRendererV2::~FStreamingAudioRendererV2()
 
 void FStreamingAudioRendererV2::Reset()
 {
-	WavePlayer.Reset();
+	WaveProxyReader.Reset();
+	InterleavedCircularBuffer.SetNum(0);
 	Shifter = nullptr;
 }
 
@@ -28,18 +35,27 @@ void FStreamingAudioRendererV2::SetAudioData(TSharedRef<FSoundWaveProxy> InSound
 {
 	SoundWaveProxy = InSoundWave.ToSharedPtr();
 
-	WavePlayer.Reset();
-	WavePlayer = CreateWavePlayer(InSoundWave);
+	WaveProxyReader.Reset();
+	WaveProxyReader = CreateProxyReader(InSoundWave);
+
+	check(WaveProxyReader.IsValid());
 
-	check(WavePlayer->IsPlayerValid())
-	
 	int32 WaveProxyNumChannels = SoundWaveProxy->GetNumChannels();
-	
+
+	int32 DecodeBufferSize = WaveProxyNumChannels * DeinterleaveBlockSizeInFrames;
+	DecodeBuffer.Reset(DecodeBufferSize);
+	DecodeBuffer.AddUninitialized(DecodeBufferSize);
+
 	NumDeinterleaveChannels = WaveProxyNumChannels;
 
 	LastLoopFrameCache.SetNumUninitialized(NumDeinterleaveChannels);
 	bLastLoopFrameCached = false;
 
+	// interleaved circular buffer for streaming source audio data
+	// make it larger than then amount we decode each block to avoid underruns
+	int32 FrameCapacity = DecodeBufferSize * 2;
+	InterleavedCircularBuffer.SetCapacity(FrameCapacity);
+
 	TrackChannelInfo = InSettings.TrackChannelInfo;
 	Shifter = InSettings.Shifter;
 	MySampler = InSettings.Sampler;
@@ -62,13 +78,13 @@ void FStreamingAudioRendererV2::MigrateToSampler(const FFusionSampler* InSampler
 
 void FStreamingAudioRendererV2::SetFrame(uint32 InFrameNum)
 {
-	if (WavePlayer->IsPlayerValid())
+	if (!WaveProxyReader->HasFailed())
 	{
 		SeekSourceAudioToFrame(InFrameNum);
 	}
 }
 
-double FStreamingAudioRendererV2::Render(TAudioBuffer<float>& OutBuffer, double InPos, int32 InMaxFrame, double InResampleInc, double InPitchShift, double InSpeed, bool MaintainPitchWhenSpeedChanges, const FGainMatrix& InGain)
+double FStreamingAudioRendererV2::Render(TAudioBuffer<float>& OutBuffer, double InPos, int32 InMaxFrame, double InResampleInc, double InPitchShift, double InSpeed, bool MaintainPitchWhenSpeedChanges, bool InShouldHonorLoopPoints, const FGainMatrix& InGain)
 {
 	if (!SoundWaveProxy || SoundWaveProxy->GetNumFrames() == 0)
 	{
@@ -83,13 +99,13 @@ double FStreamingAudioRendererV2::Render(TAudioBuffer<float>& OutBuffer, double
 		{
 			Increment *= InSpeed;
 		}
-		return RenderInternal(OutBuffer, InPos, InMaxFrame, Increment, InGain);
+		return RenderInternal(OutBuffer, InPos, InMaxFrame, Increment, InShouldHonorLoopPoints, InGain);
 	}
 
-	return Shifter->Render(OutBuffer, InPos, InMaxFrame, InResampleInc, InPitchShift, InSpeed, MaintainPitchWhenSpeedChanges, InGain);
+	return Shifter->Render(OutBuffer, InPos, InMaxFrame, InResampleInc, InPitchShift, InSpeed, MaintainPitchWhenSpeedChanges, InShouldHonorLoopPoints, InGain);
 }
 
-double FStreamingAudioRendererV2::RenderInternal(TAudioBuffer<float>& OutBuffer, double InPos, int32 InMaxFrame, double InInc, const FGainMatrix& InGain)
+double FStreamingAudioRendererV2::RenderInternal(TAudioBuffer<float>& OutBuffer, double InPos, int32 InMaxFrame, double InInc, bool InShouldHonorLoopPoints, const FGainMatrix& InGain)
 {
 	check(OutBuffer.GetNumValidFrames() <= AudioRendering::kMicroSliceSize);
 
@@ -102,25 +118,25 @@ double FStreamingAudioRendererV2::RenderInternal(TAudioBuffer<float>& OutBuffer,
 	uint32 NumOutFrames = OutBuffer.GetNumValidFrames();
 
 	FLerpData LerpArray[AudioRendering::kMicroSliceSize];
-	InPos = CalculateLerpData(LerpArray, AudioRendering::kMicroSliceSize, NumOutFrames, InPos, InMaxFrame, InInc);
+	InPos = CalculateLerpData(LerpArray, AudioRendering::kMicroSliceSize, NumOutFrames, InPos, InMaxFrame, InShouldHonorLoopPoints, InInc);
 
 	if (TrackChannelInfo && TrackChannelInfo->Num() > 0)
 	{
-		RenderMultiChannelRoutedUnshifted(OutBuffer, LerpArray, NumOutFrames, InGain, InInc);
+		RenderMultiChannelRoutedUnshifted(OutBuffer, LerpArray, NumOutFrames, InGain, InInc, InShouldHonorLoopPoints);
 	}
 	else if (SoundWaveProxy->GetNumChannels() <= 2)
 	{
-		RenderSimpleUnshifted(OutBuffer, LerpArray, NumOutFrames, InGain, InInc);
+		RenderSimpleUnshifted(OutBuffer, LerpArray, NumOutFrames, InGain, InInc, InShouldHonorLoopPoints);
 	}
 	else
 	{
-		RenderMultiChannelUnshifted(OutBuffer, LerpArray, NumOutFrames, InGain, InInc);
+		RenderMultiChannelUnshifted(OutBuffer, LerpArray, NumOutFrames, InGain, InInc, InShouldHonorLoopPoints);
 	}
 
 	return InPos;
 }
 
-double FStreamingAudioRendererV2::RenderUnshifted(TAudioBuffer<float>& OutBuffer, double InPos, int32 InMaxFrame, double InInc, const FGainMatrix& InGain)
+double FStreamingAudioRendererV2::RenderUnshifted(TAudioBuffer<float>& OutBuffer, double InPos, int32 InMaxFrame, double InInc, bool InShouldHonorLoopPoints, const FGainMatrix& InGain)
 {
 	if (!SoundWaveProxy)
 	{
@@ -128,7 +144,7 @@ double FStreamingAudioRendererV2::RenderUnshifted(TAudioBuffer<float>& OutBuffer
 		return InPos;
 	}
 
-	return RenderInternal(OutBuffer, InPos, InMaxFrame, InInc, InGain);
+	return RenderInternal(OutBuffer, InPos, InMaxFrame, InInc, InShouldHonorLoopPoints, InGain);
 }
 
 int32 FStreamingAudioRendererV2::CalculateNumFramesNeeded(const FLerpData* LerpData, int32 NumPoints)
@@ -163,7 +179,7 @@ int32 FStreamingAudioRendererV2::CalculateNumFramesNeeded(const FLerpData* LerpD
 	return NumNeeded;
 }
 
-void FStreamingAudioRendererV2::RenderSimpleUnshifted(TAudioBuffer<float>& OutBuffer, const FLerpData* LerpArray, uint32 InNumFrames, const FGainMatrix& InGain, double InInc)
+void FStreamingAudioRendererV2::RenderSimpleUnshifted(TAudioBuffer<float>& OutBuffer, const FLerpData* LerpArray, uint32 InNumFrames, const FGainMatrix& InGain, double InInc, bool InShouldHonorLoopPoints)
 {
 	check(InNumFrames <= AudioRendering::kMicroSliceSize);
 
@@ -185,7 +201,7 @@ void FStreamingAudioRendererV2::RenderSimpleUnshifted(TAudioBuffer<float>& OutBu
 	}
 	
 	WorkBuffer.SetNum(NumSamplesNeeded);
-	GenerateSourceAudio(StartFrameIndex, WorkBuffer);
+	GenerateSourceAudio(StartFrameIndex, WorkBuffer, InShouldHonorLoopPoints);
 
 	int32 PosA;
 	int32 PosB;
@@ -219,7 +235,7 @@ void FStreamingAudioRendererV2::RenderSimpleUnshifted(TAudioBuffer<float>& OutBu
 	}
 }
 
-void FStreamingAudioRendererV2::RenderMultiChannelUnshifted(TAudioBuffer<float>& OutBuffer, const FLerpData* LerpArray, uint32 InNumFrames, const FGainMatrix& InGain, double InInc)
+void FStreamingAudioRendererV2::RenderMultiChannelUnshifted(TAudioBuffer<float>& OutBuffer, const FLerpData* LerpArray, uint32 InNumFrames, const FGainMatrix& InGain, double InInc, bool InShouldHonorLoopPoints)
 {
 	check(InNumFrames <= AudioRendering::kMicroSliceSize);
 
@@ -246,7 +262,7 @@ void FStreamingAudioRendererV2::RenderMultiChannelUnshifted(TAudioBuffer<float>&
 	}
 	
 	WorkBuffer.SetNum(NumSamplesNeeded);
-	GenerateSourceAudio(StartFrameIndex, WorkBuffer);
+	GenerateSourceAudio(StartFrameIndex, WorkBuffer, InShouldHonorLoopPoints);
 
 	int32 PosA;
 	int32 PosB;
@@ -292,7 +308,7 @@ void FStreamingAudioRendererV2::RenderMultiChannelUnshifted(TAudioBuffer<float>&
 	}
 }
 
-void FStreamingAudioRendererV2::RenderMultiChannelRoutedUnshifted(TAudioBuffer<float>& OutBuffer, const FLerpData* LerpArray, uint32 InNumFrames, const FGainMatrix& InGain, double InInc)
+void FStreamingAudioRendererV2::RenderMultiChannelRoutedUnshifted(TAudioBuffer<float>& OutBuffer, const FLerpData* LerpArray, uint32 InNumFrames, const FGainMatrix& InGain, double InInc, bool InShouldHonorLoopPoints)
 {
 	check(InNumFrames <= AudioRendering::kMicroSliceSize);
 
@@ -319,7 +335,7 @@ void FStreamingAudioRendererV2::RenderMultiChannelRoutedUnshifted(TAudioBuffer<f
 	}
 	
 	WorkBuffer.SetNum(NumSamplesNeeded);
-	GenerateSourceAudio(StartFrameIndex, WorkBuffer);
+	GenerateSourceAudio(StartFrameIndex, WorkBuffer, InShouldHonorLoopPoints);
 
 	int32 PosA;
 	int32 PosB;
@@ -399,7 +415,7 @@ void FStreamingAudioRendererV2::RenderMultiChannelRoutedUnshifted(TAudioBuffer<f
 
 void FStreamingAudioRendererV2::SeekSourceAudioToFrame(uint32 FrameIdx)
 {
-	check(WavePlayer->IsPlayerValid());
+	check(!WaveProxyReader->HasFailed());
 	
 	// no need to seek
 	uint32 SourceFrameIndex = GetSourceAudioFrameIndex();
@@ -408,7 +424,7 @@ void FStreamingAudioRendererV2::SeekSourceAudioToFrame(uint32 FrameIdx)
 	{
 		return;
 	}
-	uint32 FramesInWave = (uint32)WavePlayer->GetSourceNumFrames();
+	uint32 FramesInWave = (uint32)WaveProxyReader->GetNumFramesInWave();
 	if (FrameIdx >= FramesInWave)
 	{
 		// passed the end
@@ -422,44 +438,80 @@ void FStreamingAudioRendererV2::SeekSourceAudioToFrame(uint32 FrameIdx)
 		{
 			return;
 		}
+
+		int32 NumFramesAhead = FrameIdx - SourceFrameIndex;
+		int32 NumSourceFramesAvailable = InterleavedCircularBuffer.Num() / NumDeinterleaveChannels;
+
+		if (NumFramesAhead <= NumSourceFramesAvailable)
+		{
+			InterleavedCircularBuffer.Pop(NumFramesAhead * NumDeinterleaveChannels);
+			check(GetSourceAudioFrameIndex() == FrameIdx);
+			return;
+		}
 	}
 	// do the actual seek
 
 
 	// do the actual seek
-	if (WavePlayer->SeekToFrame(FrameIdx))
+	if (WaveProxyReader->SeekToFrame(FrameIdx))
 	{
 		// at this point, we should be synced up!!
-		ensure(WavePlayer->GetCurrentPlaybackFrame() == FrameIdx);
+		ensure(WaveProxyReader->GetFrameIndex() == FrameIdx);
+		InterleavedCircularBuffer.SetNum(0);
 	}
 
-	checkf(GetSourceAudioFrameIndex() == FrameIdx, TEXT("Failed to seek. SourceFrameIdx = %d, ExpectedFrameIdx = %d"), GetSourceAudioFrameIndex(), FrameIdx);
+	check(GetSourceAudioFrameIndex() == FrameIdx);
+}
+
+void FStreamingAudioRendererV2::DecodeSourceAudio(Audio::TCircularAudioBuffer<float>& OutBuffer)
+{
+	check(!WaveProxyReader->HasFailed());
+	
+	const int32 NumSamplesToGenerate = DeinterleaveBlockSizeInFrames * WaveProxyReader->GetNumChannels();
+	check(NumSamplesToGenerate == DecodeBuffer.Num());
+	
+	if (OutBuffer.Remainder() == 0)
+	{
+		return;
+	}
+	else if (OutBuffer.Remainder() < (uint32)DecodeBuffer.Num())
+	{
+		Audio::FAlignedFloatBuffer RemainderBuffer;
+		RemainderBuffer.SetNum(OutBuffer.Remainder());
+		int32 NumPopped = WaveProxyReader->PopAudio(RemainderBuffer);
+		OutBuffer.Push(RemainderBuffer.GetData(), NumPopped);
+	}
+	else
+	{
+		int32 NumPopped = WaveProxyReader->PopAudio(DecodeBuffer);
+		OutBuffer.Push(DecodeBuffer.GetData(), NumPopped);
+	}
 }
 
 uint32 FStreamingAudioRendererV2::GetSourceAudioFrameIndex()
 {
-	check(WavePlayer->IsPlayerValid());
-	int32 PlaybackFrame = WavePlayer->GetCurrentPlaybackFrame();
-	int32 ReaderFramesInWave = WavePlayer->GetSourceNumFrames();
-	if (ReaderFramesInWave < PlaybackFrame)
+	check(!WaveProxyReader->HasFailed());
+	int32 ReaderFrameIndex = WaveProxyReader->GetFrameIndex();
+	int32 ReaderFramesInWave = WaveProxyReader->GetNumFramesInWave();
+	if (ReaderFramesInWave < ReaderFrameIndex)
 	{
-		PlaybackFrame = ReaderFramesInWave;
+		ReaderFrameIndex = ReaderFramesInWave;
 	}
 
-	return PlaybackFrame;
+	int32 NumDecodedSamples = InterleavedCircularBuffer.Num();
+	int32 NumDecodedFrames = NumDecodedSamples / NumDeinterleaveChannels;
+
+	return ReaderFrameIndex - NumDecodedFrames;
 
 }
 
-TUniquePtr<FSoundWaveProxyPlayer> FStreamingAudioRendererV2::CreateWavePlayer(TSharedRef<FSoundWaveProxy> WaveProxy)
+TUniquePtr<FSoundWaveProxyReader> FStreamingAudioRendererV2::CreateProxyReader(TSharedRef<FSoundWaveProxy> WaveProxy)
 {
-	float SampleRate = WaveProxy->GetSampleRate();
-	int32 NumChannels = WaveProxy->GetNumChannels();
-	FSoundWaveProxyPlayer::FSettings PlayerSettings(SampleRate, NumChannels);
-	PlayerSettings.MaxDecodeSizeFrames = MaxDecodeSizeInFrames;
-	PlayerSettings.bMaintainAudioSync = true;
-	TUniquePtr<FSoundWaveProxyPlayer> NewWavePlayer = FSoundWaveProxyPlayer::Create(PlayerSettings);
-	NewWavePlayer->SetSoundWave(WaveProxy);
-	return MoveTemp(NewWavePlayer);
+	FSoundWaveProxyReader::FSettings WaveReaderSettings;
+	WaveReaderSettings.MaxDecodeSizeInFrames = MaxDecodeSizeInFrames;
+	WaveReaderSettings.StartTimeInSeconds = 0.0f;
+	WaveReaderSettings.bMaintainAudioSync = true;
+	return FSoundWaveProxyReader::Create(WaveProxy, WaveReaderSettings);	
 }
 
 bool FStreamingAudioRendererV2::HasLoopSection() const
@@ -476,16 +528,16 @@ uint32 FStreamingAudioRendererV2::GetLoopStartFrame() const
 
 uint32 FStreamingAudioRendererV2::GetLoopEndFrame() const
 {
-	check(HasLoopSection());
+	check(HasLoopSection())
 	const FSoundWaveCuePoint& LoopRegion = SoundWaveProxy->GetLoopRegions()[0];
 	return LoopRegion.FramePosition + LoopRegion.FrameLength;
 }
 
-void FStreamingAudioRendererV2::GenerateSourceAudio(uint32 StartFrameIndex, Audio::FAlignedFloatBuffer& OutAudio)
+void FStreamingAudioRendererV2::GenerateSourceAudio(uint32 StartFrameIndex, Audio::FAlignedFloatBuffer& OutAudio, bool bHonorLoopRegion)
 {
 	check(OutAudio.Num() % NumDeinterleaveChannels == 0);
 
-	if (HasLoopSection())
+	if (bHonorLoopRegion && HasLoopSection())
 	{
 		uint32 NumFramesRequested = OutAudio.Num() / NumDeinterleaveChannels;
 		uint32 FirstFrameInLoop = GetLoopStartFrame();
@@ -495,7 +547,7 @@ void FStreamingAudioRendererV2::GenerateSourceAudio(uint32 StartFrameIndex, Audi
 
 		if (!ensure(StartFrameIndex <= LastFrameInLoop))
 		{
-			GenerateSourceAudioInternal(StartFrameIndex, OutAudio);
+			GenerateSourceAudioInternal(StartFrameIndex, OutAudio.GetData(), OutAudio.Num());
 			return;
 		}
 
@@ -523,7 +575,7 @@ void FStreamingAudioRendererV2::GenerateSourceAudio(uint32 StartFrameIndex, Audi
 					// Hmmmm.
 					if(ensure(GetSourceAudioFrameIndex() == StartFrameIndex))
 					{
-						GenerateSourceAudioInternal(StartFrameIndex, OutAudio);
+						GenerateSourceAudioInternal(StartFrameIndex, OutAudio.GetData(), NumSamples);
 					}
 					else
 					{
@@ -536,7 +588,7 @@ void FStreamingAudioRendererV2::GenerateSourceAudio(uint32 StartFrameIndex, Audi
 			}
 			else
 			{
-				GenerateSourceAudioInternal(StartFrameIndex, OutAudio);
+				GenerateSourceAudioInternal(StartFrameIndex, OutAudio.GetData(), NumSamples);
 			}
 
 			/// cache loop end frame...
@@ -552,49 +604,60 @@ void FStreamingAudioRendererV2::GenerateSourceAudio(uint32 StartFrameIndex, Audi
 			int32 SampleOffset = NumSamples;
 			NumFrames = NumFramesRequested - NumFrames;
 			NumSamples = NumFrames * NumDeinterleaveChannels;
-			GenerateSourceAudioInternal(FirstFrameInLoop, MakeArrayView(OutAudio.GetData() + SampleOffset, NumSamples));
+			GenerateSourceAudioInternal(FirstFrameInLoop, OutAudio.GetData() + SampleOffset, NumSamples);
 		}
 		else
 		{
-			GenerateSourceAudioInternal(StartFrameIndex, OutAudio);
+			GenerateSourceAudioInternal(StartFrameIndex, OutAudio.GetData(), OutAudio.Num());
 		}
 	}
 	else
 	{
-		GenerateSourceAudioInternal(StartFrameIndex, OutAudio);
+		GenerateSourceAudioInternal(StartFrameIndex, OutAudio.GetData(), OutAudio.Num());
 	}
 }
 
-void FStreamingAudioRendererV2::GenerateSourceAudioInternal(uint32 StartFrameIndex, TArrayView<float> OutAudio)
+void FStreamingAudioRendererV2::GenerateSourceAudioInternal(uint32 StartFrameIndex, float* OutAudioData, uint32 NumSamples)
 {
-	int32 NumSamplesRequested = OutAudio.Num();
+	int32 NumSamplesRequested = NumSamples;
 	int32 BufferIdx = 0;
 	
-	while (NumSamplesRequested > 0 && WavePlayer->IsPlayerValid())
+	while (NumSamplesRequested > 0 && !WaveProxyReader->HasFailed())
 	{
 		if (GetSourceAudioFrameIndex() != StartFrameIndex)
 		{
 			SeekSourceAudioToFrame(StartFrameIndex);
 		}
-
-		int32 NumFramesRead = WavePlayer->GenerateSourceAudio(OutAudio.Slice(BufferIdx, OutAudio.Num() - BufferIdx));
-		int32 NumSamplesRead = WavePlayer->GetOutputNumChannels() * NumFramesRead;
 		
-		if (NumSamplesRead == 0)
+		check(InterleavedCircularBuffer.Num() % NumDeinterleaveChannels == 0);
+
+		int32 NumSamplesAvailable = InterleavedCircularBuffer.Num();
+
+		if (NumSamplesRequested > NumSamplesAvailable)
+		{
+			DecodeSourceAudio(InterleavedCircularBuffer);
+		}
+
+		int32 NumSamplesToRead = FMath::Min((int32)InterleavedCircularBuffer.Num(), NumSamplesRequested);
+
+		if (NumSamplesToRead == 0)
 		{
 			break;
 		}
 
-		NumSamplesRequested -= NumSamplesRead;
-		BufferIdx += NumSamplesRead;
-		StartFrameIndex += (uint32)(NumSamplesRead / NumDeinterleaveChannels);
+		InterleavedCircularBuffer.Peek(OutAudioData + BufferIdx, NumSamplesToRead);
+		NumSamplesRequested -= NumSamplesToRead;
+		BufferIdx += NumSamplesToRead;
+		StartFrameIndex += (uint32)(NumSamplesToRead / NumDeinterleaveChannels);
 	}
 
 	if (NumSamplesRequested > 0)
 	{
-		FMemory::Memzero(OutAudio.GetData() + BufferIdx, sizeof(float) * NumSamplesRequested);
+		FMemory::Memzero(&OutAudioData[BufferIdx], sizeof(float) * NumSamplesRequested);
 
 		UE_LOG(LogHarmonixStreamingAudioRendererV2, Verbose, TEXT("%s: Failed to generated samples: StartFrameIndex: %d, NumSamplesRequested: %d"), 
 			*SoundWaveProxy->GetFName().ToString(), StartFrameIndex, NumSamplesRequested);
 	}
 }
+
+PRAGMA_ENABLE_DEPRECATION_WARNINGS
diff --git a/Engine/Plugins/Runtime/Harmonix/Source/HarmonixDsp/Private/HarmonixDsp/AudioData/StreamingAudioRendererV2.h b/Engine/Plugins/Runtime/Harmonix/Source/HarmonixDsp/Private/HarmonixDsp/AudioData/StreamingAudioRendererV2.h
--- a/Engine/Plugins/Runtime/Harmonix/Source/HarmonixDsp/Private/HarmonixDsp/AudioData/StreamingAudioRendererV2.h
+++ b/Engine/Plugins/Runtime/Harmonix/Source/HarmonixDsp/Private/HarmonixDsp/AudioData/StreamingAudioRendererV2.h
@@ -14,12 +14,14 @@
 #define UE_API HARMONIXDSP_API
 
 class FSoundWaveProxy;
-class FSoundWaveProxyPlayer;
+class FSoundWaveProxyReader;
 class FFusionSampler;
 class IStretcherAndPitchShifter;
 
 DECLARE_LOG_CATEGORY_EXTERN(LogHarmonixStreamingAudioRendererV2, Log, All);
 
+PRAGMA_DISABLE_DEPRECATION_WARNINGS
+
 class FStreamingAudioRendererV2 : public IAudioDataRenderer
 {
 public:
@@ -38,27 +40,29 @@ public:
 	UE_API virtual void SetFrame(uint32 InFrameNum) override;
 
 	UE_API virtual double Render(TAudioBuffer<float>& OutBuffer, double InPos, int32 InMaxFrame, double InResampleInc, double InPitchShift, double InSpeed,
-		bool MaintainPitchWhenSpeedChanges, const FGainMatrix& InGain) override;
+		bool MaintainPitchWhenSpeedChanges, bool InShouldHonorLoopPoints, const FGainMatrix& InGain) override;
 
-	UE_API double RenderInternal(TAudioBuffer<float>& OutBuffer, double Pos, int32 MaxFrame, double Inc, const FGainMatrix& Gain);
+	UE_API double RenderInternal(TAudioBuffer<float>& OutBuffer, double Pos, int32 MaxFrame, double Inc, bool ShouldHonorLoopPoints, const FGainMatrix& Gain);
 
 	UE_API virtual double RenderUnshifted(TAudioBuffer<float>& OutBuffer, double InPos, int32 InMaxFrame, double InInc,
-		const FGainMatrix& InGain) override;
+		bool InShouldHonorLoopPoints, const FGainMatrix& InGain) override;
 
-	UE_API void GenerateSourceAudio(uint32 StartFrame, Audio::FAlignedFloatBuffer& OutAudio);
+	UE_API void GenerateSourceAudio(uint32 StartFrame, Audio::FAlignedFloatBuffer& OutAudio, bool bHonorLoopRegion = false);
 
 private:
 
-	void RenderSimpleUnshifted(TAudioBuffer<float>& OutBuffer, const FLerpData* LerpArray, uint32 InNumFrames, const FGainMatrix& InGain, double InInc);
-	void RenderMultiChannelUnshifted(TAudioBuffer<float>& OutBuffer, const FLerpData* LerpArray, uint32 InNumFrames, const FGainMatrix& InGain, double InInc);
-	void RenderMultiChannelRoutedUnshifted(TAudioBuffer<float>& OutBuffer, const FLerpData* LerpArray, uint32 InNumFrames, const FGainMatrix& InGain, double InInc);
+	void RenderSimpleUnshifted(TAudioBuffer<float>& OutBuffer, const FLerpData* LerpArray, uint32 InNumFrames, const FGainMatrix& InGain, double InInc, bool ShouldHonorLoopPoints);
+	void RenderMultiChannelUnshifted(TAudioBuffer<float>& OutBuffer, const FLerpData* LerpArray, uint32 InNumFrames, const FGainMatrix& InGain, double InInc, bool ShouldHonorLoopPoints);
+	void RenderMultiChannelRoutedUnshifted(TAudioBuffer<float>& OutBuffer, const FLerpData* LerpArray, uint32 InNumFrames, const FGainMatrix& InGain, double InInc, bool ShouldHonorLoopPoints);
 
 	void SeekSourceAudioToFrame(uint32 FrameIdx);
-	void GenerateSourceAudioInternal(uint32 StartFrameIndex, TArrayView<float> OutAudio);
+	void DecodeSourceAudio(Audio::TCircularAudioBuffer<float>& OutBuffer);
+
+	void GenerateSourceAudioInternal(uint32 StartFrameIndex, float* OutAudioData, uint32 NumSamples);
 
 	uint32 GetSourceAudioFrameIndex();
 
-	static TUniquePtr<FSoundWaveProxyPlayer> CreateWavePlayer(TSharedRef<FSoundWaveProxy> WaveProxy);
+	static TUniquePtr<FSoundWaveProxyReader> CreateProxyReader(TSharedRef<FSoundWaveProxy> WaveProxy);
 
 	bool HasLoopSection() const;
 	uint32 GetLoopStartFrame() const;
@@ -73,14 +77,17 @@ private:
 	// it will get loaded on construction 
 	TSharedPtr<FSoundWaveProxy> SoundWaveProxy;
 
-	TUniquePtr<FSoundWaveProxyPlayer> WavePlayer;
+	TUniquePtr<FSoundWaveProxyReader> WaveProxyReader;
 
+	Audio::FAlignedFloatBuffer DecodeBuffer;
 	Audio::FAlignedFloatBuffer WorkBuffer;
 
+	Audio::TCircularAudioBuffer<float> InterleavedCircularBuffer;
 	int32 NumDeinterleaveChannels;
 	Audio::FAlignedFloatBuffer LastLoopFrameCache;
 	bool bLastLoopFrameCached = false;
 
+
 	// needs to be small enough to avoid audio artifacts and syncing issues
 	// but also large enough that we're decoding multiple times per block
 	static constexpr int32 DeinterleaveBlockSizeInFrames = 256;
@@ -89,4 +96,6 @@ private:
 	int32 CalculateNumFramesNeeded(const FLerpData* LerpData, int32 NumPoints);
 };
 
+PRAGMA_ENABLE_DEPRECATION_WARNINGS
+
 #undef UE_API
diff --git a/Engine/Plugins/Runtime/Harmonix/Source/HarmonixDsp/Private/HarmonixDsp/AudioDataRenderer.cpp b/Engine/Plugins/Runtime/Harmonix/Source/HarmonixDsp/Private/HarmonixDsp/AudioDataRenderer.cpp
--- a/Engine/Plugins/Runtime/Harmonix/Source/HarmonixDsp/Private/HarmonixDsp/AudioDataRenderer.cpp
+++ b/Engine/Plugins/Runtime/Harmonix/Source/HarmonixDsp/Private/HarmonixDsp/AudioDataRenderer.cpp
@@ -3,7 +3,7 @@
 #include "HarmonixDsp/AudioDataRenderer.h"
 #include "Sound/SoundWave.h"
 
-double IAudioDataRenderer::CalculateLerpData(FLerpData* LerpArray, int32 InNumPoints, uint32 InNumOutFrames, double InPos, int32 InMaxFrame, double InInc) const
+double IAudioDataRenderer::CalculateLerpData(FLerpData* LerpArray, int32 InNumPoints, uint32 InNumOutFrames, double InPos, int32 InMaxFrame, bool InHonorLoopPoints, double InInc) const
 {
 	check(InNumPoints >= (int32)InNumOutFrames);
 	FMemory::Memset(LerpArray, 0, sizeof(FLerpData) * InNumPoints);
@@ -24,12 +24,14 @@ double IAudioDataRenderer::CalculateLerpData(FLerpData* LerpArray, int32 InNumPo
 	uint32 LoopStartFrameIndex = 0;
 	uint32 LoopEndFrameIndex = 0;
 	bool IsLooping = false;
+	bool HasTailSection = false;
 	if (SoundWaveProxy->GetLoopRegions().Num() > 0)
 	{
-		IsLooping = true;
 		const FSoundWaveCuePoint& LoopRegion = SoundWaveProxy->GetLoopRegions()[0];
 		LoopStartFrameIndex = LoopRegion.FramePosition;
 		LoopEndFrameIndex = LoopRegion.FramePosition + LoopRegion.FrameLength;
+		HasTailSection = LoopEndFrameIndex != SoundWaveProxy->GetNumFrames() - 1;
+		IsLooping = (!HasTailSection || (HasTailSection && InHonorLoopPoints));
 	}
 
 
diff --git a/Engine/Plugins/Runtime/Harmonix/Source/HarmonixDsp/Private/HarmonixDsp/FusionSampler/FusionVoice.cpp b/Engine/Plugins/Runtime/Harmonix/Source/HarmonixDsp/Private/HarmonixDsp/FusionSampler/FusionVoice.cpp
--- a/Engine/Plugins/Runtime/Harmonix/Source/HarmonixDsp/Private/HarmonixDsp/FusionSampler/FusionVoice.cpp
+++ b/Engine/Plugins/Runtime/Harmonix/Source/HarmonixDsp/Private/HarmonixDsp/FusionSampler/FusionVoice.cpp
@@ -728,6 +728,14 @@ uint32 FFusionVoice::Process(uint32 InSliceIndex, uint32 InSubsliceIndex, float*
 
 	OutputBuffer.SetAliasedChannelData(OutData, InNumChannels);
 
+	// we should loop (if possible) as long as we are not currently releasing the voice
+	// (This seem hinky to me. Often looping samples DO NOT have a nice exit 
+	// portion after the loop end. In these cases the sample will just stop short. - Buzz)
+	bool IsInReleaseStage = (AdsrVolume.GetStage() == EAdsrStage::Release);
+
+	// Forcing this to true so that Adsr release actually has samples to work on!
+	bool ShouldHonorLoopPoints = true; // !IsInReleaseStage; 
+
 	//--------------------------------
 	// update filter coefs
 	//--------------------------------
@@ -885,6 +893,7 @@ uint32 FFusionVoice::Process(uint32 InSliceIndex, uint32 InSubsliceIndex, float*
 					PitchThisSlice,
 					static_cast<double>(RenderSpeed),
 					MaintainPitchWhenSpeedChanges,
+					ShouldHonorLoopPoints,
 					WorkingGain);
 
 				ActiveRenderer->SetFrame((uint32)SamplePos);
diff --git a/Engine/Plugins/Runtime/Harmonix/Source/HarmonixDsp/Private/HarmonixDsp/Stretchers/SmbPitchShifter.cpp b/Engine/Plugins/Runtime/Harmonix/Source/HarmonixDsp/Private/HarmonixDsp/Stretchers/SmbPitchShifter.cpp
--- a/Engine/Plugins/Runtime/Harmonix/Source/HarmonixDsp/Private/HarmonixDsp/Stretchers/SmbPitchShifter.cpp
+++ b/Engine/Plugins/Runtime/Harmonix/Source/HarmonixDsp/Private/HarmonixDsp/Stretchers/SmbPitchShifter.cpp
@@ -225,9 +225,9 @@ void FSmbPitchShifter::Reset()
 	}
 }
 
-double FSmbPitchShifter::Render(TAudioBuffer<float>& OutputData, double InPos, int32 InMaxFrame, double ResampleInc, double InPitchShift, double InSpeed, bool MaintainPitchWhenSpeedChanges, const FGainMatrix& InGain)
+double FSmbPitchShifter::Render(TAudioBuffer<float>& OutputData, double InPos, int32 InMaxFrame, double ResampleInc, double InPitchShift, double InSpeed, bool MaintainPitchWhenSpeedChanges, bool ShouldHonorLoopPoints, const FGainMatrix& InGain)
 {
-	double NewPos = MyAudioRenderer->RenderUnshifted(OutputData, InPos, InMaxFrame, ResampleInc * InSpeed, InGain);
+	double NewPos = MyAudioRenderer->RenderUnshifted(OutputData, InPos, InMaxFrame, ResampleInc * InSpeed, ShouldHonorLoopPoints, InGain);
 	float ActualPitchShift = MaintainPitchWhenSpeedChanges ? (float)(InPitchShift / InSpeed) : InPitchShift;
 	int32 NumOutSamples = OutputData.GetNumValidFrames();
 	int32 NumOutChannels = OutputData.GetNumValidChannels();
diff --git a/Engine/Plugins/Runtime/Harmonix/Source/HarmonixDsp/Private/HarmonixDsp/Stretchers/SmbPitchShifter.h b/Engine/Plugins/Runtime/Harmonix/Source/HarmonixDsp/Private/HarmonixDsp/Stretchers/SmbPitchShifter.h
--- a/Engine/Plugins/Runtime/Harmonix/Source/HarmonixDsp/Private/HarmonixDsp/Stretchers/SmbPitchShifter.h
+++ b/Engine/Plugins/Runtime/Harmonix/Source/HarmonixDsp/Private/HarmonixDsp/Stretchers/SmbPitchShifter.h
@@ -106,7 +106,7 @@ public:
 		float* InLeftData, float* InRightData,
 		float* OutLeftData, float* OutRightData) override;
 
-	UE_API virtual double Render(TAudioBuffer<float>& OutputData, double Pos, int32 MaxFrame, double ResampleInc, double PitchShift, double Speed, bool MaintainPitchWhenSpeedChanges, const FGainMatrix& Gain) override;
+	UE_API virtual double Render(TAudioBuffer<float>& OutputData, double Pos, int32 MaxFrame, double ResampleInc, double PitchShift, double Speed, bool MaintainPitchWhenSpeedChanges, bool ShouldHonorLoopPoints, const FGainMatrix& Gain) override;
 
 	UE_API virtual bool TakeInput(TAudioBuffer<float>& InBuffer) override;
 	UE_API virtual bool InputSilence(int32 NumFrames) override;
diff --git a/Engine/Plugins/Runtime/Harmonix/Source/HarmonixDsp/Public/HarmonixDsp/AudioDataRenderer.h b/Engine/Plugins/Runtime/Harmonix/Source/HarmonixDsp/Public/HarmonixDsp/AudioDataRenderer.h
--- a/Engine/Plugins/Runtime/Harmonix/Source/HarmonixDsp/Public/HarmonixDsp/AudioDataRenderer.h
+++ b/Engine/Plugins/Runtime/Harmonix/Source/HarmonixDsp/Public/HarmonixDsp/AudioDataRenderer.h
@@ -40,9 +40,10 @@ public:
 	virtual void SetFrame(uint32 InFrameNum) = 0;
 
 	virtual double Render(TAudioBuffer<float>& OutBuffer, double InPos, int32 InMaxFrame, double InResampleInc, double InPitchShift, double InSpeed,
-		bool MaintainPitchWhenSpeedChanges, const FGainMatrix& InGain) = 0;
+		bool MaintainPitchWhenSpeedChanges, bool InShouldHonorLoopPoints, const FGainMatrix& InGain) = 0;
 
-	virtual double RenderUnshifted(TAudioBuffer<float>& OutBuffer, double InPos, int32 InMaxFrame, double InInc, const FGainMatrix& InGain) = 0;
+	virtual double RenderUnshifted(TAudioBuffer<float>& OutBuffer, double InPos, int32 InMaxFrame, double InInc,
+		bool InShouldHonorLoopPoints, const FGainMatrix& InGain) = 0;
 
 protected:
 
@@ -59,5 +60,5 @@ protected:
 		float WeightB;
 	};
 
-	double CalculateLerpData(FLerpData* LerpArray, int32 InNumPoints, uint32 InNumOutFrames, double InPos, int32 InMaxFrame, double Inc) const;
+	double CalculateLerpData(FLerpData* LerpArray, int32 InNumPoints, uint32 InNumOutFrames, double InPos, int32 InMaxFrame, bool InHonorLoopPoints, double Inc) const;
 };
diff --git a/Engine/Plugins/Runtime/Harmonix/Source/HarmonixDsp/Public/HarmonixDsp/StretcherAndPitchShifter.h b/Engine/Plugins/Runtime/Harmonix/Source/HarmonixDsp/Public/HarmonixDsp/StretcherAndPitchShifter.h
--- a/Engine/Plugins/Runtime/Harmonix/Source/HarmonixDsp/Public/HarmonixDsp/StretcherAndPitchShifter.h
+++ b/Engine/Plugins/Runtime/Harmonix/Source/HarmonixDsp/Public/HarmonixDsp/StretcherAndPitchShifter.h
@@ -61,7 +61,7 @@ public:
 		SetupAndResetImpl(ShifterSettings);
 	}
 
-	virtual double Render(TAudioBuffer<float>& OutputData, double Pos, int32 MaxFrame, double ResampleInc, double PitchShift, double Speed, bool MaintainPitchWhenSpeedChanges, const FGainMatrix& Gain) = 0;
+	virtual double Render(TAudioBuffer<float>& OutputData, double Pos, int32 MaxFrame, double ResampleInc, double PitchShift, double Speed, bool MaintainPitchWhenSpeedChanges, bool ShouldHonorLoopPoints, const FGainMatrix& Gain) = 0;
 
 	virtual void StereoPitchShift(float InPitchShift, int32 NumSampsToProcess, float* InLeftData, float* InRightData, float* OutLeftData, float* OutRightData) = 0;
 
