Reset additional playback state and zero audio output buffers when the simple sampler node is reset
Prevent stale audio data from persisting across capture resets

#rnx
#Synced-CL 52524518

diff --git a/Engine/Plugins/Runtime/Harmonix/Source/HarmonixMetasound/Private/HarmonixMetasound/Nodes/SimpleSamplerNode.cpp b/Engine/Plugins/Runtime/Harmonix/Source/HarmonixMetasound/Private/HarmonixMetasound/Nodes/SimpleSamplerNode.cpp
--- a/Engine/Plugins/Runtime/Harmonix/Source/HarmonixMetasound/Private/HarmonixMetasound/Nodes/SimpleSamplerNode.cpp
+++ b/Engine/Plugins/Runtime/Harmonix/Source/HarmonixMetasound/Private/HarmonixMetasound/Nodes/SimpleSamplerNode.cpp
@@ -247,8 +247,13 @@ namespace HarmonixMetasound::SimpleSamplerNode
 			CaptureBuffer.SetNumValidFrames(0);
 			CaptureBuffer.SetNumValidChannels(Inputs.AudioInputs.Num());
 			bCurrentlyPlaying = bCurrentlyCapturing = bFullyCaptured = false;
+			bPlayingReverse = bHavePlayedThisCapture = false;
 			CaptureLengthFrames = PlayLengthFrames = NextPlayFrame = 0;
 			*Outputs.InputAudioInUse = false;
+			for (FAudioBufferWriteRef& Out : Outputs.AudioOutputs)
+			{
+				Out->Zero();
+			}
 		}
 
 		enum class EInputOrSilence
