Unshelved from pending changelist '52755636':

Refactor MusicClock from an abstract native class into a UObject that reads from a new IMusicSource interface, removing the ECalibratedMusicTimebase parameter in favor of one clock per time stream
Introduce MusicSource abstraction with Metasound, Runtime, and Offset variants, and extend MidiClockUpdateSubsystem to drive both sources and clocks each frame
Add Blueprint-exposed beat, bar, section, and play-state events to UMusicClock and include unit tests for the new architecture

#rnx
#Synced-CL 52524518

#Synced-CL 52404697

diff --git a/Engine/Plugins/Runtime/Harmonix/Source/HarmonixMetasound/Private/HarmonixMetasound/MusicClock/MetasoundMusicClock.cpp b/Engine/Plugins/Runtime/Harmonix/Source/HarmonixMetasound/Private/HarmonixMetasound/MusicClock/MetasoundMusicClock.cpp
deleted file mode 100644
diff --git a/Engine/Plugins/Runtime/Harmonix/Source/HarmonixMetasound/Private/HarmonixMetasound/MusicClock/MetasoundMusicClock.h b/Engine/Plugins/Runtime/Harmonix/Source/HarmonixMetasound/Private/HarmonixMetasound/MusicClock/MetasoundMusicClock.h
deleted file mode 100644
diff --git a/Engine/Plugins/Runtime/Harmonix/Source/HarmonixMetasound/Private/HarmonixMetasound/MusicClock/MusicClock.cpp b/Engine/Plugins/Runtime/Harmonix/Source/HarmonixMetasound/Private/HarmonixMetasound/MusicClock/MusicClock.cpp
--- a/Engine/Plugins/Runtime/Harmonix/Source/HarmonixMetasound/Private/HarmonixMetasound/MusicClock/MusicClock.cpp
+++ b/Engine/Plugins/Runtime/Harmonix/Source/HarmonixMetasound/Private/HarmonixMetasound/MusicClock/MusicClock.cpp
@@ -1,234 +1,458 @@
 // Copyright Epic Games, Inc. All Rights Reserved.
 
 #include "HarmonixMetasound/MusicClock/MusicClock.h"
+#include "HarmonixMetasound/Subsystems/MidiClockUpdateSubsystem.h"
 
-#include "Engine/World.h"
-#include "HarmonixMetasound/Components/MusicClockTypes.h"
+DEFINE_LOG_CATEGORY_STATIC(LogUMusicClock, Log, All)
 
-void Harmonix::FMusicClock::ResetDefaultSongMaps(float DefaultTempo, int32 DefaultTimeSigNum, int32 DefaultTimeSigDenom)
+UMusicClock::UMusicClock()
 {
-	DefaultMaps.EmptyAllMaps();
-	DefaultMaps.Init(Harmonix::Midi::Constants::GTicksPerQuarterNoteInt);
-	DefaultMaps.GetTempoMap().AddTempoInfoPoint(Harmonix::Midi::Constants::BPMToMidiTempo(DefaultTempo), 0);
-	DefaultMaps.GetBarMap().AddTimeSignatureAtBarIncludingCountIn(0, DefaultTimeSigNum, DefaultTimeSigDenom);
 }
 
-void Harmonix::FMusicClock::UpdateBeatsAndBarsPerSecond()
+void UMusicClock::BeginDestroy()
 {
-	// Midi quarter note to beat conversion.
-	const float BPM = GetBeatsPerMinute();
-	CurrentBeatsPerSecond = (BPM / 60.0) * CurrentClockAdvanceRate;
-	CurrentBarsPerSecond = (TimeSignatureNum > 0) ? (CurrentBeatsPerSecond / TimeSignatureNum) : 0;
+	UMidiClockUpdateSubsystem::StopTrackingMusicClock(this);
+	Super::BeginDestroy();
 }
 
-Harmonix::FMusicClock::FMusicClock(UObject* WorldContextObject, const FMusicClockSettingsBase& Settings)
+void UMusicClock::SetSource(const TScriptInterface<IMusicSource>& InSource)
 {
-	ResetDefaultSongMaps(Settings.DefaultTempo, Settings.DefaultTimeSigNumerator, Settings.DefaultTimeSigDenominator);
-	ContextObj = TWeakObjectPtr<UObject>(WorldContextObject);
+	SourceObject = InSource.GetObject();
+	SourceInterfaceCache = InSource ? Cast<IMusicSource>(InSource.GetObject()) : nullptr;
+
+	// Reset state when rebinding
+	CachedClockState = EMusicClockState::Stopped;
+	CurrentPos.Reset();
+	PrevPos.Reset();
+	DeltaBarF = 0.f;
+	DeltaBeatF = 0.f;
+	LastBroadcastBar = -1;
+	LastBroadcastBeat = -1;
+	LastBroadcastSection = FSongSection();
+	bLoopedThisFrame = false;
+	bSeekedThisFrame = false;
 }
 
-Harmonix::FMusicClock::~FMusicClock()
+TScriptInterface<IMusicSource> UMusicClock::GetSource() const
 {
+	if (UObject* Obj = SourceObject.Get())
+	{
+		return TScriptInterface<IMusicSource>(Obj);
+	}
+	return {};
+}
+
+IMusicSource* UMusicClock::GetValidSource() const
+{
+	return SourceObject.IsValid() ? SourceInterfaceCache : nullptr;
 }
 
-bool Harmonix::FMusicClock::CalculateSongPosWithOffset(float MsOffset, ECalibratedMusicTimebase Timebase, FMidiSongPos& OutResult) const
+void UMusicClock::UpdateForGameFrame()
 {
-	const ISongMapEvaluator* SongMapEvaluator = GetCurrentSongMapEvaluator();
-	if (!SongMapEvaluator)
+	if (!IsInGameThread())
 	{
-		return false;
+		return;
 	}
 
-	const FMidiSongPos& SongPos = GetCurrentSongPos(Timebase);
-	if (!SongPos.IsValid())
+	if (GFrameCounter == LastUpdateFrame)
 	{
-		return false;
+		return;
 	}
 
-	OutResult.SetByTime((SongPos.SecondsIncludingCountIn * 1000.0f) + MsOffset, *SongMapEvaluator);
-	return true;
-}
+	IMusicSource* SourcePtr = GetValidSource();
+	if (!SourcePtr)
+	{
+		// Source was destroyed — transition to Stopped if we weren't already
+		if (SourceInterfaceCache != nullptr)
+		{
+			SourceInterfaceCache = nullptr;
+		}
+		if (CachedClockState != EMusicClockState::Stopped)
+		{
+			CachedClockState = EMusicClockState::Stopped;
+			CurrentPos.Reset();
+			PrevPos.Reset();
+			DeltaBarF = 0.f;
+			DeltaBeatF = 0.f;
+			Tempo = 0.f;
+			CurrentBeatsPerSecondCached = 0.f;
+			CurrentBarsPerSecondCached = 0.f;
+			ClockAdvanceRate = 1.f;
+			bLoopedThisFrame = false;
+			bSeekedThisFrame = false;
+			PlayStateEvent.Broadcast(EMusicClockState::Stopped);
+		}
+		LastUpdateFrame = GFrameCounter;
+		return;
+	}
 
-double Harmonix::FMusicClock::GetWallClockTime() const
-{
-	if (TStrongObjectPtr<UObject> Obj = ContextObj.Pin())
+	EMusicClockState PrevClockState = CachedClockState;
+	Harmonix::ESourceState SourceState = SourcePtr->GetSourceState();
+
+	// Derive the new clock state from the source
+	EMusicClockState NewClockState;
+	switch (SourceState)
 	{
-		if (Obj->GetWorld())
+	case Harmonix::ESourceState::Running:   NewClockState = EMusicClockState::Running; break;
+	case Harmonix::ESourceState::Paused:    NewClockState = EMusicClockState::Paused; break;
+	case Harmonix::ESourceState::Preparing: NewClockState = EMusicClockState::Stopped; break;
+	case Harmonix::ESourceState::Stopped:   NewClockState = EMusicClockState::Stopped; break;
+	default:                                NewClockState = EMusicClockState::Stopped; break;
+	}
+	CachedClockState = NewClockState;
+
+	if (SourceState == Harmonix::ESourceState::Running)
+	{
+		PrevPos = CurrentPos;
+		CurrentPos = SourcePtr->GetCurrentSongPos();
+
+		DeltaBarF = CurrentPos.BarsIncludingCountIn - PrevPos.BarsIncludingCountIn;
+		DeltaBeatF = CurrentPos.BeatsIncludingCountIn - PrevPos.BeatsIncludingCountIn;
+
+		bLoopedThisFrame = SourcePtr->LoopedThisFrame();
+		bSeekedThisFrame = SourcePtr->SeekedThisFrame();
+
+		UpdatePlaybackRate(CurrentPos, SourcePtr->GetCurrentClockAdvanceRate());
+		BroadcastSongPosChanges();
+	}
+	else if (SourceState == Harmonix::ESourceState::Stopped)
+	{
+		// When stopped, reset position to 0
+		if (PrevClockState != EMusicClockState::Stopped)
+		{
+			CurrentPos.Reset();
+			PrevPos.Reset();
+			DeltaBarF = 0.f;
+			DeltaBeatF = 0.f;
+			Tempo = 0.f;
+			CurrentBeatsPerSecondCached = 0.f;
+			CurrentBarsPerSecondCached = 0.f;
+			ClockAdvanceRate = 1.f;
+		}
+		bLoopedThisFrame = false;
+		bSeekedThisFrame = false;
+	}
+	else if (SourceState == Harmonix::ESourceState::Preparing)
+	{
+		// Preparing = intent to run, but no data yet. Reset position to zero
+		// so we don't report stale data, but don't broadcast events.
+		if (PrevClockState != EMusicClockState::Stopped)
 		{
-			return Obj->GetWorld()->GetRealTimeSeconds();
+			CurrentPos.Reset();
+			PrevPos.Reset();
+			DeltaBarF = 0.f;
+			DeltaBeatF = 0.f;
 		}
+		bLoopedThisFrame = false;
+		bSeekedThisFrame = false;
+	}
+	else // Paused
+	{
+		// When paused, keep current position frozen. Zero out deltas.
+		DeltaBarF = 0.f;
+		DeltaBeatF = 0.f;
+		bLoopedThisFrame = false;
+		bSeekedThisFrame = false;
 	}
 
-	if (GWorld)
+	// Broadcast state changes
+	if (NewClockState != PrevClockState)
 	{
-		return GWorld->GetRealTimeSeconds();
+		PlayStateEvent.Broadcast(NewClockState);
 	}
-	
-	return 0.0;
+
+	LastUpdateFrame = GFrameCounter;
+}
+
+EMusicClockState UMusicClock::GetState() const
+{
+	return CachedClockState;
 }
 
+const FMidiSongPos& UMusicClock::GetCurrentSongPos() const
+{
+	return CurrentPos;
+}
 
-const FMidiSongPos& Harmonix::FMusicClock::GetCurrentSongPos(ECalibratedMusicTimebase Timebase) const
+const FMidiSongPos& UMusicClock::GetPreviousSongPos() const
 {
-	return CurrentSongPos[GetArrayIndexForTimebase(Timebase)];
+	return PrevPos;
 }
 
-const FMidiSongPos& Harmonix::FMusicClock::GetPreviousSongPos(ECalibratedMusicTimebase Timebase) const
+float UMusicClock::GetSecondsIncludingCountIn() const
 {
-	return PrevSongPos[GetArrayIndexForTimebase(Timebase)];
+	return CurrentPos.SecondsIncludingCountIn;
 }
 
-float Harmonix::FMusicClock::GetDeltaBarF(ECalibratedMusicTimebase Timebase) const
+float UMusicClock::GetSecondsFromBarOne() const
 {
-	return DeltaBarF[GetArrayIndexForTimebase(Timebase)];
+	return CurrentPos.SecondsFromBarOne;
 }
 
-float Harmonix::FMusicClock::GetDeltaBeatF(ECalibratedMusicTimebase Timebase) const
+float UMusicClock::GetBarsIncludingCountIn() const
 {
-	return DeltaBeatF[GetArrayIndexForTimebase(Timebase)];
+	return CurrentPos.BarsIncludingCountIn;
 }
 
-void Harmonix::FMusicClock::Start()
+float UMusicClock::GetBeatsIncludingCountIn() const
 {
-	OnStart();
+	return CurrentPos.BeatsIncludingCountIn;
+}
 
-	for (ECalibratedMusicTimebase Timebase : TEnumRange<ECalibratedMusicTimebase>())
-	{
-		CurrentSongPos[(int32)Timebase].Reset();
-	}
+FMusicTimestamp UMusicClock::GetCurrentTimestamp() const
+{
+	return CurrentPos.Timestamp;
+}
 
-	MusicClockState = EMusicClockState::Running;
+float UMusicClock::GetCurrentTempo() const
+{
+	return Tempo;
 }
 
-void Harmonix::FMusicClock::Continue()
+float UMusicClock::GetCurrentBeatsPerMinute() const
 {
-	if (MusicClockState != EMusicClockState::Paused)
-	{
-		return;
-	}
-	
-	OnContinue();
-	MusicClockState = EMusicClockState::Running;
+	ensure(TimeSignatureDenom != 0 || Tempo == 0);
+	return Tempo * TimeSignatureDenom / 4.0f;
 }
 
-void Harmonix::FMusicClock::Pause()
+float UMusicClock::GetCurrentBeatsPerSecond() const
 {
-	if (MusicClockState != EMusicClockState::Running)
-	{
-		return;
-	}
+	return CurrentBeatsPerSecondCached;
+}
+
+float UMusicClock::GetCurrentBarsPerSecond() const
+{
+	return CurrentBarsPerSecondCached;
+}
 
-	OnPause();
-	MusicClockState = EMusicClockState::Paused;
+float UMusicClock::GetCurrentClockAdvanceRate() const
+{
+	return ClockAdvanceRate;
 }
 
-void Harmonix::FMusicClock::Stop()
+void UMusicClock::GetCurrentTimeSignature(int& OutNumerator, int& OutDenominator) const
 {
-	OnStop();
-	MusicClockState = EMusicClockState::Stopped;
+	OutNumerator = TimeSignatureNum;
+	OutDenominator = TimeSignatureDenom;
+}
 
-	for (ECalibratedMusicTimebase Timebase : TEnumRange<ECalibratedMusicTimebase>())
-	{
-		CurrentSongPos[(int32)Timebase].Reset();
-	}
+FString UMusicClock::GetCurrentSectionName() const
+{
+	return CurrentPos.CurrentSongSection.Name;
 }
 
-void Harmonix::FMusicClock::UpdateMusicPlaybackRate(float InTempo, float InClockAdvanceRate, int32 InTimeSigNum, int32 InTimeSigDenom)
+int32 UMusicClock::GetCurrentSectionIndex() const
 {
-	bool bHasAnyTimingChanges = false;
-	if (!FMath::IsNearlyEqual(Tempo, InTempo))
+	const ISongMapEvaluator* Maps = GetSongMaps();
+	if (Maps)
 	{
-		Tempo = InTempo;
-		bHasAnyTimingChanges = true;
+		return Maps->GetSectionIndexAtTick(CurrentPos.CurrentSongSection.StartTick);
 	}
+	return -1;
+}
 
-	if (!FMath::IsNearlyEqual(CurrentClockAdvanceRate, InClockAdvanceRate))
-	{
-		CurrentClockAdvanceRate = InClockAdvanceRate;
-		bHasAnyTimingChanges = true;
-	}
-	
-	if (TimeSignatureNum != InTimeSigNum || TimeSignatureDenom != InTimeSigDenom)
+float UMusicClock::GetCurrentSectionStartMs() const
+{
+	const ISongMapEvaluator* Maps = GetSongMaps();
+	if (Maps)
 	{
-		TimeSignatureNum = InTimeSigNum;
-		TimeSignatureDenom = InTimeSigDenom;
-		bHasAnyTimingChanges = true;
+		return Maps->TickToMs(CurrentPos.CurrentSongSection.StartTick);
 	}
+	return 0.f;
+}
 
-	// only update bar/beat durations if there have been any other timing changes
-	if (bHasAnyTimingChanges)
+float UMusicClock::GetCurrentSectionLengthMs() const
+{
+	const ISongMapEvaluator* Maps = GetSongMaps();
+	if (Maps)
 	{
-		UpdateBeatsAndBarsPerSecond();
+		return Maps->TickToMs(CurrentPos.CurrentSongSection.LengthTicks);
 	}
+	return 0.f;
 }
 
-float Harmonix::FMusicClock::GetBeatsPerMinute() const
+float UMusicClock::GetDistanceFromCurrentBeat() const
 {
-	// if we have a tempo set, we should have a time signature.
-	ensure(TimeSignatureDenom != 0 || Tempo == 0);
-    // Tempo is MidiTempo, which is in quarter notes per minute.
-	return Tempo * TimeSignatureDenom / 4.0f; 
+	return CurrentPos.BeatsIncludingCountIn - FMath::FloorToFloat(CurrentPos.BeatsIncludingCountIn);
 }
 
-void Harmonix::FMusicClock::UpdateClockForGameFrame()
+float UMusicClock::GetDistanceToNextBeat() const
 {
-	//	Not for use outside the game thread.
-	if (ensureMsgf(
-		IsInGameThread(),
-		TEXT("%hs called from non-game thread.  This is not supported!"), __FUNCTION__) == false)
+	return 1.0f - GetDistanceFromCurrentBeat();
+}
+
+float UMusicClock::GetDistanceToClosestBeat() const
+{
+	float Dist = GetDistanceFromCurrentBeat();
+	return FMath::Min(Dist, 1.0f - Dist);
+}
+
+float UMusicClock::GetDistanceFromCurrentBar() const
+{
+	return CurrentPos.BarsIncludingCountIn - FMath::FloorToFloat(CurrentPos.BarsIncludingCountIn);
+}
+
+float UMusicClock::GetDistanceToNextBar() const
+{
+	return 1.0f - GetDistanceFromCurrentBar();
+}
+
+float UMusicClock::GetDistanceToClosestBar() const
+{
+	float Dist = GetDistanceFromCurrentBar();
+	return FMath::Min(Dist, 1.0f - Dist);
+}
+
+float UMusicClock::GetDeltaBar() const
+{
+	return DeltaBarF;
+}
+
+float UMusicClock::GetDeltaBeat() const
+{
+	return DeltaBeatF;
+}
+
+bool UMusicClock::LoopedThisFrame() const
+{
+	return bLoopedThisFrame;
+}
+
+bool UMusicClock::SeekedThisFrame() const
+{
+	return bSeekedThisFrame;
+}
+
+bool UMusicClock::IsLooping() const
+{
+	IMusicSource* SourcePtr = GetValidSource();
+	return SourcePtr ? SourcePtr->IsLooping() : false;
+}
+
+float UMusicClock::GetLoopStartMs() const
+{
+	IMusicSource* SourcePtr = GetValidSource();
+	return SourcePtr ? SourcePtr->GetLoopStartMs() : 0.f;
+}
+
+float UMusicClock::GetLoopLengthMs() const
+{
+	IMusicSource* SourcePtr = GetValidSource();
+	return SourcePtr ? SourcePtr->GetLoopLengthMs() : 0.f;
+}
+
+const ISongMapEvaluator* UMusicClock::GetSongMaps() const
+{
+	IMusicSource* SourcePtr = GetValidSource();
+	if (SourcePtr)
 	{
-		return;
+		return SourcePtr->GetCurrentSongMapEvaluator();
 	}
+	return nullptr;
+}
 
-	if (GFrameCounter == LastUpdateFrame)
+float UMusicClock::GetSongLengthMs() const
+{
+	const ISongMapEvaluator* Maps = GetSongMaps();
+	if (Maps)
 	{
-		return;
+		return Maps->GetSongLengthMs();
 	}
+	return 0.f;
+}
 
-	//	Run the actual clock update.
-	if (GetState() == EMusicClockState::Running)
+float UMusicClock::GetSongRemainingMs() const
+{
+	float LengthMs = GetSongLengthMs();
+	if (LengthMs <= 0.f)
 	{
-		for (ECalibratedMusicTimebase Timebase : TEnumRange<ECalibratedMusicTimebase>())
-		{
-			PrevSongPos[(int32)Timebase] = CurrentSongPos[(int32)Timebase];
-		}
-
-		if (RefreshCurrentSongPos())
-		{
-			for (ECalibratedMusicTimebase Timebase : TEnumRange<ECalibratedMusicTimebase>())
-			{
-				DeltaBarF[(int32)Timebase] = CurrentSongPos[(int32)Timebase].BarsIncludingCountIn - PrevSongPos[(int32)Timebase].BarsIncludingCountIn;
-				DeltaBeatF[(int32)Timebase] = CurrentSongPos[(int32)Timebase].BeatsIncludingCountIn - PrevSongPos[(int32)Timebase].BeatsIncludingCountIn;
-			}
-
-			LastUpdateFrame = GFrameCounter;
-		}
+		return 0.f;
 	}
+	return FMath::Max(0.f, LengthMs - (CurrentPos.SecondsIncludingCountIn * 1000.f));
 }
 
-int32 Harmonix::FMusicClock::GetArrayIndexForTimebase(ECalibratedMusicTimebase Timebase) const
+TOptional<FVector> UMusicClock::TryGetAudioSourceLocation() const
 {
-	int32 TimebaseIndex = (int32)Timebase;
-	if (!ensure(TimebaseIndex >= 0 && TimebaseIndex < (int32)ECalibratedMusicTimebase::Count))
+	IMusicSource* SourcePtr = GetValidSource();
+	if (SourcePtr)
 	{
-		// minor robustness, return something even if there's an illegal value.
-		return (int32)ECalibratedMusicTimebase::VideoRenderTime;
+		return SourcePtr->TryGetAudioSourceLocation();
 	}
-
-	return TimebaseIndex;
+	return {};
 }
 
-void Harmonix::FMusicClock::SetCurrentSongPos(ECalibratedMusicTimebase Timebase, const FMidiSongPos& InSongPos)
+#if !UE_BUILD_SHIPPING
+FString UMusicClock::GetDisplayName() const
 {
-	CurrentSongPos[GetArrayIndexForTimebase(Timebase)] = InSongPos;
+	IMusicSource* SourcePtr = GetValidSource();
+	if (SourcePtr)
+	{
+		return SourcePtr->GetDisplayName();
+	}
+	return TEXT("UMusicClock (no source)");
 }
+#endif
 
-void Harmonix::FMusicClock::SetCurrentSongPosByTime(ECalibratedMusicTimebase Timebase, float InElapsedMs, const ISongMapEvaluator& Map)
+void UMusicClock::UpdatePlaybackRate(const FMidiSongPos& SongPos, float InClockAdvanceRate)
 {
-	CurrentSongPos[GetArrayIndexForTimebase(Timebase)].SetByTime(InElapsedMs, Map);
+	bool bChanged = false;
+
+	if (!FMath::IsNearlyEqual(Tempo, SongPos.Tempo))
+	{
+		Tempo = SongPos.Tempo;
+		bChanged = true;
+	}
+	if (!FMath::IsNearlyEqual(ClockAdvanceRate, InClockAdvanceRate))
+	{
+		ClockAdvanceRate = InClockAdvanceRate;
+		bChanged = true;
+	}
+	if (TimeSignatureNum != SongPos.TimeSigNumerator || TimeSignatureDenom != SongPos.TimeSigDenominator)
+	{
+		TimeSignatureNum = SongPos.TimeSigNumerator;
+		TimeSignatureDenom = SongPos.TimeSigDenominator;
+		bChanged = true;
+	}
+
+	if (bChanged)
+	{
+		float BPM = GetCurrentBeatsPerMinute();
+		CurrentBeatsPerSecondCached = (BPM / 60.0f) * ClockAdvanceRate;
+		CurrentBarsPerSecondCached = (TimeSignatureNum > 0) ? (CurrentBeatsPerSecondCached / TimeSignatureNum) : 0.f;
+	}
 }
 
-void Harmonix::FMusicClock::SetCurrentSongPosByTick(ECalibratedMusicTimebase Timebase, float InTick, const ISongMapEvaluator& Map)
+void UMusicClock::BroadcastSongPosChanges()
 {
-	CurrentSongPos[GetArrayIndexForTimebase(Timebase)].SetByTick(InTick, Map);
+	if (!CurrentPos.IsValid())
+	{
+		return;
+	}
+
+	const int32 CurrBar = FMath::FloorToInt32(CurrentPos.BarsIncludingCountIn);
+	if (LastBroadcastBar != CurrBar)
+	{
+		BarEvent.Broadcast(CurrentPos.Timestamp.Bar);
+		LastBroadcastBar = CurrBar;
+	}
+
+	const int32 CurrBeat = FMath::FloorToInt32(CurrentPos.BeatsIncludingCountIn);
+	if (LastBroadcastBeat != CurrBeat)
+	{
+		BeatEvent.Broadcast(CurrBeat, FMath::FloorToInt32(CurrentPos.Timestamp.Beat));
+		LastBroadcastBeat = CurrBeat;
+	}
+
+	const FSongSection& SongSection = CurrentPos.CurrentSongSection;
+	if (LastBroadcastSection.StartTick != SongSection.StartTick || LastBroadcastSection.LengthTicks != SongSection.LengthTicks)
+	{
+		const ISongMapEvaluator* Maps = GetSongMaps();
+		if (Maps)
+		{
+			float SectionStartMs = Maps->TickToMs(SongSection.StartTick);
+			float SectionLengthMs = Maps->TickToMs(SongSection.LengthTicks);
+			SectionEvent.Broadcast(SongSection.Name, SectionStartMs, SectionLengthMs);
+		}
+		LastBroadcastSection = FSongSection(SongSection.Name, SongSection.StartTick, SongSection.LengthTicks);
+	}
 }
diff --git a/Engine/Plugins/Runtime/Harmonix/Source/HarmonixMetasound/Private/HarmonixMetasound/MusicClock/RuntimeMusicClock.cpp b/Engine/Plugins/Runtime/Harmonix/Source/HarmonixMetasound/Private/HarmonixMetasound/MusicClock/RuntimeMusicClock.cpp
deleted file mode 100644
diff --git a/Engine/Plugins/Runtime/Harmonix/Source/HarmonixMetasound/Private/HarmonixMetasound/MusicClock/RuntimeMusicClock.h b/Engine/Plugins/Runtime/Harmonix/Source/HarmonixMetasound/Private/HarmonixMetasound/MusicClock/RuntimeMusicClock.h
deleted file mode 100644
diff --git a/Engine/Plugins/Runtime/Harmonix/Source/HarmonixMetasound/Private/HarmonixMetasound/MusicClock/WorldTimeSourceController.cpp b/Engine/Plugins/Runtime/Harmonix/Source/HarmonixMetasound/Private/HarmonixMetasound/MusicClock/WorldTimeSourceController.cpp
--- a/Engine/Plugins/Runtime/Harmonix/Source/HarmonixMetasound/Private/HarmonixMetasound/MusicClock/WorldTimeSourceController.cpp
+++ b/Engine/Plugins/Runtime/Harmonix/Source/HarmonixMetasound/Private/HarmonixMetasound/MusicClock/WorldTimeSourceController.cpp
@@ -19,11 +19,6 @@ TOptional<FVector> Harmonix::FWorldTimeSourceController::TryGetAudioSourceLocati
 	return {};
 }
 
-ECalibratedMusicTimebase Harmonix::FWorldTimeSourceController::GetTimebase() const
-{
-	return ECalibratedMusicTimebase::VideoRenderTime;
-}
-
 double Harmonix::FWorldTimeSourceController::GetDeltaTime() const
 {
 	if (TStrongObjectPtr<UObject> Obj = ContextObj.Pin())
diff --git a/Engine/Plugins/Runtime/Harmonix/Source/HarmonixMetasound/Private/HarmonixMetasound/MusicSource/MetasoundMusicSource.cpp b/Engine/Plugins/Runtime/Harmonix/Source/HarmonixMetasound/Private/HarmonixMetasound/MusicSource/MetasoundMusicSource.cpp
new file mode 100644
--- /dev/null
+++ b/Engine/Plugins/Runtime/Harmonix/Source/HarmonixMetasound/Private/HarmonixMetasound/MusicSource/MetasoundMusicSource.cpp
@@ -0,0 +1,641 @@
+// Copyright Epic Games, Inc. All Rights Reserved.
+
+#include "HarmonixMetasound/MusicSource/MetasoundMusicSource.h"
+#include "Components/AudioComponent.h"
+#include "Harmonix.h"
+#include "HarmonixMetasound/Analysis/MidiSongPosVertexAnalyzer.h"
+#include "HarmonixMetasound/DataTypes/MidiClock.h"
+#include "MetasoundGenerator.h"
+#include "MetasoundGeneratorHandle.h"
+#include "MetasoundSource.h"
+
+DEFINE_LOG_CATEGORY_STATIC(LogMetasoundMusicSource, Log, All)
+
+namespace MetasoundMusicSourceSmoothing
+{
+	float kP = 0.18f;
+	float LagSeconds = 0.030f;
+	float MaxErrorSecondsBeforeJump = 0.060f;
+	double MinSyncSpeed = 0.98;
+	double MaxSyncSpeed = 1.02;
+}
+
+void UMetasoundMusicSource::BeginDestroy()
+{
+	UMidiClockUpdateSubsystem::StopTrackingMusicSource(this);
+	Disconnect();
+	Super::BeginDestroy();
+}
+
+// ---- Connection ----
+
+bool UMetasoundMusicSource::ConnectToAudioComponent(UAudioComponent* InAudioComponent, FName OutputPinName)
+{
+	AudioComponentToWatch = InAudioComponent;
+	MetasoundOutputName = OutputPinName;
+
+	OnAttachedDelegate = UMetasoundGeneratorHandle::FOnAttached::FDelegate::CreateLambda([](){});
+	OnDetachedDelegate = UMetasoundGeneratorHandle::FOnDetached::FDelegate::CreateLambda([](){});
+
+	return AttemptToConnect(InAudioComponent);
+}
+
+void UMetasoundMusicSource::Disconnect()
+{
+	check(IsInGameThread());
+	SetClockHistory(nullptr);
+	DetachAllCallbacks();
+	AudioComponentToWatch.Reset();
+	CurrentGeneratorHandle.Reset();
+	CachedSourceState = Harmonix::ESourceState::Stopped;
+	CurrentPos.Reset();
+	PrevPos.Reset();
+}
+
+// ---- Transport (IMusicSource) ----
+// Transport is derived from the Metasound's actual state, not from explicit calls.
+// These are no-ops — control the Metasound through its AudioComponent directly.
+
+void UMetasoundMusicSource::Start() {}
+void UMetasoundMusicSource::Pause() {}
+void UMetasoundMusicSource::Continue() {}
+void UMetasoundMusicSource::Stop() {}
+
+Harmonix::ESourceState UMetasoundMusicSource::GetSourceState() const
+{
+	return CachedSourceState;
+}
+
+Harmonix::ESourceEvent UMetasoundMusicSource::GetLatestSourceEvent() const
+{
+	if (CachedSourceState != Harmonix::ESourceState::Running)
+	{
+		return Harmonix::ESourceEvent::None;
+	}
+	if (bLoopDetected)
+	{
+		return Harmonix::ESourceEvent::Loop;
+	}
+	if (bSeekDetected)
+	{
+		return Harmonix::ESourceEvent::Seek;
+	}
+	return Harmonix::ESourceEvent::Advance;
+}
+
+const FMidiSongPos& UMetasoundMusicSource::GetCurrentSongPos() const
+{
+	return CurrentPos;
+}
+
+const FMidiSongPos& UMetasoundMusicSource::GetPreviousSongPos() const
+{
+	return PrevPos;
+}
+
+const ISongMapEvaluator* UMetasoundMusicSource::GetCurrentSongMapEvaluator() const
+{
+	check(IsInGameThread());
+	if (ClockHistory && CurrentMapChain && CurrentMapChain->SongMaps)
+	{
+		return CurrentMapChain->SongMaps.Get();
+	}
+	return nullptr;
+}
+
+float UMetasoundMusicSource::GetCurrentClockAdvanceRate() const
+{
+	return CurrentAdvanceRate;
+}
+
+bool UMetasoundMusicSource::LoopedThisFrame() const
+{
+	return bLoopDetected;
+}
+
+bool UMetasoundMusicSource::SeekedThisFrame() const
+{
+	return bSeekDetected;
+}
+
+bool UMetasoundMusicSource::IsLooping() const
+{
+	return CurrentMapChain && CurrentMapChain->LoopLengthTicks > 0;
+}
+
+float UMetasoundMusicSource::GetLoopStartMs() const
+{
+	if (CurrentMapChain && CurrentMapChain->SongMaps && CurrentMapChain->LoopLengthTicks > 0)
+	{
+		return CurrentMapChain->SongMaps->TickToMs(CurrentMapChain->FirstTickInLoop);
+	}
+	return 0.f;
+}
+
+float UMetasoundMusicSource::GetLoopLengthMs() const
+{
+	if (CurrentMapChain && CurrentMapChain->SongMaps && CurrentMapChain->LoopLengthTicks > 0)
+	{
+		float StartMs = CurrentMapChain->SongMaps->TickToMs(CurrentMapChain->FirstTickInLoop);
+		float EndMs = CurrentMapChain->SongMaps->TickToMs(CurrentMapChain->FirstTickInLoop + CurrentMapChain->LoopLengthTicks);
+		return EndMs - StartMs;
+	}
+	return 0.f;
+}
+
+TOptional<FVector> UMetasoundMusicSource::TryGetAudioSourceLocation() const
+{
+	if (UAudioComponent* AudioComponent = AudioComponentToWatch.Get())
+	{
+		return AudioComponent->GetComponentLocation();
+	}
+	return {};
+}
+
+#if !UE_BUILD_SHIPPING
+FString UMetasoundMusicSource::GetDisplayName() const
+{
+#if WITH_EDITORONLY_DATA
+	if (UAudioComponent* Component = AudioComponentToWatch.Get())
+	{
+		if (UMetaSoundSource* Source = Cast<UMetaSoundSource>(Component->GetSound()))
+		{
+			return FString::Printf(TEXT("MetasoundMusicSource: %s"), *Source->GetDisplayName().ToString());
+		}
+	}
+#endif
+	return TEXT("MetasoundMusicSource");
+}
+#endif
+
+// ---- Update (called per frame by subsystem) ----
+
+void UMetasoundMusicSource::Update()
+{
+	if (ensureMsgf(IsInGameThread(), TEXT("%hs called from non-game thread"), __FUNCTION__) == false)
+	{
+		return;
+	}
+
+	// Try to connect if not yet connected
+	if (AudioComponentToWatch.IsValid() && !CurrentGeneratorHandle)
+	{
+		AttemptToConnect(AudioComponentToWatch.Get());
+	}
+
+	bSeekDetected = false;
+	bLoopDetected = false;
+
+	if (!ClockHistory)
+	{
+		// Connected but no history yet, or not connected at all
+		CachedSourceState = AudioComponentToWatch.IsValid()
+			? Harmonix::ESourceState::Preparing
+			: Harmonix::ESourceState::Stopped;
+		return;
+	}
+
+	// Read the Metasound's actual transport state from the clock history
+	auto Entry = ClockHistory->Positions.GetEntry(ClockHistory->Positions.GetLastWriteIndex());
+	bool bMetasoundIsPlaying = Entry->Item.CurrentTransportState == HarmonixMetasound::EMusicPlayerTransportState::Playing;
+
+	if (bMetasoundIsPlaying)
+	{
+		CachedSourceState = Harmonix::ESourceState::Running;
+		PrevPos = CurrentPos;
+		RefreshFromHistory();
+	}
+	else
+	{
+		// Metasound exists but isn't playing
+		Harmonix::ESourceState PrevState = CachedSourceState;
+		CachedSourceState = Harmonix::ESourceState::Stopped;
+
+		// If we just transitioned to stopped, reset position
+		if (PrevState == Harmonix::ESourceState::Running)
+		{
+			CurrentPos.Reset();
+			PrevPos.Reset();
+		}
+	}
+}
+
+// ---- History-Based Position (Smoothed Audio Render Time) ----
+
+void UMetasoundMusicSource::RefreshFromHistory()
+{
+	using namespace HarmonixMetasound::Analysis;
+	using namespace MetasoundMusicSourceSmoothing;
+
+	check(IsInGameThread());
+
+	if (!SmoothedClockHistoryCursor.DataAvailable())
+	{
+		return;
+	}
+
+	if (!SmoothedClockHistoryCursor.Queue)
+	{
+		return;
+	}
+
+	if (!CurrentMapChain || !CurrentMapChain->SongMaps || !CurrentMapChain->IsLatest())
+	{
+		CurrentMapChain = ClockHistory->GetLatestMapsForConsumer();
+		if (!CurrentMapChain || !CurrentMapChain->SongMaps)
+		{
+			return;
+		}
+	}
+
+	// Read raw position from the history ring
+	auto Entry = ClockHistory->Positions.GetEntry(ClockHistory->Positions.GetLastWriteIndex());
+	Metasound::FSampleCount LastRenderPosSampleCount = Entry->Item.SampleCount;
+	float SpeedAtRawRenderTime = Entry->Item.CurrentSpeed;
+	LastTickSeen = Entry->Item.UpToTick;
+	bool bClockIsStopped = Entry->Item.CurrentTransportState != HarmonixMetasound::EMusicPlayerTransportState::Playing;
+
+	double CurrentWallClockSeconds = FPlatformTime::Seconds();
+	// Initialize sync point on first update
+	if (RenderStartWallClockTimeSeconds == 0.0)
+	{
+		RenderStartSampleCount = LastRenderPosSampleCount;
+		RenderStartWallClockTimeSeconds = CurrentWallClockSeconds - static_cast<double>(RenderStartSampleCount) / static_cast<double>(ClockHistory->SampleRate);
+		RenderSmoothingLagSeconds = LagSeconds;
+		ErrorTracker.Reset();
+		LastRefreshWallClockTimeSeconds = RenderStartWallClockTimeSeconds;
+	}
+	
+	DeltaSecondsBetweenRefreshes = CurrentWallClockSeconds - LastRefreshWallClockTimeSeconds;
+	LastRefreshWallClockTimeSeconds = CurrentWallClockSeconds;
+
+	double ExpectedRenderedSeconds = (CurrentWallClockSeconds - RenderStartWallClockTimeSeconds) * SyncSpeed;
+	double RenderedSeconds = static_cast<double>(LastRenderPosSampleCount) / static_cast<double>(ClockHistory->SampleRate);
+	double Error = RenderedSeconds - ExpectedRenderedSeconds;
+
+	if (!bClockIsStopped)
+	{
+		ErrorTracker.Push(Error);
+
+		if (FMath::Abs(ErrorTracker.Min()) > MaxErrorSecondsBeforeJump)
+		{
+			UE_LOG(LogMetasoundMusicSource, Verbose, TEXT("======== MASSIVE ERROR (%f) - SEEKING ==========="), static_cast<float>(Error));
+			RenderStartSampleCount = LastRenderPosSampleCount;
+			RenderStartWallClockTimeSeconds = CurrentWallClockSeconds - static_cast<double>(RenderStartSampleCount) / static_cast<double>(ClockHistory->SampleRate);
+			ExpectedRenderedSeconds = RenderedSeconds;
+			RenderSmoothingLagSeconds = LagSeconds;
+			ErrorTracker.Reset();
+			Error = 0;
+			SyncSpeed = 1.0;
+		}
+
+		if (ExpectedRenderedSeconds > 0.0)
+		{
+			SyncSpeed += kP * ErrorTracker.Min() / ExpectedRenderedSeconds;
+		}
+		SyncSpeed = FMath::Clamp(SyncSpeed, MinSyncSpeed, MaxSyncSpeed);
+	}
+
+	Metasound::FSampleCount ExpectedRenderPosSampleCount = static_cast<Metasound::FSampleCount>(ExpectedRenderedSeconds * static_cast<double>(ClockHistory->SampleRate));
+
+	// Smooth the audio render position
+	const ISongMapEvaluator* Maps = GetCurrentSongMapEvaluator();
+	float SpeedAtSmoothedTime = 1.f;
+	float SmoothedTick = Maps->MsToTick(CurrentPos.SecondsIncludingCountIn * 1000.f);
+	float SmoothedTempoMapTick = SmoothedState.TempoMapTick;
+
+	EHistoryFailureType Result = CalculateSmoothedTick(
+		ExpectedRenderPosSampleCount, LastRenderPosSampleCount,
+		SmoothedTick, SmoothedTempoMapTick, SpeedAtSmoothedTime,
+		SmoothedClockHistoryCursor, RenderSmoothingLagSeconds);
+
+	if (Result != EHistoryFailureType::None)
+	{
+		if (LastRenderPosSampleCount > static_cast<Metasound::FSampleCount>(RenderSmoothingLagSeconds * ClockHistory->SampleRate * 2))
+		{
+			if (RenderSmoothingLagSeconds < 0.250f)
+			{
+				RenderSmoothingLagSeconds += 0.005f;
+				UE_LOG(LogMetasoundMusicSource, Verbose, TEXT("Bumping smoothing lag to %f"), RenderSmoothingLagSeconds);
+			}
+		}
+		return;
+	}
+
+	if (SmoothedTempoMapTick != SmoothedTick && CurrentMapChain->LoopLengthTicks <= 0)
+	{
+		UpdateSmoothedPositionForOffsetClock(SmoothedTick, SmoothedTempoMapTick);
+	}
+	else
+	{
+		UpdateSmoothedPosition(SmoothedTick, SmoothedTempoMapTick);
+	}
+
+	CurrentAdvanceRate = SpeedAtRawRenderTime;
+}
+
+void UMetasoundMusicSource::UpdateSmoothedPosition(float SmoothedTick, float SmoothedTempoMapTick)
+{
+	SmoothedState.TempoMapMs = CurrentMapChain->SongMaps->TickToMs(SmoothedTempoMapTick);
+	SmoothedState.TempoMapMs += RenderSmoothingLagSeconds * 1000.f;
+
+	SmoothedState.TempoMapTick = CurrentMapChain->SongMaps->MsToTick(SmoothedState.TempoMapMs);
+	CurrentPos = CalculateSongPosForLoopingOrMonotonic(
+		SmoothedState.TempoMapMs, SmoothedState.LocalTick, bSeekDetected, bLoopDetected);
+}
+
+void UMetasoundMusicSource::UpdateSmoothedPositionForOffsetClock(float SmoothedTick, float SmoothedTempoMapTick)
+{
+	float SmoothedPositionMs = CurrentMapChain->SongMaps->TickToMs(SmoothedTick);
+	SmoothedPositionMs += RenderSmoothingLagSeconds * 1000.f;
+
+	CurrentPos = CalculateSongPosForOffsetClock(
+		SmoothedPositionMs, SmoothedTick - SmoothedTempoMapTick,
+		SmoothedState.LocalTick, bSeekDetected);
+
+	float SmoothedTempoMapMs = CurrentMapChain->SongMaps->TickToMs(SmoothedTempoMapTick);
+	SmoothedState.TempoMapMs = SmoothedTempoMapMs + (RenderSmoothingLagSeconds * 1000.f);
+	SmoothedState.TempoMapTick = CurrentMapChain->SongMaps->MsToTick(SmoothedState.TempoMapMs);
+}
+
+// ---- Song Position Calculation ----
+
+FMidiSongPos UMetasoundMusicSource::CalculateSongPosForLoopingOrMonotonic(
+	float AbsoluteMs, float& InOutPositionTick, bool& OutSeekDetected, bool& OutLoopDetected) const
+{
+	FMidiSongPos OutSongPos;
+	if (!ClockHistory || !CurrentMapChain || !CurrentMapChain->SongMaps)
+	{
+		InOutPositionTick = 0.f;
+		return OutSongPos;
+	}
+
+	float NewPositionTick = 0.f;
+	if (CurrentMapChain->LoopLengthTicks > 0)
+	{
+		float DrivingTick = CurrentMapChain->SongMaps->MsToTick(AbsoluteMs);
+		float TickPastLoop = static_cast<float>(CurrentMapChain->FirstTickInLoop + CurrentMapChain->LoopLengthTicks);
+		if (DrivingTick >= TickPastLoop)
+		{
+			float WrappedTick = FMath::Fmod(DrivingTick - CurrentMapChain->FirstTickInLoop, static_cast<float>(CurrentMapChain->LoopLengthTicks));
+			OutLoopDetected = (InOutPositionTick - WrappedTick) > static_cast<float>(CurrentMapChain->LoopLengthTicks - 240);
+
+			OutSongPos.SetByTick(WrappedTick, *CurrentMapChain->SongMaps);
+			OutSongPos.Tempo = CurrentMapChain->SongMaps->GetTempoAtTick(FMath::FloorToInt32(DrivingTick));
+
+			if (!OutLoopDetected)
+			{
+				OutSeekDetected = CheckForSeek(InOutPositionTick, WrappedTick, OutSongPos.Tempo, CurrentMapChain->SongMaps->GetTicksPerQuarterNote());
+			}
+
+			InOutPositionTick = WrappedTick;
+			return OutSongPos;
+		}
+		NewPositionTick = DrivingTick;
+	}
+	else
+	{
+		NewPositionTick = CurrentMapChain->SongMaps->MsToTick(AbsoluteMs);
+	}
+
+	OutSongPos.SetByTimeAndTick(AbsoluteMs, NewPositionTick, *CurrentMapChain->SongMaps);
+	OutSeekDetected = CheckForSeek(InOutPositionTick, NewPositionTick, OutSongPos.Tempo, CurrentMapChain->SongMaps->GetTicksPerQuarterNote());
+	InOutPositionTick = NewPositionTick;
+	return OutSongPos;
+}
+
+FMidiSongPos UMetasoundMusicSource::CalculateSongPosForOffsetClock(
+	float PositionMs, float ClockTickOffset, float& InOutPositionTick, bool& OutSeekDetected) const
+{
+	FMidiSongPos OutSongPos;
+	if (!ClockHistory || !CurrentMapChain || !CurrentMapChain->SongMaps)
+	{
+		InOutPositionTick = 0.f;
+		return OutSongPos;
+	}
+
+	float NewPositionTick = CurrentMapChain->SongMaps->MsToTick(PositionMs);
+	OutSongPos.SetByTick(NewPositionTick, *CurrentMapChain->SongMaps);
+	OutSongPos.Tempo = CurrentMapChain->SongMaps->GetTempoAtTick(FMath::FloorToInt32(NewPositionTick - ClockTickOffset));
+
+	OutSeekDetected = CheckForSeek(InOutPositionTick, NewPositionTick, OutSongPos.Tempo, CurrentMapChain->SongMaps->GetTicksPerQuarterNote());
+	InOutPositionTick = NewPositionTick;
+	return OutSongPos;
+}
+
+bool UMetasoundMusicSource::CheckForSeek(float FirstTick, float NextTick, float CurrentTempo, int32 TicksPerQuarter) const
+{
+	float QuartersPerSecond = CurrentTempo / 60.f;
+	float ExpectedDeltaQuarters = QuartersPerSecond * static_cast<float>(DeltaSecondsBetweenRefreshes);
+	float ExpectedDeltaTicks = ExpectedDeltaQuarters * static_cast<float>(TicksPerQuarter);
+	return FMath::Abs(ExpectedDeltaTicks - (NextTick - FirstTick)) > (ExpectedDeltaTicks * 2.f);
+}
+
+// ---- Smoothed Tick from History ----
+
+UMetasoundMusicSource::EHistoryFailureType UMetasoundMusicSource::CalculateSmoothedTick(
+	Metasound::FSampleCount ExpectedRenderPosSampleCount,
+	Metasound::FSampleCount LastRenderPosSampleCount,
+	float& SmoothedLocalTick, float& SmoothedTempoMapTick, float& CurrentSpeed,
+	HarmonixMetasound::Analysis::FMidiClockSongPositionHistory::FReadCursor& ReadCursor,
+	float LookBehindSeconds)
+{
+	using namespace HarmonixMetasound::Analysis;
+
+	Metasound::FSampleCount LookingForSampleFrame = ExpectedRenderPosSampleCount - static_cast<Metasound::FSampleCount>(LookBehindSeconds * ClockHistory->SampleRate);
+
+	int32 NumHistoryAvailable = ReadCursor.NumDataAvailable();
+	if (LookingForSampleFrame >= LastRenderPosSampleCount && NumHistoryAvailable > 1)
+	{
+		while (ReadCursor.NumDataAvailable() > 1)
+		{
+			ReadCursor.ConsumeNext();
+		}
+		NumHistoryAvailable = ReadCursor.NumDataAvailable();
+	}
+
+	if (NumHistoryAvailable == 0)
+	{
+		return EHistoryFailureType::NotEnoughDataInTheHistoryRing;
+	}
+
+	FMidiClockSongPositionHistory::FScopedItemPeekRef PeekNextRef = ReadCursor.PeekNext();
+
+	if (NumHistoryAvailable == 1 || PeekNextRef->SampleCount > LookingForSampleFrame)
+	{
+		SmoothedLocalTick = static_cast<float>(PeekNextRef->UpToTick);
+		SmoothedTempoMapTick = static_cast<float>(PeekNextRef->TempoMapTick);
+		CurrentSpeed = PeekNextRef->CurrentSpeed;
+		return EHistoryFailureType::None;
+	}
+
+	FMidiClockSongPositionHistory::FScopedItemPeekRef PeekOneAheadRef(ReadCursor.PeekAhead(1));
+	while (PeekOneAheadRef && PeekOneAheadRef->SampleCount <= LookingForSampleFrame)
+	{
+		ReadCursor.PeekAhead(2, PeekOneAheadRef);
+		ReadCursor.PeekAhead(1, PeekNextRef);
+		ReadCursor.ConsumeNext();
+	}
+
+	if (!PeekOneAheadRef)
+	{
+		return EHistoryFailureType::CaughtUpToRenderPosition;
+	}
+
+	check(LookingForSampleFrame >= PeekNextRef->SampleCount && LookingForSampleFrame < PeekOneAheadRef->SampleCount);
+
+	float LerpAlpha = static_cast<float>(LookingForSampleFrame - PeekNextRef->SampleCount) / static_cast<float>(PeekOneAheadRef->SampleCount - PeekNextRef->SampleCount);
+	SmoothedLocalTick = FMath::Lerp<float>(static_cast<float>(PeekNextRef->UpToTick), static_cast<float>(PeekOneAheadRef->UpToTick), LerpAlpha);
+	SmoothedTempoMapTick = FMath::Lerp<float>(static_cast<float>(PeekNextRef->TempoMapTick), static_cast<float>(PeekOneAheadRef->TempoMapTick), LerpAlpha);
+	CurrentSpeed = PeekNextRef->CurrentSpeed;
+
+	return EHistoryFailureType::None;
+}
+
+FString UMetasoundMusicSource::HistoryFailureTypeToString(EHistoryFailureType Error)
+{
+	switch (Error)
+	{
+	case EHistoryFailureType::None: return TEXT("None");
+	case EHistoryFailureType::NotEnoughDataInTheHistoryRing: return TEXT("NotEnoughDataInTheHistoryRing");
+	case EHistoryFailureType::NotEnoughHistory: return TEXT("NotEnoughHistory");
+	case EHistoryFailureType::LookingForTimeInTheFutureOfWhatHasEvenRendered: return TEXT("LookingForTimeInTheFutureOfWhatHasEvenRendered");
+	case EHistoryFailureType::CaughtUpToRenderPosition: return TEXT("CaughtUpToRenderPosition");
+	}
+	return TEXT("Unrecognized");
+}
+
+// ---- Generator Connection ----
+
+bool UMetasoundMusicSource::AttemptToConnect(UAudioComponent* InAudioComponent)
+{
+	using namespace HarmonixMetasound::Analysis;
+
+	check(IsInGameThread());
+	if (!AudioComponentToWatch.IsValid() || MetasoundOutputName.IsNone())
+	{
+		return false;
+	}
+
+	if (!Cast<UMetaSoundSource>(AudioComponentToWatch->GetSound()))
+	{
+		return false;
+	}
+
+	DetachAllCallbacks();
+	CurrentGeneratorHandle.Reset(UMetasoundGeneratorHandle::CreateMetaSoundGeneratorHandle(AudioComponentToWatch.Get()));
+	if (!CurrentGeneratorHandle)
+	{
+		return false;
+	}
+
+	bool bWatchingOutput = CurrentGeneratorHandle->WatchOutput(MetasoundOutputName,
+		FOnMetasoundOutputValueChangedNative::CreateLambda([](FName, const FMetaSoundOutput&) {}),
+		FMidiSongPosVertexAnalyzer::GetAnalyzerName(),
+		FMidiSongPosVertexAnalyzer::SongPosition.Name);
+
+	if (bWatchingOutput)
+	{
+		ensure(CurrentGeneratorHandle->TryCreateAnalyzerAddress(MetasoundOutputName,
+			FMidiSongPosVertexAnalyzer::GetAnalyzerName(),
+			FMidiSongPosVertexAnalyzer::SongPosition.Name,
+			MidiSongPosAnalyzerAddress));
+	}
+
+	GeneratorAttachedCallbackHandle = CurrentGeneratorHandle->OnGeneratorHandleAttached.AddLambda([this]()
+	{
+		OnGeneratorAttached();
+	});
+	GeneratorDetachedCallbackHandle = CurrentGeneratorHandle->OnGeneratorHandleDetached.AddLambda([this]()
+	{
+		OnGeneratorDetached();
+	});
+	GeneratorIOUpdatedCallbackHandle = CurrentGeneratorHandle->OnIOUpdatedWithChanges.AddLambda([this](const TArray<Metasound::FVertexInterfaceChange>& Changes)
+	{
+		OnGeneratorIOUpdatedWithChanges(Changes);
+	});
+
+	UMetasoundGeneratorHandle::FOnSetGraph::FDelegate OnSetGraph;
+	OnSetGraph.BindLambda([this]() { OnGraphSet(); });
+	GraphChangedCallbackHandle = CurrentGeneratorHandle->AddGraphSetCallback(MoveTemp(OnSetGraph));
+
+	return true;
+}
+
+void UMetasoundMusicSource::DetachAllCallbacks()
+{
+	if (CurrentGeneratorHandle)
+	{
+		CurrentGeneratorHandle->OnGeneratorHandleAttached.Remove(GeneratorAttachedCallbackHandle);
+		GeneratorAttachedCallbackHandle.Reset();
+		CurrentGeneratorHandle->OnGeneratorHandleDetached.Remove(GeneratorDetachedCallbackHandle);
+		GeneratorDetachedCallbackHandle.Reset();
+		CurrentGeneratorHandle->OnIOUpdatedWithChanges.Remove(GeneratorIOUpdatedCallbackHandle);
+		GeneratorIOUpdatedCallbackHandle.Reset();
+		CurrentGeneratorHandle->RemoveGraphSetCallback(GraphChangedCallbackHandle);
+		GraphChangedCallbackHandle.Reset();
+	}
+	SetClockHistory(nullptr);
+}
+
+void UMetasoundMusicSource::OnGeneratorAttached()
+{
+	UnderlyingGenerator = CurrentGeneratorHandle->GetGenerator();
+	SetClockHistory(UMidiClockUpdateSubsystem::GetOrCreateClockHistory(MidiSongPosAnalyzerAddress));
+	if (!ClockHistory)
+	{
+		UE_LOG(LogMetasoundMusicSource, Warning, TEXT("OnGeneratorAttached: Failed to create clock history for address '%s'"), *MidiSongPosAnalyzerAddress.ToString());
+		return;
+	}
+	SmoothedClockHistoryCursor = ClockHistory->CreateReadCursor();
+
+	// Reset PLL sync state so it re-initializes on first RefreshFromHistory()
+	RenderStartSampleCount = 0;
+	RenderStartWallClockTimeSeconds = 0.0;
+}
+
+void UMetasoundMusicSource::OnGraphSet()
+{
+	SetClockHistory(UMidiClockUpdateSubsystem::GetOrCreateClockHistory(MidiSongPosAnalyzerAddress));
+	if (!ClockHistory)
+	{
+		UE_LOG(LogMetasoundMusicSource, Warning, TEXT("OnGraphSet: Failed to create clock history for address '%s'"), *MidiSongPosAnalyzerAddress.ToString());
+		return;
+	}
+	SmoothedClockHistoryCursor = ClockHistory->CreateReadCursor();
+}
+
+void UMetasoundMusicSource::OnGeneratorIOUpdatedWithChanges(const TArray<Metasound::FVertexInterfaceChange>& VertexInterfaceChanges)
+{
+	if (!MetasoundOutputName.IsNone() && VertexInterfaceChanges.Num() > 0)
+	{
+		SetClockHistory(UMidiClockUpdateSubsystem::GetOrCreateClockHistory(MidiSongPosAnalyzerAddress));
+		if (!ClockHistory)
+		{
+			UE_LOG(LogMetasoundMusicSource, Warning, TEXT("OnGeneratorIOUpdatedWithChanges: Failed to create clock history for address '%s'"), *MidiSongPosAnalyzerAddress.ToString());
+			return;
+		}
+		SmoothedClockHistoryCursor = ClockHistory->CreateReadCursor();
+	}
+}
+
+void UMetasoundMusicSource::OnGeneratorDetached()
+{
+	TWeakPtr<Metasound::FMetasoundGenerator> DetachingGenerator = CurrentGeneratorHandle->GetGenerator();
+	if (DetachingGenerator != UnderlyingGenerator)
+	{
+		return;
+	}
+
+	SetClockHistory(nullptr);
+	SmoothedClockHistoryCursor = HarmonixMetasound::Analysis::FMidiClockSongPositionHistory::FReadCursor();
+}
+
+void UMetasoundMusicSource::SetClockHistory(const UMidiClockUpdateSubsystem::FClockHistoryPtr& NewHistory)
+{
+	if (NewHistory != ClockHistory)
+	{
+		ClockHistory = NewHistory;
+		CurrentMapChain = nullptr;
+	}
+}
diff --git a/Engine/Plugins/Runtime/Harmonix/Source/HarmonixMetasound/Private/HarmonixMetasound/MusicSource/MusicSource.cpp b/Engine/Plugins/Runtime/Harmonix/Source/HarmonixMetasound/Private/HarmonixMetasound/MusicSource/MusicSource.cpp
new file mode 100644
--- /dev/null
+++ b/Engine/Plugins/Runtime/Harmonix/Source/HarmonixMetasound/Private/HarmonixMetasound/MusicSource/MusicSource.cpp
@@ -0,0 +1,2 @@
+// Copyright Epic Games, Inc. All Rights Reserved.
+#include "HarmonixMetasound/MusicSource/MusicSource.h"
diff --git a/Engine/Plugins/Runtime/Harmonix/Source/HarmonixMetasound/Private/HarmonixMetasound/MusicSource/MusicSourceBlueprintLibrary.cpp b/Engine/Plugins/Runtime/Harmonix/Source/HarmonixMetasound/Private/HarmonixMetasound/MusicSource/MusicSourceBlueprintLibrary.cpp
new file mode 100644
--- /dev/null
+++ b/Engine/Plugins/Runtime/Harmonix/Source/HarmonixMetasound/Private/HarmonixMetasound/MusicSource/MusicSourceBlueprintLibrary.cpp
@@ -0,0 +1,32 @@
+// Copyright Epic Games, Inc. All Rights Reserved.
+
+#include "HarmonixMetasound/MusicSource/MusicSourceBlueprintLibrary.h"
+#include "HarmonixMetasound/MusicClock/MusicClock.h"
+#include "HarmonixMetasound/Subsystems/MidiClockUpdateSubsystem.h"
+
+TScriptInterface<IMusicSource> UMusicSourceBlueprintLibrary::CreateMusicSource(
+	UObject* Outer,
+	const TInstancedStruct<FMusicSourceSettings>& Settings)
+{
+	if (!Outer || !Settings.IsValid())
+	{
+		return {};
+	}
+
+	return Settings.Get().CreateSource(Outer);
+}
+
+UMusicClock* UMusicSourceBlueprintLibrary::CreateMusicClock(
+	UObject* Outer,
+	const TScriptInterface<IMusicSource>& Source)
+{
+	if (!Outer || !Source)
+	{
+		return nullptr;
+	}
+
+	UMusicClock* Clock = NewObject<UMusicClock>(Outer);
+	Clock->SetSource(Source);
+	UMidiClockUpdateSubsystem::TrackMusicClock(Clock);
+	return Clock;
+}
diff --git a/Engine/Plugins/Runtime/Harmonix/Source/HarmonixMetasound/Private/HarmonixMetasound/MusicSource/MusicSourceSettings.cpp b/Engine/Plugins/Runtime/Harmonix/Source/HarmonixMetasound/Private/HarmonixMetasound/MusicSource/MusicSourceSettings.cpp
new file mode 100644
--- /dev/null
+++ b/Engine/Plugins/Runtime/Harmonix/Source/HarmonixMetasound/Private/HarmonixMetasound/MusicSource/MusicSourceSettings.cpp
@@ -0,0 +1,69 @@
+// Copyright Epic Games, Inc. All Rights Reserved.
+
+#include "HarmonixMetasound/MusicSource/MusicSourceSettings.h"
+#include "HarmonixMetasound/MusicSource/MetasoundMusicSource.h"
+#include "HarmonixMetasound/MusicSource/RuntimeMusicSource.h"
+#include "HarmonixMetasound/MusicSource/OffsetMusicSource.h"
+#include "HarmonixMetasound/MusicClock/WorldTimeSourceController.h"
+#include "HarmonixMetasound/Subsystems/MidiClockUpdateSubsystem.h"
+#include "Components/AudioComponent.h"
+#include "HarmonixMidi/MidiFile.h"
+
+TScriptInterface<IMusicSource> FMusicSourceSettings::CreateSource(UObject* Outer) const
+{
+	checkNoEntry(); // Base class — subclasses must override. Stripped in shipping.
+	UE_LOG(LogTemp, Warning, TEXT("FMusicSourceSettings::CreateSource called on base type. Use a concrete settings subclass."));
+	return {};
+}
+
+TScriptInterface<IMusicSource> FMetasoundMusicSourceSettings::CreateSource(UObject* Outer) const
+{
+	if (!AudioComponent)
+	{
+		return {};
+	}
+
+	UMetasoundMusicSource* Source = NewObject<UMetasoundMusicSource>(Outer);
+	Source->ConnectToAudioComponent(AudioComponent, OutputPinName);
+	UMidiClockUpdateSubsystem::TrackMusicSource(Source);
+	return TScriptInterface<IMusicSource>(Source);
+}
+
+TScriptInterface<IMusicSource> FMidiMusicSourceSettings::CreateSource(UObject* Outer) const
+{
+	if (!MidiFile)
+	{
+		return {};
+	}
+
+	TSharedPtr<Harmonix::FWorldTimeSourceController> TimeSource = MakeShared<Harmonix::FWorldTimeSourceController>(Outer);
+	URuntimeMusicSource* Source = NewObject<URuntimeMusicSource>(Outer);
+	Source->InitializeWithMidi(TimeSource, MidiFile);
+	UMidiClockUpdateSubsystem::TrackMusicSource(Source);
+	return TScriptInterface<IMusicSource>(Source);
+}
+
+TScriptInterface<IMusicSource> FManualMusicSourceSettings::CreateSource(UObject* Outer) const
+{
+	TSharedPtr<Harmonix::FWorldTimeSourceController> TimeSource = MakeShared<Harmonix::FWorldTimeSourceController>(Outer);
+	URuntimeMusicSource* Source = NewObject<URuntimeMusicSource>(Outer);
+	Source->Initialize(TimeSource, nullptr);
+	Source->SetTempo(Tempo);
+	Source->SetTimeSignature(TimeSigNumerator, TimeSigDenominator);
+	UMidiClockUpdateSubsystem::TrackMusicSource(Source);
+	return TScriptInterface<IMusicSource>(Source);
+}
+
+TScriptInterface<IMusicSource> FOffsetMusicSourceSettings::CreateSource(UObject* Outer) const
+{
+	if (!ParentSource)
+	{
+		return {};
+	}
+
+	UOffsetMusicSource* Source = NewObject<UOffsetMusicSource>(Outer);
+	Source->SetParentSource(ParentSource);
+	Source->SetOffsetMs(OffsetMs);
+	UMidiClockUpdateSubsystem::TrackMusicSource(Source);
+	return TScriptInterface<IMusicSource>(Source);
+}
diff --git a/Engine/Plugins/Runtime/Harmonix/Source/HarmonixMetasound/Private/HarmonixMetasound/MusicSource/OffsetMusicSource.cpp b/Engine/Plugins/Runtime/Harmonix/Source/HarmonixMetasound/Private/HarmonixMetasound/MusicSource/OffsetMusicSource.cpp
new file mode 100644
--- /dev/null
+++ b/Engine/Plugins/Runtime/Harmonix/Source/HarmonixMetasound/Private/HarmonixMetasound/MusicSource/OffsetMusicSource.cpp
@@ -0,0 +1,233 @@
+// Copyright Epic Games, Inc. All Rights Reserved.
+
+#include "HarmonixMetasound/MusicSource/OffsetMusicSource.h"
+
+void UOffsetMusicSource::SetParentSource(const TScriptInterface<IMusicSource>& InParentSource)
+{
+	ParentSourceObject = InParentSource.GetObject();
+	ParentSourceCache = InParentSource ? Cast<IMusicSource>(InParentSource.GetObject()) : nullptr;
+	bParentHasLooped = false;
+}
+
+TScriptInterface<IMusicSource> UOffsetMusicSource::GetParentSource() const
+{
+	if (UObject* Obj = ParentSourceObject.Get())
+	{
+		return TScriptInterface<IMusicSource>(Obj);
+	}
+	return {};
+}
+
+IMusicSource* UOffsetMusicSource::GetValidParentSource() const
+{
+	return ParentSourceObject.IsValid() ? ParentSourceCache : nullptr;
+}
+
+void UOffsetMusicSource::SetOffsetMs(float InOffsetMs)
+{
+	OffsetMs = InOffsetMs;
+}
+
+float UOffsetMusicSource::GetOffsetMs() const
+{
+	return OffsetMs;
+}
+
+Harmonix::ESourceState UOffsetMusicSource::GetSourceState() const
+{
+	IMusicSource* Parent = GetValidParentSource();
+	return Parent ? Parent->GetSourceState() : Harmonix::ESourceState::Stopped;
+}
+
+Harmonix::ESourceEvent UOffsetMusicSource::GetLatestSourceEvent() const
+{
+	IMusicSource* Parent = GetValidParentSource();
+	if (!Parent)
+	{
+		return Harmonix::ESourceEvent::None;
+	}
+
+	if (bLoopedThisFrame)
+	{
+		return Harmonix::ESourceEvent::Loop;
+	}
+	if (bSeekedThisFrame)
+	{
+		return Harmonix::ESourceEvent::Seek;
+	}
+
+	if (Parent->GetSourceState() == Harmonix::ESourceState::Running)
+	{
+		return Harmonix::ESourceEvent::Advance;
+	}
+	return Harmonix::ESourceEvent::None;
+}
+
+const FMidiSongPos& UOffsetMusicSource::GetCurrentSongPos() const
+{
+	return CurrentPos;
+}
+
+const FMidiSongPos& UOffsetMusicSource::GetPreviousSongPos() const
+{
+	return PrevPos;
+}
+
+const ISongMapEvaluator* UOffsetMusicSource::GetCurrentSongMapEvaluator() const
+{
+	IMusicSource* Parent = GetValidParentSource();
+	return Parent ? Parent->GetCurrentSongMapEvaluator() : nullptr;
+}
+
+float UOffsetMusicSource::GetCurrentClockAdvanceRate() const
+{
+	IMusicSource* Parent = GetValidParentSource();
+	return Parent ? Parent->GetCurrentClockAdvanceRate() : 1.f;
+}
+
+bool UOffsetMusicSource::LoopedThisFrame() const
+{
+	return bLoopedThisFrame;
+}
+
+bool UOffsetMusicSource::SeekedThisFrame() const
+{
+	return bSeekedThisFrame;
+}
+
+bool UOffsetMusicSource::IsLooping() const
+{
+	IMusicSource* Parent = GetValidParentSource();
+	return Parent ? Parent->IsLooping() : false;
+}
+
+float UOffsetMusicSource::GetLoopStartMs() const
+{
+	IMusicSource* Parent = GetValidParentSource();
+	return Parent ? Parent->GetLoopStartMs() : 0.f;
+}
+
+float UOffsetMusicSource::GetLoopLengthMs() const
+{
+	IMusicSource* Parent = GetValidParentSource();
+	return Parent ? Parent->GetLoopLengthMs() : 0.f;
+}
+
+TOptional<FVector> UOffsetMusicSource::TryGetAudioSourceLocation() const
+{
+	IMusicSource* Parent = GetValidParentSource();
+	return Parent ? Parent->TryGetAudioSourceLocation() : TOptional<FVector>{};
+}
+
+#if !UE_BUILD_SHIPPING
+FString UOffsetMusicSource::GetDisplayName() const
+{
+	IMusicSource* Parent = GetValidParentSource();
+	if (Parent)
+	{
+		return FString::Printf(TEXT("OffsetMusicSource(%+.1fms): %s"), OffsetMs, *Parent->GetDisplayName());
+	}
+	return TEXT("OffsetMusicSource (no parent)");
+}
+#endif
+
+void UOffsetMusicSource::Update()
+{
+	IMusicSource* Parent = GetValidParentSource();
+	if (!Parent)
+	{
+		if (ParentSourceCache != nullptr)
+		{
+			ParentSourceCache = nullptr;
+		}
+		CurrentPos.Reset();
+		PrevPos.Reset();
+		bLoopedThisFrame = false;
+		bSeekedThisFrame = false;
+		return;
+	}
+
+	if (Parent->GetSourceState() != Harmonix::ESourceState::Running)
+	{
+		if (Parent->GetSourceState() == Harmonix::ESourceState::Stopped)
+		{
+			CurrentPos.Reset();
+			PrevPos.Reset();
+		}
+		bLoopedThisFrame = false;
+		bSeekedThisFrame = false;
+		return;
+	}
+
+	const ISongMapEvaluator* Maps = Parent->GetCurrentSongMapEvaluator();
+	if (!Maps)
+	{
+		return;
+	}
+
+	PrevPos = CurrentPos;
+
+	const FMidiSongPos& ParentPos = Parent->GetCurrentSongPos();
+	float ParentMs = ParentPos.SecondsIncludingCountIn * 1000.f;
+	float OffsetTimeMs = ParentMs + OffsetMs;
+
+	bSeekedThisFrame = Parent->SeekedThisFrame();
+	bLoopedThisFrame = false;
+
+	if (Parent->IsLooping())
+	{
+		float LoopStart = Parent->GetLoopStartMs();
+		float LoopLength = Parent->GetLoopLengthMs();
+
+		// Track whether the parent has completed at least one full loop pass.
+		// Reset on seek (e.g., seeking back to start).
+		if (bSeekedThisFrame)
+		{
+			bParentHasLooped = false;
+		}
+		if (Parent->LoopedThisFrame())
+		{
+			bParentHasLooped = true;
+		}
+
+		if (LoopLength > 0.f)
+		{
+			// Only wrap around the loop region if:
+			//  - The parent has already looped (steady-state), OR
+			//  - The offset time is naturally within or past the loop region
+			//
+			// This prevents a negative offset from jumping to the end of the
+			// loop on the first pass (before the parent has ever crossed the
+			// loop boundary). In that case, the offset time stays negative —
+			// the source is "before" the music start.
+			//
+			// Example with loop 0-4000ms, offset -100ms:
+			//   First pass, parent at 50ms  -> OffsetTime = -50ms  -> stays -50ms (no wrap)
+			//   First pass, parent at 200ms -> OffsetTime = 100ms  -> no wrap needed (in bounds)
+			//   After loop, parent at 50ms  -> OffsetTime = -50ms  -> wraps to 3950ms
+			if (bParentHasLooped || OffsetTimeMs >= LoopStart)
+			{
+				float RelativeMs = OffsetTimeMs - LoopStart;
+				RelativeMs = FMath::Fmod(RelativeMs + LoopLength, LoopLength);
+				OffsetTimeMs = LoopStart + RelativeMs;
+			}
+
+			// Detect if we looped: large position delta relative to loop length
+			float PrevMs = PrevPos.SecondsIncludingCountIn * 1000.f;
+			if (PrevMs > 0.f)
+			{
+				float DeltaMs = OffsetTimeMs - PrevMs;
+				if (FMath::Abs(DeltaMs) > (LoopLength * 0.5f))
+				{
+					bLoopedThisFrame = true;
+				}
+			}
+		}
+	}
+	else
+	{
+		bLoopedThisFrame = Parent->LoopedThisFrame();
+	}
+
+	CurrentPos.SetByTime(OffsetTimeMs, *Maps);
+}
diff --git a/Engine/Plugins/Runtime/Harmonix/Source/HarmonixMetasound/Private/HarmonixMetasound/MusicSource/RuntimeMusicSource.cpp b/Engine/Plugins/Runtime/Harmonix/Source/HarmonixMetasound/Private/HarmonixMetasound/MusicSource/RuntimeMusicSource.cpp
new file mode 100644
--- /dev/null
+++ b/Engine/Plugins/Runtime/Harmonix/Source/HarmonixMetasound/Private/HarmonixMetasound/MusicSource/RuntimeMusicSource.cpp
@@ -0,0 +1,325 @@
+// Copyright Epic Games, Inc. All Rights Reserved.
+
+#include "HarmonixMetasound/MusicSource/RuntimeMusicSource.h"
+#include "Harmonix.h"
+#include "HarmonixMetasound/Subsystems/MidiClockUpdateSubsystem.h"
+#include "HarmonixMidi/MidiFile.h"
+
+void URuntimeMusicSource::BeginDestroy()
+{
+	UMidiClockUpdateSubsystem::StopTrackingMusicSource(this);
+	Super::BeginDestroy();
+}
+
+void URuntimeMusicSource::Initialize(TSharedPtr<Harmonix::ITimeSource> InTimeSource, const ISongMapEvaluator* InSongMaps)
+{
+	TimeSource = InTimeSource;
+	MidiFile = nullptr;
+
+	DefaultMaps.EmptyAllMaps();
+	if (InSongMaps)
+	{
+		DefaultMaps.Copy(*InSongMaps, 0, InSongMaps->GetSongLengthData().LastTick);
+	}
+	else
+	{
+		DefaultMaps.Init(Harmonix::Midi::Constants::GTicksPerQuarterNoteInt);
+		DefaultMaps.GetTempoMap().AddTempoInfoPoint(Harmonix::Midi::Constants::BPMToMidiTempo(120.0f), 0);
+		DefaultMaps.GetBarMap().AddTimeSignatureAtBarIncludingCountIn(0, 4, 4);
+	}
+}
+
+void URuntimeMusicSource::InitializeWithMidi(TSharedPtr<Harmonix::ITimeSource> InTimeSource, UMidiFile* InMidiFile)
+{
+	TimeSource = InTimeSource;
+	MidiFile = InMidiFile;
+}
+
+void URuntimeMusicSource::Start()
+{
+	if (TimeSource)
+	{
+		TimeSource->RequestStart();
+	}
+}
+
+void URuntimeMusicSource::Pause()
+{
+	if (TimeSource)
+	{
+		TimeSource->RequestPause();
+	}
+}
+
+void URuntimeMusicSource::Continue()
+{
+	if (TimeSource)
+	{
+		TimeSource->RequestContinue();
+	}
+}
+
+void URuntimeMusicSource::Stop()
+{
+	if (TimeSource)
+	{
+		TimeSource->RequestStop();
+	}
+	CurrentPos.Reset();
+	PrevPos.Reset();
+}
+
+Harmonix::ESourceState URuntimeMusicSource::GetSourceState() const
+{
+	if (TimeSource)
+	{
+		return TimeSource->GetCurrentSourceState();
+	}
+	return Harmonix::ESourceState::Stopped;
+}
+
+Harmonix::ESourceEvent URuntimeMusicSource::GetLatestSourceEvent() const
+{
+	if (TimeSource)
+	{
+		return TimeSource->GetLatestSourceEvent();
+	}
+	return Harmonix::ESourceEvent::None;
+}
+
+const FMidiSongPos& URuntimeMusicSource::GetCurrentSongPos() const
+{
+	return CurrentPos;
+}
+
+const FMidiSongPos& URuntimeMusicSource::GetPreviousSongPos() const
+{
+	return PrevPos;
+}
+
+const ISongMapEvaluator* URuntimeMusicSource::GetCurrentSongMapEvaluator() const
+{
+	return GetMaps();
+}
+
+float URuntimeMusicSource::GetCurrentClockAdvanceRate() const
+{
+	if (TimeSource)
+	{
+		return TimeSource->GetSpeed();
+	}
+	return 1.0f;
+}
+
+bool URuntimeMusicSource::LoopedThisFrame() const
+{
+	return TimeSource && TimeSource->GetLatestSourceEvent() == Harmonix::ESourceEvent::Loop;
+}
+
+bool URuntimeMusicSource::SeekedThisFrame() const
+{
+	return TimeSource && TimeSource->GetLatestSourceEvent() == Harmonix::ESourceEvent::Seek;
+}
+
+bool URuntimeMusicSource::IsLooping() const
+{
+	return LoopNumBars > 0 && CachedLoopLengthMs > 0.f;
+}
+
+float URuntimeMusicSource::GetLoopStartMs() const
+{
+	return CachedLoopStartMs;
+}
+
+float URuntimeMusicSource::GetLoopLengthMs() const
+{
+	return CachedLoopLengthMs;
+}
+
+void URuntimeMusicSource::SetLoopRegionByBars(int32 StartBar, int32 NumBars)
+{
+	LoopStartBar = StartBar;
+	LoopNumBars = FMath::Max(NumBars, 0);
+	RecomputeLoopRegionMs();
+}
+
+void URuntimeMusicSource::ClearLoopRegion()
+{
+	LoopStartBar = 0;
+	LoopNumBars = 0;
+	CachedLoopStartMs = 0.f;
+	CachedLoopLengthMs = 0.f;
+}
+
+void URuntimeMusicSource::RecomputeLoopRegionMs()
+{
+	if (LoopNumBars <= 0)
+	{
+		CachedLoopStartMs = 0.f;
+		CachedLoopLengthMs = 0.f;
+		return;
+	}
+
+	const ISongMapEvaluator* Maps = GetMaps();
+	if (!Maps)
+	{
+		CachedLoopStartMs = 0.f;
+		CachedLoopLengthMs = 0.f;
+		return;
+	}
+
+	int32 EndBar = LoopStartBar + LoopNumBars;
+
+	// Convert bar numbers to ticks, then to ms
+	float StartTick = Maps->BarIncludingCountInToTick(LoopStartBar);
+	float EndTick = Maps->BarIncludingCountInToTick(EndBar);
+
+	CachedLoopStartMs = Maps->TickToMs(StartTick);
+	CachedLoopLengthMs = Maps->TickToMs(EndTick) - CachedLoopStartMs;
+}
+
+void URuntimeMusicSource::Seek(const FMusicTimestamp& Timestamp)
+{
+	if (!TimeSource)
+	{
+		return;
+	}
+
+	const ISongMapEvaluator* Maps = GetMaps();
+	if (!Maps)
+	{
+		return;
+	}
+
+	float Tick = Maps->MusicTimestampToTick(Timestamp);
+	float Ms = Maps->TickToMs(Tick);
+	float Seconds = Ms / 1000.f;
+
+	TimeSource->RequestSeek(Seconds);
+}
+
+void URuntimeMusicSource::SetTempo(float BPM)
+{
+	if (BPM <= 0.f)
+	{
+		return;
+	}
+
+	// When driven by a MIDI file, we can't modify external maps — copy to DefaultMaps first
+	if (MidiFile.IsValid())
+	{
+		const ISongMapEvaluator* ExternalMaps = MidiFile->GetSongMaps();
+		if (ExternalMaps)
+		{
+			DefaultMaps.Copy(*ExternalMaps, 0, ExternalMaps->GetSongLengthData().LastTick);
+		}
+		MidiFile = nullptr;
+	}
+
+	DefaultMaps.EmptyTempoMap();
+	DefaultMaps.GetTempoMap().AddTempoInfoPoint(Harmonix::Midi::Constants::BPMToMidiTempo(BPM), 0);
+
+	RecomputeLoopRegionMs();
+}
+
+void URuntimeMusicSource::SetTimeSignature(int32 Numerator, int32 Denominator)
+{
+	if (Numerator <= 0 || Denominator <= 0)
+	{
+		return;
+	}
+
+	// When driven by a MIDI file, copy to DefaultMaps first
+	if (MidiFile.IsValid())
+	{
+		const ISongMapEvaluator* ExternalMaps = MidiFile->GetSongMaps();
+		if (ExternalMaps)
+		{
+			DefaultMaps.Copy(*ExternalMaps, 0, ExternalMaps->GetSongLengthData().LastTick);
+		}
+		MidiFile = nullptr;
+	}
+
+	DefaultMaps.EmptyBarMap();
+	DefaultMaps.EmptyBeatMap();
+	DefaultMaps.GetBarMap().AddTimeSignatureAtBarIncludingCountIn(0, Numerator, Denominator);
+
+	RecomputeLoopRegionMs();
+}
+
+void URuntimeMusicSource::SetSpeed(float Speed)
+{
+	if (TimeSource)
+	{
+		TimeSource->RequestSetSpeed(Speed);
+	}
+}
+
+TOptional<FVector> URuntimeMusicSource::TryGetAudioSourceLocation() const
+{
+	if (TimeSource)
+	{
+		return TimeSource->TryGetAudioSourceLocation();
+	}
+	return {};
+}
+
+void URuntimeMusicSource::Update()
+{
+	check(IsInGameThread());
+
+	if (!TimeSource)
+	{
+		return;
+	}
+
+	TimeSource->Update();
+
+	if (TimeSource->GetCurrentSourceState() != Harmonix::ESourceState::Running)
+	{
+		return;
+	}
+
+	const ISongMapEvaluator* Maps = GetMaps();
+	if (!Maps)
+	{
+		return;
+	}
+
+	PrevPos = CurrentPos;
+
+	double SourceTimeSeconds = TimeSource->GetCurrentTime();
+	float SourceTimeMs = static_cast<float>(SourceTimeSeconds * 1000.0);
+
+	if (LoopNumBars > 0 && CachedLoopLengthMs > 0.f)
+	{
+		float LoopEndMs = CachedLoopStartMs + CachedLoopLengthMs;
+		if (SourceTimeMs >= LoopEndMs)
+		{
+			float Relative = FMath::Fmod(SourceTimeMs - CachedLoopStartMs, CachedLoopLengthMs);
+			SourceTimeMs = CachedLoopStartMs + Relative;
+		}
+	}
+
+	CurrentPos.SetByTime(SourceTimeMs, *Maps);
+}
+
+#if !UE_BUILD_SHIPPING
+FString URuntimeMusicSource::GetDisplayName() const
+{
+	if (TimeSource)
+	{
+		return FString::Printf(TEXT("RuntimeMusicSource: %s"), *TimeSource->GetDisplayName());
+	}
+	return TEXT("RuntimeMusicSource (no time source)");
+}
+#endif
+
+const ISongMapEvaluator* URuntimeMusicSource::GetMaps() const
+{
+	if (MidiFile.IsValid())
+	{
+		return MidiFile->GetSongMaps();
+	}
+	return &DefaultMaps;
+}
diff --git a/Engine/Plugins/Runtime/Harmonix/Source/HarmonixMetasound/Private/HarmonixMetasound/Subsystems/MidiClockUpdateSubsystem.cpp b/Engine/Plugins/Runtime/Harmonix/Source/HarmonixMetasound/Private/HarmonixMetasound/Subsystems/MidiClockUpdateSubsystem.cpp
--- a/Engine/Plugins/Runtime/Harmonix/Source/HarmonixMetasound/Private/HarmonixMetasound/Subsystems/MidiClockUpdateSubsystem.cpp
+++ b/Engine/Plugins/Runtime/Harmonix/Source/HarmonixMetasound/Private/HarmonixMetasound/Subsystems/MidiClockUpdateSubsystem.cpp
@@ -5,6 +5,8 @@
 #include "Engine/Engine.h"
 #include "HAL/IConsoleManager.h"
 #include "HarmonixMetasound/Components/MusicClockComponent.h"
+#include "HarmonixMetasound/MusicClock/MusicClock.h"
+#include "HarmonixMetasound/MusicSource/MusicSource.h"
 #include "HarmonixMetasound/DataTypes/MidiClock.h"
 #include "Misc/CoreDelegates.h"
 #include "Misc/ScopeLock.h"
@@ -53,6 +55,44 @@ void UMidiClockUpdateSubsystem::StopTrackingMusicClockComponent(UMusicClockCompo
 	}
 }
 
+void UMidiClockUpdateSubsystem::TrackMusicSource(UObject* Source)
+{
+	check(GEngine);
+	UMidiClockUpdateSubsystem* UpdateSubsystem = GEngine->GetEngineSubsystem<UMidiClockUpdateSubsystem>();
+	check(UpdateSubsystem);
+	UpdateSubsystem->TrackMusicSourceImpl(Source);
+}
+
+void UMidiClockUpdateSubsystem::StopTrackingMusicSource(UObject* Source)
+{
+	if (GEngine)
+	{
+		if (UMidiClockUpdateSubsystem* UpdateSubsystem = GEngine->GetEngineSubsystem<UMidiClockUpdateSubsystem>())
+		{
+			UpdateSubsystem->StopTrackingMusicSourceImpl(Source);
+		}
+	}
+}
+
+void UMidiClockUpdateSubsystem::TrackMusicClock(UMusicClock* Clock)
+{
+	check(GEngine);
+	UMidiClockUpdateSubsystem* UpdateSubsystem = GEngine->GetEngineSubsystem<UMidiClockUpdateSubsystem>();
+	check(UpdateSubsystem);
+	UpdateSubsystem->TrackMusicClockImpl(Clock);
+}
+
+void UMidiClockUpdateSubsystem::StopTrackingMusicClock(UMusicClock* Clock)
+{
+	if (GEngine)
+	{
+		if (UMidiClockUpdateSubsystem* UpdateSubsystem = GEngine->GetEngineSubsystem<UMidiClockUpdateSubsystem>())
+		{
+			UpdateSubsystem->StopTrackingMusicClockImpl(Clock);
+		}
+	}
+}
+
 double UMidiClockUpdateSubsystem::GetUpdateTimeSeconds()
 {
 	return UpdateTimeSeconds;
@@ -113,8 +153,73 @@ void UMidiClockUpdateSubsystem::StopTrackingMusicClockComponentImpl(UMusicClockC
 	TrackedMusicClockComponents.Remove(Clock);
 }
 
+void UMidiClockUpdateSubsystem::TrackMusicSourceImpl(UObject* Source)
+{
+	check(Source && Source->GetClass()->ImplementsInterface(UMusicSource::StaticClass()));
+	TrackedMusicSources.AddUnique(Source);
+}
+
+void UMidiClockUpdateSubsystem::StopTrackingMusicSourceImpl(UObject* Source)
+{
+	TrackedMusicSources.Remove(Source);
+}
+
+void UMidiClockUpdateSubsystem::TrackMusicClockImpl(UMusicClock* Clock)
+{
+	TrackedMusicClocks.AddUnique(Clock);
+}
+
+void UMidiClockUpdateSubsystem::StopTrackingMusicClockImpl(UMusicClock* Clock)
+{
+	TrackedMusicClocks.Remove(Clock);
+}
+
+void UMidiClockUpdateSubsystem::UpdateMusicSources()
+{
+	TRACE_CPUPROFILER_EVENT_SCOPE(UpdateMusicSources);
+
+	for (auto It = TrackedMusicSources.CreateIterator(); It; ++It)
+	{
+		if (UObject* SourceObj = It->Get())
+		{
+			IMusicSource* Source = Cast<IMusicSource>(SourceObj);
+			if (Source)
+			{
+				Source->Update();
+			}
+		}
+		else
+		{
+			It.RemoveCurrentSwap();
+		}
+	}
+}
+
+void UMidiClockUpdateSubsystem::UpdateMusicClocks()
+{
+	TRACE_CPUPROFILER_EVENT_SCOPE(UpdateMusicClocks);
+
+	for (auto It = TrackedMusicClocks.CreateIterator(); It; ++It)
+	{
+		if (UMusicClock* Clock = It->Get())
+		{
+			Clock->UpdateForGameFrame();
+		}
+		else
+		{
+			It.RemoveCurrentSwap();
+		}
+	}
+}
+
 void UMidiClockUpdateSubsystem::CoreDelegatesSamplingInput()
 {
+	// Update order matters:
+	// 1. Music sources produce new time data
+	// 2. Music clocks read from sources, compute deltas, broadcast events
+	// 3. Legacy MusicClockComponents (adapter path)
+	UpdateMusicSources();
+	UpdateMusicClocks();
 	UpdateUMusicClockComponents();
 }
 
@@ -151,5 +256,7 @@ void UMidiClockUpdateSubsystem::UpdateUMusicClockComponents()
 // the test code doesn't need knowledge of how the low-res clocks are being ticked...
 void UMidiClockUpdateSubsystem::TickForTesting()
 {
+	UpdateMusicSources();
+	UpdateMusicClocks();
 	UpdateUMusicClockComponents();
 }
diff --git a/Engine/Plugins/Runtime/Harmonix/Source/HarmonixMetasound/Public/HarmonixMetasound/MusicClock/MusicClock.h b/Engine/Plugins/Runtime/Harmonix/Source/HarmonixMetasound/Public/HarmonixMetasound/MusicClock/MusicClock.h
--- a/Engine/Plugins/Runtime/Harmonix/Source/HarmonixMetasound/Public/HarmonixMetasound/MusicClock/MusicClock.h
+++ b/Engine/Plugins/Runtime/Harmonix/Source/HarmonixMetasound/Public/HarmonixMetasound/MusicClock/MusicClock.h
@@ -1,104 +1,283 @@
 // Copyright Epic Games, Inc. All Rights Reserved.
 #pragma once
 
-#include "Harmonix/MusicalTimebase.h"
 #include "HarmonixMetasound/MusicClock/MusicClockState.h"
+#include "HarmonixMetasound/MusicSource/MusicSource.h"
 #include "HarmonixMidi/MidiSongPos.h"
 #include "HarmonixMidi/SongMaps.h"
 
+#include "MusicClock.generated.h"
+
 #define UE_API HARMONIXMETASOUND_API
 
-struct FMusicClockSettingsBase;
+DECLARE_DYNAMIC_MULTICAST_DELEGATE_TwoParams(FMusicClockBeatEvent, int, BeatNumber, int, BeatInBar);
+DECLARE_DYNAMIC_MULTICAST_DELEGATE_OneParam(FMusicClockBarEvent, int, BarNumber);
+DECLARE_DYNAMIC_MULTICAST_DELEGATE_ThreeParams(FMusicClockSectionEvent, const FString&, SectionName, float, SectionStartMs, float, SectionLengthMs);
+DECLARE_DYNAMIC_MULTICAST_DELEGATE_OneParam(FMusicClockPlayStateEvent, EMusicClockState, State);
 
-namespace Harmonix
+/**
+ * A read-only view into a music source's time.
+ *
+ * UMusicClock does not own or control playback. It observes an IMusicSource
+ * and provides:
+ *   - Current/previous musical position (FMidiSongPos)
+ *   - Tempo, time signature, BPM queries
+ *   - Beat/bar/section event broadcasting
+ *   - Song map delegation for time conversions
+ *
+ * A single UMusicClock represents a single time stream with no
+ * ECalibratedMusicTimebase parameter. For calibrated offsets (experienced,
+ * video render, etc.) use UOffsetMusicSource feeding into a separate UMusicClock.
+ */
+UCLASS(MinimalAPI, BlueprintType)
+class UMusicClock : public UObject
 {
+	GENERATED_BODY()
 
-class FMusicClock
-{
 public:
-	UE_API FMusicClock(UObject* WorldContextObject, const FMusicClockSettingsBase& Settings);
-	UE_API virtual ~FMusicClock();
+	UE_API UMusicClock();
+	UE_API virtual void BeginDestroy() override;
+
+	/** Bind this clock to a music source. The clock becomes a read-only view into that source. */
+	UFUNCTION(BlueprintCallable, Category = "MusicClock")
+	UE_API void SetSource(const TScriptInterface<IMusicSource>& InSource);
+
+	/** Get the source this clock is bound to. */
+	UFUNCTION(BlueprintPure, Category = "MusicClock")
+	UE_API TScriptInterface<IMusicSource> GetSource() const;
+
+	/**
+	 * Update the clock for the current game frame.
+	 * Called by the update subsystem. Pulls position from the source,
+	 * computes deltas, and broadcasts events.
+	 */
+	UE_API virtual void UpdateForGameFrame();
+
+	// ---- State ----
+
+	/** Get the current clock state (Stopped, Paused, Running). */
+	UFUNCTION(BlueprintPure, Category = "MusicClock")
+	UE_API virtual EMusicClockState GetState() const;
+
+	// ---- Position ----
+
+	/** Get the current musical position for this frame. */
+	UFUNCTION(BlueprintPure, Category = "MusicClock")
+	UE_API const FMidiSongPos& GetCurrentSongPos() const;
+
+	/** Get the previous frame's musical position. */
+	UFUNCTION(BlueprintPure, Category = "MusicClock")
+	UE_API const FMidiSongPos& GetPreviousSongPos() const;
+
+	/** Time from the beginning of the authored music content, including count-in/pickup bars. */
+	UFUNCTION(BlueprintPure, Category = "MusicClock")
+	UE_API float GetSecondsIncludingCountIn() const;
+
+	/** Time from Bar 1 Beat 1 (the classic "start of the song"), excluding count-in/pickup. */
+	UFUNCTION(BlueprintPure, Category = "MusicClock")
+	UE_API float GetSecondsFromBarOne() const;
+
+	/** Fractional total bars from the beginning, including count-in. */
+	UFUNCTION(BlueprintPure, Category = "MusicClock")
+	UE_API float GetBarsIncludingCountIn() const;
+
+	/** Fractional total beats from the beginning, including count-in. */
+	UFUNCTION(BlueprintPure, Category = "MusicClock")
+	UE_API float GetBeatsIncludingCountIn() const;
+
+	/**
+	 * Get the current musical timestamp (Bar + Beat).
+	 * Bar 1, Beat 1.0 is the beginning of the song after count-in.
+	 */
+	UFUNCTION(BlueprintPure, Category = "MusicClock")
+	UE_API FMusicTimestamp GetCurrentTimestamp() const;
+
+	// ---- Tempo & Time Signature ----
+
+	/** Current tempo in MIDI quarter notes per minute. */
+	UFUNCTION(BlueprintPure, Category = "MusicClock")
+	UE_API float GetCurrentTempo() const;
+
+	/** Current BPM in true beats (not quarter notes). */
+	UFUNCTION(BlueprintPure, Category = "MusicClock")
+	UE_API float GetCurrentBeatsPerMinute() const;
+
+	/** Current beats per second, accounting for tempo, time signature, and speed. */
+	UFUNCTION(BlueprintPure, Category = "MusicClock")
+	UE_API float GetCurrentBeatsPerSecond() const;
+
+	/** Current bars per second, accounting for tempo, time signature, and speed. */
+	UFUNCTION(BlueprintPure, Category = "MusicClock")
+	UE_API float GetCurrentBarsPerSecond() const;
+
+	/** Current playback speed multiplier from the source. */
+	UFUNCTION(BlueprintPure, Category = "MusicClock")
+	UE_API float GetCurrentClockAdvanceRate() const;
+
+	/** Get the current time signature as numerator and denominator. */
+	UFUNCTION(BlueprintPure, Category = "MusicClock")
+	UE_API void GetCurrentTimeSignature(int& OutNumerator, int& OutDenominator) const;
+
+	// ---- Sections ----
+
+	/** Name of the current song section (from MIDI section markers). */
+	UFUNCTION(BlueprintPure, Category = "MusicClock")
+	UE_API FString GetCurrentSectionName() const;
+
+	/** Index of the current section in the section map. */
+	UFUNCTION(BlueprintPure, Category = "MusicClock")
+	UE_API int32 GetCurrentSectionIndex() const;
 
-	virtual TOptional<FVector> TryGetAudioSourceLocation() const = 0;
+	/** Start time of the current section in milliseconds. */
+	UFUNCTION(BlueprintPure, Category = "MusicClock")
+	UE_API float GetCurrentSectionStartMs() const;
 
-	UE_API void UpdateClockForGameFrame();
+	/** Length of the current section in milliseconds. */
+	UFUNCTION(BlueprintPure, Category = "MusicClock")
+	UE_API float GetCurrentSectionLengthMs() const;
 
-	UE_API virtual bool CalculateSongPosWithOffset(float MsOffset, ECalibratedMusicTimebase Timebase, FMidiSongPos& OutResult) const;
+	// ---- Beat/Bar Distance ----
+	// All distance values are fractional: 0.0 = on the beat/bar, 0.5 = halfway to the next.
 
-	UE_API double GetWallClockTime() const;
+	/** Fractional distance past the most recent beat (0.0 = exactly on beat). */
+	UFUNCTION(BlueprintPure, Category = "MusicClock")
+	UE_API float GetDistanceFromCurrentBeat() const;
 
-	EMusicClockState GetState() const { return MusicClockState; }
-	UE_API const FMidiSongPos& GetCurrentSongPos(ECalibratedMusicTimebase Timebase) const;
-	UE_API const FMidiSongPos& GetPreviousSongPos(ECalibratedMusicTimebase Timebase) const;
-	UE_API float GetDeltaBarF(ECalibratedMusicTimebase Timebase) const;
-	UE_API float GetDeltaBeatF(ECalibratedMusicTimebase Timebase) const;
+	/** Fractional distance until the next beat (0.0 = exactly on next beat). */
+	UFUNCTION(BlueprintPure, Category = "MusicClock")
+	UE_API float GetDistanceToNextBeat() const;
 
-	UE_API void Start();
-	UE_API void Pause();
-	UE_API void Continue();
-	UE_API void Stop();
-	
-	virtual const ISongMapEvaluator* GetCurrentSongMapEvaluator() const = 0;
+	/** Fractional distance to the closest beat (past or future), always <= 0.5. */
+	UFUNCTION(BlueprintPure, Category = "MusicClock")
+	UE_API float GetDistanceToClosestBeat() const;
 
-	virtual bool LoopedThisFrame(ECalibratedMusicTimebase Timebase) const = 0;
-	virtual bool SeekedThisFrame(ECalibratedMusicTimebase Timebase) const = 0;
+	/** Fractional distance past the most recent bar (0.0 = exactly on bar). */
+	UFUNCTION(BlueprintPure, Category = "MusicClock")
+	UE_API float GetDistanceFromCurrentBar() const;
 
-	virtual void Disconnect() = 0;
+	/** Fractional distance until the next bar. */
+	UFUNCTION(BlueprintPure, Category = "MusicClock")
+	UE_API float GetDistanceToNextBar() const;
 
-	UE_API void UpdateMusicPlaybackRate(float InTempo, float InClockAdvanceRate, int32 InTimeSigNum, int32 InTimeSigDenom);
+	/** Fractional distance to the closest bar (past or future), always <= 0.5. */
+	UFUNCTION(BlueprintPure, Category = "MusicClock")
+	UE_API float GetDistanceToClosestBar() const;
 
-	// This returns our BeatsPerMinute which is our true beats, not quarter notes.
-	UE_API float GetBeatsPerMinute() const;
+	// ---- Deltas ----
+
+	/** Fractional bars advanced this frame. */
+	UFUNCTION(BlueprintPure, Category = "MusicClock")
+	UE_API float GetDeltaBar() const;
+
+	/** Fractional beats advanced this frame. */
+	UFUNCTION(BlueprintPure, Category = "MusicClock")
+	UE_API float GetDeltaBeat() const;
+
+	// ---- Discontinuities ----
+
+	/** Whether the source crossed a loop boundary this frame. */
+	UFUNCTION(BlueprintPure, Category = "MusicClock")
+	UE_API bool LoopedThisFrame() const;
+
+	/** Whether the source performed a non-contiguous time jump this frame. */
+	UFUNCTION(BlueprintPure, Category = "MusicClock")
+	UE_API bool SeekedThisFrame() const;
+
+	// ---- Loop Region ----
+
+	/** Whether the source is looping. */
+	UFUNCTION(BlueprintPure, Category = "MusicClock")
+	UE_API virtual bool IsLooping() const;
+
+	/** Start of loop region in ms. Only meaningful if IsLooping(). */
+	UFUNCTION(BlueprintPure, Category = "MusicClock")
+	UE_API virtual float GetLoopStartMs() const;
+
+	/** Length of loop region in ms. Only meaningful if IsLooping(). */
+	UFUNCTION(BlueprintPure, Category = "MusicClock")
+	UE_API virtual float GetLoopLengthMs() const;
+
+	// ---- Song Map Delegation ----
+
+	/** Get the song maps from the bound source, for time<->beat conversions. */
+	UE_API virtual const ISongMapEvaluator* GetSongMaps() const;
+
+	// ---- Song Data ----
+
+	/** Total length of the song in milliseconds. */
+	UFUNCTION(BlueprintPure, Category = "MusicClock")
+	UE_API float GetSongLengthMs() const;
+
+	/** Remaining time until the end of the song in milliseconds. */
+	UFUNCTION(BlueprintPure, Category = "MusicClock")
+	UE_API float GetSongRemainingMs() const;
+
+	// ---- Spatial ----
+
+	/** Get the world-space location of the audio source, if available. */
+	UFUNCTION(BlueprintPure, Category = "MusicClock")
+	UE_API TOptional<FVector> TryGetAudioSourceLocation() const;
+
+	// ---- Events ----
+
+	UPROPERTY(BlueprintAssignable, Category = "MusicClock")
+	FMusicClockBeatEvent BeatEvent;
+
+	UPROPERTY(BlueprintAssignable, Category = "MusicClock")
+	FMusicClockBarEvent BarEvent;
+
+	UPROPERTY(BlueprintAssignable, Category = "MusicClock")
+	FMusicClockSectionEvent SectionEvent;
+
+	UPROPERTY(BlueprintAssignable, Category = "MusicClock")
+	FMusicClockPlayStateEvent PlayStateEvent;
 
 #if !UE_BUILD_SHIPPING
-	// Short user readable description good for logging. Should include name of song or sound source or whatnot. 
-	virtual FString GetDisplayName() const = 0;
+	UE_API FString GetDisplayName() const;
 #endif
 
-	FSongMaps DefaultMaps;
-	
-	// Note that tempo is MidiTempo, quarter notes per minute.
-	float Tempo = 0.0f;
-	int TimeSignatureNum = 0;
-	int TimeSignatureDenom = 0;
+protected:
 
-	float CurrentBeatsPerSecond = 0.0f;
-	float CurrentBarsPerSecond = 0.0f;
-	float CurrentClockAdvanceRate = 1.0f;
+	/** Weak reference to the source UObject. Does not prevent GC. */
+	TWeakObjectPtr<UObject> SourceObject;
 
-	FMidiSongPos CurrentSongPos[(int32)ECalibratedMusicTimebase::Count];
-	FMidiSongPos PrevSongPos[(int32)ECalibratedMusicTimebase::Count];
+	/** Cached interface pointer — valid only when SourceObject is valid. */
+	IMusicSource* SourceInterfaceCache = nullptr;
 
-	// TODO shouldn't these deltas be rolled into the MidiSongPos?
-	float DeltaBarF[(int32)ECalibratedMusicTimebase::Count];
-	float DeltaBeatF[(int32)ECalibratedMusicTimebase::Count];
+	/** Get the source if still alive, or nullptr. */
+	IMusicSource* GetValidSource() const;
 
-protected:
 
-	// TODO: Remove Transport Controls from the Clock
-	virtual void OnStart() = 0;
-	virtual void OnPause() = 0;
-	virtual void OnContinue() = 0;
-	virtual void OnStop() = 0;
-	
-	UE_API int32 GetArrayIndexForTimebase(ECalibratedMusicTimebase Timebase) const;
+	FMidiSongPos CurrentPos;
+	FMidiSongPos PrevPos;
 
-	UE_API void SetCurrentSongPos(ECalibratedMusicTimebase Timebase, const FMidiSongPos& InSongPos);
-	UE_API void SetCurrentSongPosByTime(ECalibratedMusicTimebase Timebase, float InElapsedMs, const ISongMapEvaluator& Map);
-	UE_API void SetCurrentSongPosByTick(ECalibratedMusicTimebase Timebase, float InTick, const ISongMapEvaluator& Map);
+	float DeltaBarF = 0.f;
+	float DeltaBeatF = 0.f;
 
-private:
+	float Tempo = 0.f;
+	int32 TimeSignatureNum = 0;
+	int32 TimeSignatureDenom = 0;
+	float CurrentBeatsPerSecondCached = 0.f;
+	float CurrentBarsPerSecondCached = 0.f;
+	float ClockAdvanceRate = 1.f;
 
-	UE_API void ResetDefaultSongMaps(float DefaultTempo, int32 DefaultTimeSigNum, int32 DefaultTimeSigDenom);
-	virtual bool RefreshCurrentSongPos() = 0;
-	UE_API void UpdateBeatsAndBarsPerSecond();
+	// Event broadcasting state
+	int32 LastBroadcastBar = -1;
+	int32 LastBroadcastBeat = -1;
+	FSongSection LastBroadcastSection;
 
 	uint64 LastUpdateFrame = 0;
 
-	EMusicClockState MusicClockState = EMusicClockState::Stopped;
+	EMusicClockState CachedClockState = EMusicClockState::Stopped;
+
+	bool bLoopedThisFrame = false;
+	bool bSeekedThisFrame = false;
+
+	/** Update cached tempo, time signature, and derived rates when the song position changes. */
+	UE_API void UpdatePlaybackRate(const FMidiSongPos& SongPos, float InClockAdvanceRate);
 
-	TWeakObjectPtr<UObject> ContextObj;
+	/** Broadcast beat, bar, and section events if thresholds were crossed this frame. */
+	UE_API void BroadcastSongPosChanges();
 };
-}
 
-#undef UE_API
\ No newline at end of file
+#undef UE_API
diff --git a/Engine/Plugins/Runtime/Harmonix/Source/HarmonixMetasound/Public/HarmonixMetasound/MusicClock/TimeSource.h b/Engine/Plugins/Runtime/Harmonix/Source/HarmonixMetasound/Public/HarmonixMetasound/MusicClock/TimeSource.h
--- a/Engine/Plugins/Runtime/Harmonix/Source/HarmonixMetasound/Public/HarmonixMetasound/MusicClock/TimeSource.h
+++ b/Engine/Plugins/Runtime/Harmonix/Source/HarmonixMetasound/Public/HarmonixMetasound/MusicClock/TimeSource.h
@@ -1,7 +1,7 @@
 // Copyright Epic Games, Inc. All Rights Reserved.
 #pragma once
 #include "CoreTypes.h"
-#include "Harmonix/MusicalTimebase.h"
+#include "Math/Vector.h"
 
 #define UE_API HARMONIXMETASOUND_API
 
@@ -10,6 +10,7 @@ namespace Harmonix
 	enum class ESourceState : uint8
 	{
 		Stopped,
+		Preparing, // Intent to run, but no data available yet (e.g., Metasound not yet connected)
 		Paused,
 		Running
 	};
@@ -51,9 +52,6 @@ namespace Harmonix
 		/** @return The location of the Audio Source for this Time Source */
 		virtual TOptional<FVector> TryGetAudioSourceLocation() const = 0;
 
-		/** The Calibrated Timebase this Time Source represents */
-		virtual ECalibratedMusicTimebase GetTimebase() const = 0;
-
 		/**
 		 * Used to update the time source once per frame
 		 *
@@ -88,10 +86,25 @@ namespace Harmonix
 		 * Retrieve the latest "Event" from this source
 		 * Whether it "Advanced", "Stopped", "Continued", etc.
 		 * Also whether it "Seeked" or "Looped"
-		 * 
+		 *
 		 * @return The Latest Event of the Source
 		 */
 		virtual ESourceEvent GetLatestSourceEvent() const = 0;
+
+		// ---- Optional Transport ----
+		// Default implementations are no-ops. Override in time sources
+		// that support direct transport control (e.g., FWorldTimeSourceController).
+
+		virtual void RequestStart(float StartTime = 0.f) {}
+		virtual void RequestStop() {}
+		virtual void RequestPause() {}
+		virtual void RequestContinue() {}
+
+		/** Seek to an absolute time position in seconds. Default is no-op. */
+		virtual void RequestSeek(float TimeInSeconds) {}
+
+		/** Set the playback speed multiplier. Default is no-op. */
+		virtual void RequestSetSpeed(float InSpeed) {}
 	};
 }
 
diff --git a/Engine/Plugins/Runtime/Harmonix/Source/HarmonixMetasound/Public/HarmonixMetasound/MusicClock/WorldTimeSourceController.h b/Engine/Plugins/Runtime/Harmonix/Source/HarmonixMetasound/Public/HarmonixMetasound/MusicClock/WorldTimeSourceController.h
--- a/Engine/Plugins/Runtime/Harmonix/Source/HarmonixMetasound/Public/HarmonixMetasound/MusicClock/WorldTimeSourceController.h
+++ b/Engine/Plugins/Runtime/Harmonix/Source/HarmonixMetasound/Public/HarmonixMetasound/MusicClock/WorldTimeSourceController.h
@@ -7,6 +7,15 @@
 
 namespace Harmonix
 {
+
+/**
+ * A time source driven by UWorld delta time.
+ *
+ * Provides a simple wall-clock time stream suitable for non-audio-driven
+ * use cases. Supports transport control (Start/Stop/Pause/Continue/Seek)
+ * and variable speed playback. Time advances by the world's delta time
+ * each frame, scaled by the speed multiplier and time dilation.
+ */
 class FWorldTimeSourceController : public ITimeSource
 {
 public:
@@ -17,8 +26,6 @@ public:
 
 	virtual TOptional<FVector> TryGetAudioSourceLocation() const override;
 
-	virtual ECalibratedMusicTimebase GetTimebase() const override;
-
 	virtual void Update() override;
 
 	virtual ESourceState GetCurrentSourceState() const override;
@@ -29,16 +36,24 @@ public:
 
 	virtual float GetSpeed() const override;
 
+	/** Set the playback speed multiplier. Clamped to [0.01, 10.0]. */
 	void SetSpeed(float InSpeed);
 
+	/** Start playback from the given time in seconds. */
 	void Start(float StartTime = 0.0f);
-
 	void Stop();
-
 	void Pause();
-
 	void Continue();
 
+	// ITimeSource optional transport
+	virtual void RequestStart(float StartTime = 0.f) override { Start(StartTime); }
+	virtual void RequestStop() override { Stop(); }
+	virtual void RequestPause() override { Pause(); }
+	virtual void RequestContinue() override { Continue(); }
+	virtual void RequestSeek(float TimeInSeconds) override { Seek(TimeInSeconds); }
+	virtual void RequestSetSpeed(float InSpeed) override { SetSpeed(InSpeed); }
+
+	/** Seek to an absolute time position in seconds. */
 	void Seek(float InSeekTime);
 
 private:
diff --git a/Engine/Plugins/Runtime/Harmonix/Source/HarmonixMetasound/Public/HarmonixMetasound/MusicSource/MetasoundMusicSource.h b/Engine/Plugins/Runtime/Harmonix/Source/HarmonixMetasound/Public/HarmonixMetasound/MusicSource/MetasoundMusicSource.h
new file mode 100644
--- /dev/null
+++ b/Engine/Plugins/Runtime/Harmonix/Source/HarmonixMetasound/Public/HarmonixMetasound/MusicSource/MetasoundMusicSource.h
@@ -0,0 +1,218 @@
+// Copyright Epic Games, Inc. All Rights Reserved.
+#pragma once
+
+#include "Harmonix/LocalMinimumMagnitudeTracker.h"
+#include "HarmonixMetasound/MusicSource/MusicSource.h"
+#include "HarmonixMetasound/Subsystems/MidiClockUpdateSubsystem.h"
+#include "HarmonixMetasound/Analysis/MidiClockSongPos.h"
+#include "MetasoundGeneratorHandle.h"
+#include "MetasoundSampleCounter.h"
+#include "UObject/StrongObjectPtr.h"
+
+#include "MetasoundMusicSource.generated.h"
+
+#define UE_API HARMONIXMETASOUND_API
+
+class UAudioComponent;
+
+namespace HarmonixMetasound::Analysis
+{
+	struct FSongMapChain;
+}
+
+namespace Metasound
+{
+	class FMetasoundGenerator;
+}
+
+/**
+ * A music source driven by a Metasound's MIDI clock output.
+ *
+ * Connects to a Metasound via UAudioComponent, reads the MIDI clock
+ * history ring buffer produced by the audio render thread, and smooths
+ * it into a game-thread-rate position suitable for gameplay.
+ *
+ * Produces a single time stream: the smoothed audio render position.
+ * Calibration offsets (experienced time, video render time) should be
+ * applied externally via UOffsetMusicSource.
+ *
+ * State is derived from the Metasound's actual transport state, not from
+ * explicit Start/Stop calls. If the AudioComponent starts playing externally,
+ * the source detects it and reports Running. Transport methods on IMusicSource
+ * are no-ops — control the Metasound through its AudioComponent or Metasound
+ * interfaces directly.
+ *
+ * Reports Preparing when connected to an AudioComponent but the generator
+ * is not yet attached or no clock history is available.
+ */
+UCLASS(MinimalAPI, BlueprintType)
+class UMetasoundMusicSource : public UObject, public IMusicSource
+{
+	GENERATED_BODY()
+
+public:
+
+	UE_API virtual void BeginDestroy() override;
+
+	/**
+	 * Connect to a Metasound on an audio component.
+	 *
+	 * @param InAudioComponent   The audio component playing a MetaSound.
+	 * @param OutputPinName      The name of the MIDI Clock output pin on the MetaSound.
+	 * @return true if connection was initiated (may complete asynchronously).
+	 */
+	UE_API bool ConnectToAudioComponent(
+		UAudioComponent* InAudioComponent,
+		FName OutputPinName = "MIDI Clock");
+
+	/** Disconnect from the Metasound and release all handles. */
+	UE_API void Disconnect();
+
+	// ---- IMusicSource ----
+
+	virtual void Start() override;
+	virtual void Pause() override;
+	virtual void Continue() override;
+	virtual void Stop() override;
+
+	virtual Harmonix::ESourceState GetSourceState() const override;
+	virtual Harmonix::ESourceEvent GetLatestSourceEvent() const override;
+
+	virtual const FMidiSongPos& GetCurrentSongPos() const override;
+	virtual const FMidiSongPos& GetPreviousSongPos() const override;
+	virtual const ISongMapEvaluator* GetCurrentSongMapEvaluator() const override;
+	virtual float GetCurrentClockAdvanceRate() const override;
+
+	virtual bool LoopedThisFrame() const override;
+	virtual bool SeekedThisFrame() const override;
+
+	virtual bool IsLooping() const override;
+	virtual float GetLoopStartMs() const override;
+	virtual float GetLoopLengthMs() const override;
+
+	virtual TOptional<FVector> TryGetAudioSourceLocation() const override;
+
+	virtual void Update() override;
+
+#if !UE_BUILD_SHIPPING
+	virtual FString GetDisplayName() const override;
+#endif
+
+private:
+	// ---- Connection State ----
+
+	FName MetasoundOutputName;
+	TWeakObjectPtr<UAudioComponent> AudioComponentToWatch;
+	TStrongObjectPtr<UMetasoundGeneratorHandle> CurrentGeneratorHandle;
+	TWeakPtr<Metasound::FMetasoundGenerator> UnderlyingGenerator;
+
+	Metasound::Frontend::FAnalyzerAddress MidiSongPosAnalyzerAddress;
+
+	UMetasoundGeneratorHandle::FOnAttached::FDelegate OnAttachedDelegate;
+	UMetasoundGeneratorHandle::FOnDetached::FDelegate OnDetachedDelegate;
+
+	FDelegateHandle GeneratorAttachedCallbackHandle;
+	FDelegateHandle GeneratorDetachedCallbackHandle;
+	FDelegateHandle GeneratorIOUpdatedCallbackHandle;
+	FDelegateHandle GraphChangedCallbackHandle;
+
+	// ---- Clock History (read from audio thread) ----
+
+	UMidiClockUpdateSubsystem::FClockHistoryPtr ClockHistory;
+	HarmonixMetasound::Analysis::FMidiClockSongPositionHistory::FReadCursor SmoothedClockHistoryCursor;
+	TSharedPtr<const HarmonixMetasound::Analysis::FSongMapChain> CurrentMapChain;
+
+	// ---- Smoothing / PLL State ----
+
+	struct FSmoothedClockState
+	{
+		float TempoMapMs = 0.f;
+		float TempoMapTick = 0.f;
+		float LocalTick = 0.f;
+	};
+
+	FSmoothedClockState SmoothedState;
+
+	double RenderStartWallClockTimeSeconds = 0.0;
+	double LastRefreshWallClockTimeSeconds = 0.0;
+	double DeltaSecondsBetweenRefreshes = 0.0;
+	Metasound::FSampleCount RenderStartSampleCount {0};
+
+	static const int kFramesOfErrorHistory = 10;
+	FLocalMinimumMagnitudeTracker<double, kFramesOfErrorHistory> ErrorTracker;
+	double SyncSpeed = 1.0;
+
+	float RenderSmoothingLagSeconds = 0.030f;
+
+	// ---- Position State ----
+
+	FMidiSongPos CurrentPos;
+	FMidiSongPos PrevPos;
+
+	float CurrentAdvanceRate = 1.f;
+	int32 LastTickSeen = 0;
+
+	/** Cached source state, derived from clock history transport each frame. */
+	Harmonix::ESourceState CachedSourceState = Harmonix::ESourceState::Stopped;
+
+	bool bSeekDetected = false;
+	bool bLoopDetected = false;
+
+	// ---- Internal Methods ----
+
+	/** Try to create a generator handle and wire up analyzer callbacks for the watched audio component. */
+	UE_API bool AttemptToConnect(UAudioComponent* InAudioComponent);
+
+	/** Remove all generator handle callbacks. */
+	UE_API void DetachAllCallbacks();
+
+	// Generator lifecycle callbacks
+	void OnGeneratorAttached();
+	void OnGeneratorDetached();
+	void OnGraphSet();
+	void OnGeneratorIOUpdatedWithChanges(const TArray<Metasound::FVertexInterfaceChange>& VertexInterfaceChanges);
+
+	/** Set (or clear) the shared clock history ring buffer used for position smoothing. */
+	UE_API void SetClockHistory(const UMidiClockUpdateSubsystem::FClockHistoryPtr& NewHistory);
+
+	/** Read the latest position from the clock history ring and apply PLL smoothing. */
+	UE_API void RefreshFromHistory();
+
+	/** Apply the smoothed tick to CurrentPos for non-offset (normal or looping) playback. */
+	UE_API void UpdateSmoothedPosition(float SmoothedTick, float SmoothedTempoMapTick);
+
+	/** Apply the smoothed tick to CurrentPos for offset clock playback (clock tick != tempo map tick). */
+	UE_API void UpdateSmoothedPositionForOffsetClock(float SmoothedTick, float SmoothedTempoMapTick);
+
+	/** Calculate song position with loop wrapping or monotonic advancement. */
+	UE_API FMidiSongPos CalculateSongPosForLoopingOrMonotonic(float AbsoluteMs, float& InOutPositionTick, bool& OutSeekDetected, bool& OutLoopDetected) const;
+
+	/** Calculate song position for an offset clock with a tick delta from the tempo map. */
+	UE_API FMidiSongPos CalculateSongPosForOffsetClock(float PositionMs, float ClockTickOffset, float& InOutPositionTick, bool& OutSeekDetected) const;
+
+	/** Detect a seek by comparing actual tick delta to expected delta based on tempo and frame time. */
+	UE_API bool CheckForSeek(float FirstTick, float NextTick, float CurrentTempo, int32 TicksPerQuarter) const;
+
+	enum class EHistoryFailureType : uint8
+	{
+		None,
+		NotEnoughDataInTheHistoryRing,
+		NotEnoughHistory,
+		LookingForTimeInTheFutureOfWhatHasEvenRendered,
+		CaughtUpToRenderPosition
+	};
+
+	/** Interpolate the clock history ring buffer to find the smoothed tick at the expected render position. */
+	UE_API EHistoryFailureType CalculateSmoothedTick(
+		Metasound::FSampleCount ExpectedRenderPosSampleCount,
+		Metasound::FSampleCount LastRenderPosSampleCount,
+		float& SmoothedLocalTick,
+		float& SmoothedTempoMapTick,
+		float& CurrentSpeed,
+		HarmonixMetasound::Analysis::FMidiClockSongPositionHistory::FReadCursor& ReadCursor,
+		float LookBehindSeconds);
+
+	static FString HistoryFailureTypeToString(EHistoryFailureType Error);
+};
+
+#undef UE_API
diff --git a/Engine/Plugins/Runtime/Harmonix/Source/HarmonixMetasound/Public/HarmonixMetasound/MusicSource/MusicSource.h b/Engine/Plugins/Runtime/Harmonix/Source/HarmonixMetasound/Public/HarmonixMetasound/MusicSource/MusicSource.h
new file mode 100644
--- /dev/null
+++ b/Engine/Plugins/Runtime/Harmonix/Source/HarmonixMetasound/Public/HarmonixMetasound/MusicSource/MusicSource.h
@@ -0,0 +1,160 @@
+// Copyright Epic Games, Inc. All Rights Reserved.
+#pragma once
+
+#include "UObject/Interface.h"
+#include "HarmonixMetasound/MusicClock/TimeSource.h"
+#include "HarmonixMidi/MidiSongPos.h"
+#include "HarmonixMidi/SongMaps.h"
+
+#include "MusicSource.generated.h"
+
+#define UE_API HARMONIXMETASOUND_API
+
+UINTERFACE(MinimalAPI, BlueprintType, meta = (CannotImplementInterfaceInBlueprint))
+class UMusicSource : public UInterface
+{
+	GENERATED_BODY()
+};
+
+/**
+ * Interface for objects that produce musical time.
+ *
+ * An IMusicSource is the owner of a music playback stream. It provides:
+ *   - Transport controls (Start/Stop/Pause/Continue)
+ *   - Current musical position (FMidiSongPos)
+ *   - Song map access for time<->beat conversions
+ *   - Discontinuity detection (loop/seek)
+ *
+ * A single IMusicSource produces a single time stream. Calibrated time
+ * offsets (experienced, video render, etc.) are handled externally by
+ * UOffsetMusicSource feeding into a UMusicClock.
+ */
+class IMusicSource
+{
+	GENERATED_BODY()
+
+public:
+
+	// ---- Transport Controls ----
+	// The source owns playback. Consumers (UMusicClock) are read-only views.
+
+	UFUNCTION(BlueprintCallable, Category = "MusicSource")
+	virtual void Start() = 0;
+
+	UFUNCTION(BlueprintCallable, Category = "MusicSource")
+	virtual void Pause() = 0;
+
+	UFUNCTION(BlueprintCallable, Category = "MusicSource")
+	virtual void Continue() = 0;
+
+	UFUNCTION(BlueprintCallable, Category = "MusicSource")
+	virtual void Stop() = 0;
+
+	/**
+	 * Seek to a musical position.
+	 * For RuntimeMusicSource: converts the timestamp to seconds via song maps and seeks the time source.
+	 * For MetasoundMusicSource: requires Metasound interface support (future work).
+	 */
+	UFUNCTION(BlueprintCallable, Category = "MusicSource")
+	virtual void Seek(const FMusicTimestamp& Timestamp) {}
+
+	// ---- Source State ----
+
+	/** Get the current transport state of the source (Stopped, Preparing, Paused, Running). */
+	virtual Harmonix::ESourceState GetSourceState() const = 0;
+
+	/** Get the most recent source event for this frame (Start, Advance, Stop, Pause, Continue, Seek, Loop). */
+	virtual Harmonix::ESourceEvent GetLatestSourceEvent() const = 0;
+
+	/** Whether the source is currently producing musical time. */
+	UFUNCTION(BlueprintCallable, Category = "MusicSource")
+	virtual bool IsPlaying() const { return GetSourceState() == Harmonix::ESourceState::Running; }
+
+	/** Whether the source is stopped (not producing time, position at zero). */
+	UFUNCTION(BlueprintCallable, Category = "MusicSource")
+	virtual bool IsStopped() const { return GetSourceState() == Harmonix::ESourceState::Stopped; }
+
+	/** Whether the source intends to run but is waiting for data (e.g., Metasound not yet connected). */
+	UFUNCTION(BlueprintCallable, Category = "MusicSource")
+	virtual bool IsPreparing() const { return GetSourceState() == Harmonix::ESourceState::Preparing; }
+
+	// ---- Time Production ----
+
+	/** The current musical position, updated once per frame via Update(). */
+	virtual const FMidiSongPos& GetCurrentSongPos() const = 0;
+
+	/** The previous frame's musical position. */
+	virtual const FMidiSongPos& GetPreviousSongPos() const = 0;
+
+	/** The song maps used to convert between time and musical position. */
+	virtual const ISongMapEvaluator* GetCurrentSongMapEvaluator() const = 0;
+
+	/** The current playback speed multiplier. */
+	virtual float GetCurrentClockAdvanceRate() const = 0;
+
+	// ---- Discontinuity Detection ----
+
+	/** Whether the source crossed a loop boundary this frame. */
+	virtual bool LoopedThisFrame() const = 0;
+
+	/** Whether the source performed a non-contiguous time jump (seek) this frame. */
+	virtual bool SeekedThisFrame() const = 0;
+
+	// ---- Loop Region ----
+
+	/** Whether this source is currently looping. */
+	virtual bool IsLooping() const = 0;
+
+	/** Start of the loop region in milliseconds. Only valid if IsLooping() is true. */
+	virtual float GetLoopStartMs() const = 0;
+
+	/** Length of the loop region in milliseconds. Only valid if IsLooping() is true. */
+	virtual float GetLoopLengthMs() const = 0;
+
+	// ---- Tempo & Time Signature Control ----
+
+	/**
+	 * Set the tempo in BPM.
+	 * For RuntimeMusicSource: modifies the underlying song maps directly.
+	 * For MetasoundMusicSource: requires Metasound interface support (future work).
+	 */
+	UFUNCTION(BlueprintCallable, Category = "MusicSource")
+	virtual void SetTempo(float BPM) {}
+
+	/**
+	 * Set the time signature.
+	 * For RuntimeMusicSource: modifies the underlying song maps directly.
+	 * For MetasoundMusicSource: requires Metasound interface support (future work).
+	 */
+	UFUNCTION(BlueprintCallable, Category = "MusicSource")
+	virtual void SetTimeSignature(int32 Numerator, int32 Denominator) {}
+
+	/**
+	 * Set the playback speed multiplier (1.0 = normal speed).
+	 * For RuntimeMusicSource: sets the speed on the underlying ITimeSource.
+	 * For MetasoundMusicSource: requires Metasound interface support (future work).
+	 */
+	UFUNCTION(BlueprintCallable, Category = "MusicSource")
+	virtual void SetSpeed(float Speed) {}
+
+	// ---- Spatial ----
+
+	/** Get the world-space location of the audio source, if available. */
+	virtual TOptional<FVector> TryGetAudioSourceLocation() const = 0;
+
+	// ---- Update ----
+
+	/**
+	 * Called once per frame by the update subsystem.
+	 * Implementations should refresh their internal position state.
+	 */
+	virtual void Update() = 0;
+
+	// ---- Debug ----
+
+#if !UE_BUILD_SHIPPING
+	virtual FString GetDisplayName() const = 0;
+#endif
+};
+
+#undef UE_API
diff --git a/Engine/Plugins/Runtime/Harmonix/Source/HarmonixMetasound/Public/HarmonixMetasound/MusicSource/MusicSourceBlueprintLibrary.h b/Engine/Plugins/Runtime/Harmonix/Source/HarmonixMetasound/Public/HarmonixMetasound/MusicSource/MusicSourceBlueprintLibrary.h
new file mode 100644
--- /dev/null
+++ b/Engine/Plugins/Runtime/Harmonix/Source/HarmonixMetasound/Public/HarmonixMetasound/MusicSource/MusicSourceBlueprintLibrary.h
@@ -0,0 +1,61 @@
+// Copyright Epic Games, Inc. All Rights Reserved.
+#pragma once
+
+#include "Kismet/BlueprintFunctionLibrary.h"
+#include "HarmonixMetasound/MusicSource/MusicSourceSettings.h"
+
+#include "MusicSourceBlueprintLibrary.generated.h"
+
+#define UE_API HARMONIXMETASOUND_API
+
+class UMusicClock;
+
+/**
+ * Blueprint function library for creating and managing music sources and clocks.
+ *
+ * These factory functions handle creation, connection, and subsystem registration
+ * in one call. The source type is determined by the settings struct — new source
+ * types can be added by game plugins without modifying this library.
+ */
+UCLASS(MinimalAPI)
+class UMusicSourceBlueprintLibrary : public UBlueprintFunctionLibrary
+{
+	GENERATED_BODY()
+
+public:
+
+	/**
+	 * Create a music source from settings.
+	 *
+	 * The settings struct type determines what kind of source is created.
+	 * The source is automatically registered with the update subsystem.
+	 *
+	 * @param Outer             Outer object for the new source (determines GC lifetime).
+	 * @param Settings          Configuration determining source type and parameters.
+	 * @return The created music source, or invalid if creation failed.
+	 */
+	UFUNCTION(BlueprintCallable, Category = "Music|Source",
+		Meta = (DefaultToSelf = "Outer"))
+	static UE_API TScriptInterface<IMusicSource> CreateMusicSource(
+		UObject* Outer,
+		const TInstancedStruct<FMusicSourceSettings>& Settings);
+
+	/**
+	 * Create a read-only music clock bound to a source.
+	 *
+	 * The clock holds a weak reference to the source — if the source is
+	 * destroyed, the clock gracefully reports Stopped.
+	 * The clock is automatically registered with the update subsystem.
+	 *
+	 * @param Outer             Outer object for the new clock (determines GC lifetime).
+	 * @param Source            The music source to observe.
+	 * @return The created music clock, or null if Source is invalid.
+	 */
+	UFUNCTION(BlueprintCallable, Category = "Music|Clock",
+		Meta = (DefaultToSelf = "Outer"))
+	static UE_API UMusicClock* CreateMusicClock(
+		UObject* Outer,
+		const TScriptInterface<IMusicSource>& Source);
+};
+
+#undef UE_API
diff --git a/Engine/Plugins/Runtime/Harmonix/Source/HarmonixMetasound/Public/HarmonixMetasound/MusicSource/MusicSourceSettings.h b/Engine/Plugins/Runtime/Harmonix/Source/HarmonixMetasound/Public/HarmonixMetasound/MusicSource/MusicSourceSettings.h
new file mode 100644
--- /dev/null
+++ b/Engine/Plugins/Runtime/Harmonix/Source/HarmonixMetasound/Public/HarmonixMetasound/MusicSource/MusicSourceSettings.h
@@ -0,0 +1,114 @@
+// Copyright Epic Games, Inc. All Rights Reserved.
+#pragma once
+
+#include "HarmonixMetasound/MusicSource/MusicSource.h"
+#include "StructUtils/InstancedStruct.h"
+
+#include "MusicSourceSettings.generated.h"
+
+#define UE_API HARMONIXMETASOUND_API
+
+class UAudioComponent;
+class UMidiFile;
+
+/**
+ * Base struct for music source configuration.
+ *
+ * Subclass this to define new music source types. The CreateSource() factory
+ * method is called by UMusicSourceBlueprintLibrary::CreateMusicSource() —
+ * Blueprint users pick the settings type from a TInstancedStruct dropdown,
+ * fill in the type-specific fields, and get back an IMusicSource.
+ *
+ * Game plugins can add new source types without modifying engine code:
+ * just subclass FMusicSourceSettings and implement CreateSource().
+ */
+USTRUCT(MinimalAPI, BlueprintType)
+struct FMusicSourceSettings
+{
+	GENERATED_BODY()
+
+	virtual ~FMusicSourceSettings() = default;
+
+	/** Create the music source described by these settings and register it with the update subsystem. */
+	UE_API virtual TScriptInterface<IMusicSource> CreateSource(UObject* Outer) const;
+};
+
+/**
+ * Creates a UMetasoundMusicSource connected to an AudioComponent.
+ */
+USTRUCT(BlueprintType, DisplayName = "Metasound Music Source Settings")
+struct FMetasoundMusicSourceSettings : public FMusicSourceSettings
+{
+	GENERATED_BODY()
+
+	/** The audio component playing a MetaSound. */
+	UPROPERTY(EditAnywhere, BlueprintReadWrite, Category = "MusicSource")
+	TObjectPtr<UAudioComponent> AudioComponent;
+
+	/** The name of the MIDI Clock output pin on the MetaSound. */
+	UPROPERTY(EditAnywhere, BlueprintReadWrite, Category = "MusicSource")
+	FName OutputPinName = "MIDI Clock";
+
+	UE_API virtual TScriptInterface<IMusicSource> CreateSource(UObject* Outer) const override;
+};
+
+/**
+ * Creates a URuntimeMusicSource driven by wall clock with tempo/time-sig from a MIDI file.
+ */
+USTRUCT(BlueprintType, DisplayName = "MIDI Music Source Settings")
+struct FMidiMusicSourceSettings : public FMusicSourceSettings
+{
+	GENERATED_BODY()
+
+	/** MIDI file providing tempo map and time signature data. */
+	UPROPERTY(EditAnywhere, BlueprintReadWrite, Category = "MusicSource")
+	TObjectPtr<UMidiFile> MidiFile;
+
+	UE_API virtual TScriptInterface<IMusicSource> CreateSource(UObject* Outer) const override;
+};
+
+/**
+ * Creates a URuntimeMusicSource driven by wall clock with manually specified tempo and time signature.
+ * Tempo and time signature can be changed at runtime via SetTempo() / SetTimeSignature() on the source.
+ */
+USTRUCT(BlueprintType, DisplayName = "Manual Music Source Settings")
+struct FManualMusicSourceSettings : public FMusicSourceSettings
+{
+	GENERATED_BODY()
+
+	/** Initial tempo in BPM. */
+	UPROPERTY(EditAnywhere, BlueprintReadWrite, Category = "MusicSource", Meta = (ClampMin = "1.0", ClampMax = "999.0"))
+	float Tempo = 120.f;
+
+	/** Initial time signature numerator (beats per bar). */
+	UPROPERTY(EditAnywhere, BlueprintReadWrite, Category = "MusicSource", Meta = (ClampMin = "1", ClampMax = "32"))
+	int32 TimeSigNumerator = 4;
+
+	/** Initial time signature denominator (beat unit). */
+	UPROPERTY(EditAnywhere, BlueprintReadWrite, Category = "MusicSource", Meta = (ClampMin = "1", ClampMax = "32"))
+	int32 TimeSigDenominator = 4;
+
+	UE_API virtual TScriptInterface<IMusicSource> CreateSource(UObject* Outer) const override;
+};
+
+/**
+ * Creates a UOffsetMusicSource that applies a millisecond offset to a parent source.
+ * Used for calibrated time streams (e.g., video render, experienced time).
+ */
+USTRUCT(BlueprintType, DisplayName = "Offset Music Source Settings")
+struct FOffsetMusicSourceSettings : public FMusicSourceSettings
+{
+	GENERATED_BODY()
+
+	/** The parent music source to offset from. */
+	UPROPERTY(EditAnywhere, BlueprintReadWrite, Category = "MusicSource")
+	TScriptInterface<IMusicSource> ParentSource;
+
+	/** Offset in milliseconds. Positive = ahead, negative = behind. */
+	UPROPERTY(EditAnywhere, BlueprintReadWrite, Category = "MusicSource")
+	float OffsetMs = 0.f;
+
+	UE_API virtual TScriptInterface<IMusicSource> CreateSource(UObject* Outer) const override;
+};
+
+#undef UE_API
diff --git a/Engine/Plugins/Runtime/Harmonix/Source/HarmonixMetasound/Public/HarmonixMetasound/MusicSource/OffsetMusicSource.h b/Engine/Plugins/Runtime/Harmonix/Source/HarmonixMetasound/Public/HarmonixMetasound/MusicSource/OffsetMusicSource.h
new file mode 100644
--- /dev/null
+++ b/Engine/Plugins/Runtime/Harmonix/Source/HarmonixMetasound/Public/HarmonixMetasound/MusicSource/OffsetMusicSource.h
@@ -0,0 +1,104 @@
+// Copyright Epic Games, Inc. All Rights Reserved.
+#pragma once
+
+#include "HarmonixMetasound/MusicSource/MusicSource.h"
+
+#include "OffsetMusicSource.generated.h"
+
+#define UE_API HARMONIXMETASOUND_API
+
+/**
+ * An IMusicSource that reads from a parent IMusicSource and applies a
+ * millisecond offset. This is how calibrated time clocks (experienced,
+ * video render) are created:
+ *
+ *   UMetasoundMusicSource (audio render)
+ *       -> UOffsetMusicSource (-30ms) -> UMusicClock (video render)
+ *       -> UOffsetMusicSource (-50ms) -> UMusicClock (experienced)
+ *
+ * Offset sources can also chain: SourceA -> OffsetB -> OffsetC, where
+ * each offset reads directly from the previous IMusicSource.
+ *
+ * Loop handling: when the parent source is looping, the offset time is
+ * wrapped around the loop boundary. On the first pass through the loop
+ * (before the parent has ever looped), negative offsets are NOT wrapped —
+ * the offset time stays negative until the parent advances past the threshold.
+ */
+UCLASS(MinimalAPI, BlueprintType)
+class UOffsetMusicSource : public UObject, public IMusicSource
+{
+	GENERATED_BODY()
+
+public:
+
+	/** Set the parent source this offset reads from. */
+	UFUNCTION(BlueprintCallable, Category = "MusicSource|Offset")
+	UE_API void SetParentSource(const TScriptInterface<IMusicSource>& InParentSource);
+
+	/** Get the parent source. */
+	UFUNCTION(BlueprintPure, Category = "MusicSource|Offset")
+	UE_API TScriptInterface<IMusicSource> GetParentSource() const;
+
+	/** Set the offset in milliseconds. Positive = ahead, negative = behind. */
+	UFUNCTION(BlueprintCallable, Category = "MusicSource|Offset")
+	UE_API void SetOffsetMs(float InOffsetMs);
+
+	/** Get the current offset. */
+	UFUNCTION(BlueprintPure, Category = "MusicSource|Offset")
+	UE_API float GetOffsetMs() const;
+
+	// ---- IMusicSource ----
+
+	// Transport is a no-op — the parent source owns playback.
+	virtual void Start() override {}
+	virtual void Pause() override {}
+	virtual void Continue() override {}
+	virtual void Stop() override {}
+
+	virtual Harmonix::ESourceState GetSourceState() const override;
+	virtual Harmonix::ESourceEvent GetLatestSourceEvent() const override;
+
+	virtual const FMidiSongPos& GetCurrentSongPos() const override;
+	virtual const FMidiSongPos& GetPreviousSongPos() const override;
+	virtual const ISongMapEvaluator* GetCurrentSongMapEvaluator() const override;
+	virtual float GetCurrentClockAdvanceRate() const override;
+
+	virtual bool LoopedThisFrame() const override;
+	virtual bool SeekedThisFrame() const override;
+
+	virtual bool IsLooping() const override;
+	virtual float GetLoopStartMs() const override;
+	virtual float GetLoopLengthMs() const override;
+
+	virtual TOptional<FVector> TryGetAudioSourceLocation() const override;
+
+	virtual void Update() override;
+
+#if !UE_BUILD_SHIPPING
+	virtual FString GetDisplayName() const override;
+#endif
+
+private:
+
+	/** Weak reference to the parent source UObject. */
+	TWeakObjectPtr<UObject> ParentSourceObject;
+
+	/** Cached interface pointer — valid only when ParentSourceObject is valid. */
+	IMusicSource* ParentSourceCache = nullptr;
+
+	/** Get the parent source if still alive, or nullptr. */
+	IMusicSource* GetValidParentSource() const;
+
+	float OffsetMs = 0.f;
+
+	FMidiSongPos CurrentPos;
+	FMidiSongPos PrevPos;
+	bool bLoopedThisFrame = false;
+	bool bSeekedThisFrame = false;
+
+	/** True once the parent source has completed at least one loop pass.
+	 *  Prevents negative offsets from wrapping to the loop end on the first pass. */
+	bool bParentHasLooped = false;
+};
+
+#undef UE_API
diff --git a/Engine/Plugins/Runtime/Harmonix/Source/HarmonixMetasound/Public/HarmonixMetasound/MusicSource/RuntimeMusicSource.h b/Engine/Plugins/Runtime/Harmonix/Source/HarmonixMetasound/Public/HarmonixMetasound/MusicSource/RuntimeMusicSource.h
new file mode 100644
--- /dev/null
+++ b/Engine/Plugins/Runtime/Harmonix/Source/HarmonixMetasound/Public/HarmonixMetasound/MusicSource/RuntimeMusicSource.h
@@ -0,0 +1,125 @@
+// Copyright Epic Games, Inc. All Rights Reserved.
+#pragma once
+
+#include "HarmonixMetasound/MusicSource/MusicSource.h"
+#include "HarmonixMetasound/MusicClock/TimeSource.h"
+#include "HarmonixMidi/MidiSongPos.h"
+#include "HarmonixMidi/SongMaps.h"
+#include "HarmonixMidi/MidiFile.h"
+
+#include "RuntimeMusicSource.generated.h"
+
+#define UE_API HARMONIXMETASOUND_API
+
+/**
+ * A music source that combines an ITimeSource (providing raw time in seconds)
+ * with an ISongMapEvaluator (providing time-to-musical-position mapping).
+ *
+ * This is the "simple" music source for non-Metasound use cases:
+ * wall clock, animation-driven, media-driven, etc. The specific time
+ * behavior is determined by which ITimeSource implementation is provided.
+ *
+ * Transport controls delegate to the ITimeSource if it supports them
+ * (e.g., FWorldTimeSourceController has Start/Stop/Pause/Continue).
+ */
+UCLASS(MinimalAPI, BlueprintType)
+class URuntimeMusicSource : public UObject, public IMusicSource
+{
+	GENERATED_BODY()
+
+public:
+
+	UE_API virtual void BeginDestroy() override;
+
+	/**
+	 * Initialize this source with a time source and song maps.
+	 *
+	 * @param InTimeSource   Provides raw elapsed time in seconds.
+	 * @param InSongMaps     Provides time<->musical position mapping. Can be null for default maps.
+	 */
+	UE_API void Initialize(TSharedPtr<Harmonix::ITimeSource> InTimeSource, const ISongMapEvaluator* InSongMaps);
+
+	/**
+	 * Initialize from a MIDI file's song maps.
+	 */
+	UE_API void InitializeWithMidi(TSharedPtr<Harmonix::ITimeSource> InTimeSource, UMidiFile* InMidiFile);
+
+	/** Get the underlying time source. */
+	TSharedPtr<Harmonix::ITimeSource> GetTimeSource() const { return TimeSource; }
+
+	// ---- IMusicSource ----
+
+	virtual void Start() override;
+	virtual void Pause() override;
+	virtual void Continue() override;
+	virtual void Stop() override;
+
+	virtual Harmonix::ESourceState GetSourceState() const override;
+	virtual Harmonix::ESourceEvent GetLatestSourceEvent() const override;
+
+	virtual const FMidiSongPos& GetCurrentSongPos() const override;
+	virtual const FMidiSongPos& GetPreviousSongPos() const override;
+	virtual const ISongMapEvaluator* GetCurrentSongMapEvaluator() const override;
+	virtual float GetCurrentClockAdvanceRate() const override;
+
+	virtual bool LoopedThisFrame() const override;
+	virtual bool SeekedThisFrame() const override;
+
+	virtual bool IsLooping() const override;
+	virtual float GetLoopStartMs() const override;
+	virtual float GetLoopLengthMs() const override;
+
+	/**
+	 * Set a loop region by musical time.
+	 * @param StartBar   The bar to start looping at (1-based, where Bar 1 = start of song).
+	 * @param NumBars    Number of bars in the loop. Pass 0 to disable looping.
+	 */
+	UFUNCTION(BlueprintCallable, Category = "MusicSource|Loop")
+	UE_API void SetLoopRegionByBars(int32 StartBar, int32 NumBars);
+
+	/** Clear the loop region, disabling looping. */
+	UFUNCTION(BlueprintCallable, Category = "MusicSource|Loop")
+	UE_API void ClearLoopRegion();
+
+	virtual void Seek(const FMusicTimestamp& Timestamp) override;
+	virtual void SetTempo(float BPM) override;
+	virtual void SetTimeSignature(int32 Numerator, int32 Denominator) override;
+	virtual void SetSpeed(float Speed) override;
+
+	virtual TOptional<FVector> TryGetAudioSourceLocation() const override;
+
+	virtual void Update() override;
+
+#if !UE_BUILD_SHIPPING
+	virtual FString GetDisplayName() const override;
+#endif
+
+private:
+
+	TSharedPtr<Harmonix::ITimeSource> TimeSource;
+
+	/** External song maps (from MIDI file). */
+	TWeakObjectPtr<UMidiFile> MidiFile;
+
+	/** Default maps used when no external maps are provided. */
+	FSongMaps DefaultMaps;
+
+	FMidiSongPos CurrentPos;
+	FMidiSongPos PrevPos;
+
+	// Loop region defined in musical time
+	int32 LoopStartBar = 0;
+	int32 LoopNumBars = 0;
+
+	// Cached ms values, recomputed from bar values + song maps
+	float CachedLoopStartMs = 0.f;
+	float CachedLoopLengthMs = 0.f;
+
+	/** Recompute cached loop start/length in ms from bar values + current song maps. */
+	UE_API void RecomputeLoopRegionMs();
+
+	/** Get the active song map evaluator (MIDI file maps or DefaultMaps). */
+	UE_API const ISongMapEvaluator* GetMaps() const;
+};
+
+#undef UE_API
diff --git a/Engine/Plugins/Runtime/Harmonix/Source/HarmonixMetasound/Public/HarmonixMetasound/Subsystems/MidiClockUpdateSubsystem.h b/Engine/Plugins/Runtime/Harmonix/Source/HarmonixMetasound/Public/HarmonixMetasound/Subsystems/MidiClockUpdateSubsystem.h
--- a/Engine/Plugins/Runtime/Harmonix/Source/HarmonixMetasound/Public/HarmonixMetasound/Subsystems/MidiClockUpdateSubsystem.h
+++ b/Engine/Plugins/Runtime/Harmonix/Source/HarmonixMetasound/Public/HarmonixMetasound/Subsystems/MidiClockUpdateSubsystem.h
@@ -15,6 +15,8 @@
 #define UE_API HARMONIXMETASOUND_API
 
 class UMusicClockComponent;
+class UMusicClock;
+class IMusicSource;
 
 namespace HarmonixMetasound
 {
@@ -47,6 +49,14 @@ public:
 	static UE_API void TrackMusicClockComponent(UMusicClockComponent* Clock);
 	static UE_API void StopTrackingMusicClockComponent(UMusicClockComponent* Clock);
 
+	/** Track a music source for per-frame updates. The source's Update() will be called each frame. */
+	static UE_API void TrackMusicSource(UObject* Source);
+	static UE_API void StopTrackingMusicSource(UObject* Source);
+
+	/** Track a music clock for per-frame updates. The clock's UpdateForGameFrame() will be called each frame. */
+	static UE_API void TrackMusicClock(UMusicClock* Clock);
+	static UE_API void StopTrackingMusicClock(UMusicClock* Clock);
+
 	static UE_API double GetUpdateTimeSeconds();
 
 	using FClockHistoryPtr = TSharedPtr<HarmonixMetasound::Analysis::FMidiClockSongPositionHistory>;
@@ -87,6 +97,18 @@ private:
 	UE_API void StopTrackingMusicClockComponentImpl(UMusicClockComponent* Clock);
 	UE_API void UpdateUMusicClockComponents();
 
+	/** Music sources tracked for per-frame Update() calls. */
+	TArray<TWeakObjectPtr<UObject>> TrackedMusicSources;
+	UE_API void TrackMusicSourceImpl(UObject* Source);
+	UE_API void StopTrackingMusicSourceImpl(UObject* Source);
+	UE_API void UpdateMusicSources();
+
+	/** Music clocks tracked for per-frame UpdateForGameFrame() calls. */
+	TArray<TWeakObjectPtr<UMusicClock>> TrackedMusicClocks;
+	UE_API void TrackMusicClockImpl(UMusicClock* Clock);
+	UE_API void StopTrackingMusicClockImpl(UMusicClock* Clock);
+	UE_API void UpdateMusicClocks();
+
 	FDelegateHandle EngineSamplingInputDelegate;
 	UE_API void CoreDelegatesSamplingInput();
 
diff --git a/Engine/Plugins/Runtime/Harmonix/Source/HarmonixMetasoundTests/Private/MusicClock/MusicClockTests.cpp b/Engine/Plugins/Runtime/Harmonix/Source/HarmonixMetasoundTests/Private/MusicClock/MusicClockTests.cpp
new file mode 100644
--- /dev/null
+++ b/Engine/Plugins/Runtime/Harmonix/Source/HarmonixMetasoundTests/Private/MusicClock/MusicClockTests.cpp
@@ -0,0 +1,613 @@
+// Copyright Epic Games, Inc. All Rights Reserved.
+
+#include "Misc/AutomationTest.h"
+
+#if WITH_DEV_AUTOMATION_TESTS
+
+#include "Engine/Engine.h"
+#include "Engine/World.h"
+#include "HarmonixMetasound/MusicClock/MusicClock.h"
+#include "HarmonixMetasound/MusicClock/TimeSource.h"
+#include "HarmonixMetasound/MusicSource/MusicSource.h"
+#include "HarmonixMetasound/MusicSource/RuntimeMusicSource.h"
+#include "HarmonixMetasound/MusicSource/OffsetMusicSource.h"
+#include "HarmonixMetasound/Subsystems/MidiClockUpdateSubsystem.h"
+#include "HarmonixMetasound/Components/MusicClockComponent.h"
+#include "HarmonixMidi/MidiFile.h"
+#include "Tests/AutomationEditorCommon.h"
+
+namespace HarmonixMetasound::MusicClock::Tests
+{
+	/**
+	 * A manually-driven time source for testing.
+	 * Allows precise control of time without depending on UWorld.
+	 */
+	class FManualTimeSource : public Harmonix::ITimeSource
+	{
+	public:
+		virtual FString GetDisplayName() const override { return TEXT("ManualTimeSource"); }
+		virtual TOptional<FVector> TryGetAudioSourceLocation() const override { return {}; }
+
+		virtual void Update() override
+		{
+			if (State == Harmonix::ESourceState::Running)
+			{
+				CurrentTime += DeltaTimeToApply;
+				LatestEvent = Harmonix::ESourceEvent::Advance;
+			}
+			else
+			{
+				LatestEvent = Harmonix::ESourceEvent::None;
+			}
+			DeltaTimeToApply = 0.0;
+		}
+
+		virtual double GetCurrentTime() const override { return CurrentTime; }
+		virtual float GetSpeed() const override { return Speed; }
+		virtual Harmonix::ESourceState GetCurrentSourceState() const override { return State; }
+		virtual Harmonix::ESourceEvent GetLatestSourceEvent() const override { return LatestEvent; }
+
+		void Start()
+		{
+			CurrentTime = 0.0;
+			State = Harmonix::ESourceState::Running;
+			LatestEvent = Harmonix::ESourceEvent::Start;
+		}
+
+		void Stop()
+		{
+			State = Harmonix::ESourceState::Stopped;
+			LatestEvent = Harmonix::ESourceEvent::Stop;
+			CurrentTime = 0.0;
+		}
+
+		void Pause()
+		{
+			if (State == Harmonix::ESourceState::Running)
+			{
+				State = Harmonix::ESourceState::Paused;
+				LatestEvent = Harmonix::ESourceEvent::Pause;
+			}
+		}
+
+		void Continue()
+		{
+			if (State == Harmonix::ESourceState::Paused)
+			{
+				State = Harmonix::ESourceState::Running;
+				LatestEvent = Harmonix::ESourceEvent::Continue;
+			}
+		}
+
+		// ITimeSource optional transport
+		virtual void RequestStart(float InStartTime = 0.f) override { Start(); }
+		virtual void RequestStop() override { Stop(); }
+		virtual void RequestPause() override { Pause(); }
+		virtual void RequestContinue() override { Continue(); }
+
+		/** Queue a time delta to be applied on the next Update() call. */
+		void AdvanceBy(double Seconds) { DeltaTimeToApply = Seconds; }
+
+		float Speed = 1.0f;
+
+	private:
+		double CurrentTime = 0.0;
+		double DeltaTimeToApply = 0.0;
+		Harmonix::ESourceState State = Harmonix::ESourceState::Stopped;
+		Harmonix::ESourceEvent LatestEvent = Harmonix::ESourceEvent::None;
+	};
+
+	/** Helper to create a source + clock pair with default 120bpm 4/4 maps. */
+	struct FTestClockSetup
+	{
+		TSharedPtr<FManualTimeSource> TimeSource;
+		URuntimeMusicSource* Source = nullptr;
+		UMusicClock* Clock = nullptr;
+
+		void Create(UObject* Outer)
+		{
+			TimeSource = MakeShared<FManualTimeSource>();
+			Source = NewObject<URuntimeMusicSource>(Outer);
+			Source->Initialize(TimeSource, nullptr);
+			Clock = NewObject<UMusicClock>(Outer);
+			Clock->SetSource(TScriptInterface<IMusicSource>(Source));
+		}
+
+		/** Advance time, update source, then update clock. Simulates one frame. */
+		void Tick(double DeltaSeconds)
+		{
+			TimeSource->AdvanceBy(DeltaSeconds);
+			Source->Update();
+			++GFrameCounter;
+			Clock->UpdateForGameFrame();
+		}
+
+		/** Update source and clock without advancing time. */
+		void TickNoAdvance()
+		{
+			Source->Update();
+			++GFrameCounter;
+			Clock->UpdateForGameFrame();
+		}
+
+		void Start()
+		{
+			Source->Start();
+			TickNoAdvance();
+		}
+
+		void Pause()
+		{
+			Source->Pause();
+			TickNoAdvance();
+		}
+
+		void Continue()
+		{
+			Source->Continue();
+		}
+
+		void Stop()
+		{
+			Source->Stop();
+			TickNoAdvance();
+		}
+	};
+
+	// =========================================================================
+	// Test 1: URuntimeMusicSource.BasicUpdate
+	// =========================================================================
+	IMPLEMENT_SIMPLE_AUTOMATION_TEST(
+		FRuntimeMusicSourceBasicTest,
+		"Harmonix.Metasound.MusicClock.RuntimeMusicSource.BasicUpdate",
+		EAutomationTestFlags::EditorContext | EAutomationTestFlags::EngineFilter)
+	bool FRuntimeMusicSourceBasicTest::RunTest(const FString&)
+	{
+		TSharedRef<FManualTimeSource> TimeSource = MakeShared<FManualTimeSource>();
+		URuntimeMusicSource* Source = NewObject<URuntimeMusicSource>();
+		Source->Initialize(TimeSource, nullptr);
+
+		// Not started yet — state should be Stopped
+		UTEST_EQUAL("Initial state is Stopped", Source->GetSourceState(), Harmonix::ESourceState::Stopped);
+
+		TimeSource->Start();
+		TimeSource->AdvanceBy(0.5);
+		Source->Update();
+
+		UTEST_EQUAL("State is Running after start", Source->GetSourceState(), Harmonix::ESourceState::Running);
+		UTEST_GREATER("Position advanced past 0", Source->GetCurrentSongPos().SecondsIncludingCountIn, 0.f);
+
+		float FirstPos = Source->GetCurrentSongPos().SecondsIncludingCountIn;
+
+		TimeSource->AdvanceBy(0.5);
+		Source->Update();
+
+		UTEST_NEARLY_EQUAL("Previous pos holds prior frame's value",
+			Source->GetPreviousSongPos().SecondsIncludingCountIn, FirstPos, 0.01f);
+		UTEST_GREATER("Current pos advanced further",
+			Source->GetCurrentSongPos().SecondsIncludingCountIn, FirstPos);
+
+		return true;
+	}
+
+	// =========================================================================
+	// Test 2: URuntimeMusicSource.WithMidiFile
+	// =========================================================================
+	IMPLEMENT_SIMPLE_AUTOMATION_TEST(
+		FRuntimeMusicSourceMidiFileTest,
+		"Harmonix.Metasound.MusicClock.RuntimeMusicSource.WithMidiFile",
+		EAutomationTestFlags::EditorContext | EAutomationTestFlags::EngineFilter)
+	bool FRuntimeMusicSourceMidiFileTest::RunTest(const FString&)
+	{
+		UMidiFile* MidiFile = NewObject<UMidiFile>();
+		MidiFile->GetSongMaps()->EmptyAllMaps();
+		MidiFile->GetSongMaps()->Init(Harmonix::Midi::Constants::GTicksPerQuarterNoteInt);
+		MidiFile->GetSongMaps()->GetTempoMap().AddTempoInfoPoint(Harmonix::Midi::Constants::BPMToMidiTempo(90.0f), 0);
+		MidiFile->GetSongMaps()->GetBarMap().AddTimeSignatureAtBarIncludingCountIn(0, 3, 4);
+
+		TSharedRef<FManualTimeSource> TimeSource = MakeShared<FManualTimeSource>();
+		URuntimeMusicSource* Source = NewObject<URuntimeMusicSource>();
+		Source->InitializeWithMidi(TimeSource, MidiFile);
+
+		TimeSource->Start();
+		// At 90bpm in 3/4 time: 1 bar = 3 beats = 2 seconds
+		TimeSource->AdvanceBy(2.0);
+		Source->Update();
+
+		const FMidiSongPos& Pos = Source->GetCurrentSongPos();
+		UTEST_GREATER("Position advanced", Pos.SecondsIncludingCountIn, 0.f);
+		UTEST_EQUAL("Time sig numerator is 3", Pos.TimeSigNumerator, 3);
+		UTEST_EQUAL("Time sig denominator is 4", Pos.TimeSigDenominator, 4);
+		UTEST_NEARLY_EQUAL("Approximately 1 bar elapsed", Pos.BarsIncludingCountIn, 1.0f, 0.15f);
+
+		return true;
+	}
+
+	// =========================================================================
+	// Test 3: URuntimeMusicSource.TransportStates
+	// Behavioral contracts from the clock's point of view:
+	//   - Stopped => time is 0
+	//   - Paused  => time is frozen at the paused-at value
+	//   - Continue => resumes from paused time
+	//   - Stop    => resets time to 0
+	// =========================================================================
+	IMPLEMENT_SIMPLE_AUTOMATION_TEST(
+		FRuntimeMusicSourceTransportTest,
+		"Harmonix.Metasound.MusicClock.RuntimeMusicSource.TransportStates",
+		EAutomationTestFlags::EditorContext | EAutomationTestFlags::EngineFilter)
+	bool FRuntimeMusicSourceTransportTest::RunTest(const FString&)
+	{
+		FTestClockSetup Setup;
+		Setup.Create(GetTransientPackage());
+
+		// --- Before start: clock reads 0, state is Stopped ---
+		UTEST_EQUAL("Initial clock state is Stopped", Setup.Clock->GetState(), EMusicClockState::Stopped);
+		UTEST_NEARLY_EQUAL("Initial clock time is 0", Setup.Clock->GetSecondsIncludingCountIn(), 0.f, KINDA_SMALL_NUMBER);
+
+		// --- Start and advance ---
+		Setup.Start();
+		Setup.Tick(1.0); // Advance 1 second
+
+		UTEST_EQUAL("Clock is Running", Setup.Clock->GetState(), EMusicClockState::Running);
+		float TimeAfterOneSecond = Setup.Clock->GetSecondsIncludingCountIn();
+		UTEST_GREATER("Clock time > 0 after advancing", TimeAfterOneSecond, 0.f);
+
+		// --- Pause: time freezes ---
+		Setup.Pause();
+		float TimeAtPause = Setup.Clock->GetSecondsIncludingCountIn();
+		UTEST_EQUAL("Clock is Paused", Setup.Clock->GetState(), EMusicClockState::Paused);
+		UTEST_GREATER("Time at pause > 0", TimeAtPause, 0.f);
+
+		// Tick again while paused — time should not advance
+		Setup.TickNoAdvance();
+		UTEST_NEARLY_EQUAL("Time unchanged while paused", Setup.Clock->GetSecondsIncludingCountIn(), TimeAtPause, 0.001f);
+		UTEST_NEARLY_EQUAL("Delta beat is 0 while paused", Setup.Clock->GetDeltaBeat(), 0.f, KINDA_SMALL_NUMBER);
+
+		// --- Continue: resumes from paused time ---
+		Setup.Continue();
+		Setup.Tick(0.5);
+		UTEST_EQUAL("Clock is Running after continue", Setup.Clock->GetState(), EMusicClockState::Running);
+		UTEST_GREATER("Time advanced past pause point", Setup.Clock->GetSecondsIncludingCountIn(), TimeAtPause);
+
+		// --- Stop: resets to 0 ---
+		Setup.Stop();
+		UTEST_EQUAL("Clock is Stopped", Setup.Clock->GetState(), EMusicClockState::Stopped);
+		UTEST_NEARLY_EQUAL("Time is 0 after stop", Setup.Clock->GetSecondsIncludingCountIn(), 0.f, KINDA_SMALL_NUMBER);
+
+		// --- Restart: begins from 0 again ---
+		Setup.Start();
+		Setup.Tick(0.5);
+		UTEST_EQUAL("Clock is Running after restart", Setup.Clock->GetState(), EMusicClockState::Running);
+		UTEST_GREATER("Time advancing from 0 after restart", Setup.Clock->GetSecondsIncludingCountIn(), 0.f);
+		UTEST_LESS("Time is less than before stop (restarted fresh)", Setup.Clock->GetSecondsIncludingCountIn(), TimeAfterOneSecond);
+
+		return true;
+	}
+
+	// =========================================================================
+	// Test 4: UMusicClock.ReadOnlyView
+	// =========================================================================
+	IMPLEMENT_SIMPLE_AUTOMATION_TEST(
+		FMusicClockReadOnlyViewTest,
+		"Harmonix.Metasound.MusicClock.UMusicClock.ReadOnlyView",
+		EAutomationTestFlags::EditorContext | EAutomationTestFlags::EngineFilter)
+	bool FMusicClockReadOnlyViewTest::RunTest(const FString&)
+	{
+		FTestClockSetup Setup;
+		Setup.Create(GetTransientPackage());
+		Setup.Start();
+
+		// Advance 2 seconds at 120bpm 4/4 = 4 beats = 1 bar
+		Setup.Tick(2.0);
+
+		UMusicClock* Clock = Setup.Clock;
+
+		UTEST_EQUAL("Clock is Running", Clock->GetState(), EMusicClockState::Running);
+		UTEST_GREATER("Seconds > 0", Clock->GetSecondsIncludingCountIn(), 0.f);
+		UTEST_GREATER("Bars > 0", Clock->GetBarsIncludingCountIn(), 0.f);
+		UTEST_GREATER("Beats > 0", Clock->GetBeatsIncludingCountIn(), 0.f);
+		UTEST_NEARLY_EQUAL("Tempo is ~120", Clock->GetCurrentTempo(), 120.f, 1.f);
+		UTEST_NEARLY_EQUAL("BPM is ~120", Clock->GetCurrentBeatsPerMinute(), 120.f, 1.f);
+		UTEST_NEARLY_EQUAL("Beats/sec is ~2", Clock->GetCurrentBeatsPerSecond(), 2.f, 0.1f);
+		UTEST_GREATER("DeltaBeat > 0", Clock->GetDeltaBeat(), 0.f);
+		UTEST_GREATER("DeltaBar > 0", Clock->GetDeltaBar(), 0.f);
+
+		int Num = 0, Denom = 0;
+		Clock->GetCurrentTimeSignature(Num, Denom);
+		UTEST_EQUAL("Time sig num is 4", Num, 4);
+		UTEST_EQUAL("Time sig denom is 4", Denom, 4);
+
+		FMusicTimestamp Timestamp = Clock->GetCurrentTimestamp();
+		UTEST_GREATER_EQUAL("Timestamp bar >= 1", Timestamp.Bar, 1);
+
+		const ISongMapEvaluator* Maps = Clock->GetSongMaps();
+		UTEST_NOT_NULL("Song maps available", Maps);
+
+		return true;
+	}
+
+	// =========================================================================
+	// Test 5: UMusicClock.StoppedReturnsZero
+	// Explicit test that a stopped clock always returns 0
+	// =========================================================================
+	IMPLEMENT_SIMPLE_AUTOMATION_TEST(
+		FMusicClockStoppedZeroTest,
+		"Harmonix.Metasound.MusicClock.UMusicClock.StoppedReturnsZero",
+		EAutomationTestFlags::EditorContext | EAutomationTestFlags::EngineFilter)
+	bool FMusicClockStoppedZeroTest::RunTest(const FString&)
+	{
+		FTestClockSetup Setup;
+		Setup.Create(GetTransientPackage());
+
+		// Before ever starting
+		UTEST_NEARLY_EQUAL("Seconds is 0 before start", Setup.Clock->GetSecondsIncludingCountIn(), 0.f, KINDA_SMALL_NUMBER);
+		UTEST_NEARLY_EQUAL("Bars is 0 before start", Setup.Clock->GetBarsIncludingCountIn(), 0.f, KINDA_SMALL_NUMBER);
+		UTEST_NEARLY_EQUAL("Beats is 0 before start", Setup.Clock->GetBeatsIncludingCountIn(), 0.f, KINDA_SMALL_NUMBER);
+
+		// Run for a bit then stop
+		Setup.Start();
+		Setup.Tick(2.0);
+		UTEST_GREATER("Time advanced while running", Setup.Clock->GetSecondsIncludingCountIn(), 0.f);
+
+		Setup.Stop();
+		UTEST_NEARLY_EQUAL("Seconds is 0 after stop", Setup.Clock->GetSecondsIncludingCountIn(), 0.f, KINDA_SMALL_NUMBER);
+		UTEST_NEARLY_EQUAL("Bars is 0 after stop", Setup.Clock->GetBarsIncludingCountIn(), 0.f, KINDA_SMALL_NUMBER);
+		UTEST_NEARLY_EQUAL("Beats is 0 after stop", Setup.Clock->GetBeatsIncludingCountIn(), 0.f, KINDA_SMALL_NUMBER);
+		UTEST_NEARLY_EQUAL("Tempo is 0 after stop", Setup.Clock->GetCurrentTempo(), 0.f, KINDA_SMALL_NUMBER);
+		UTEST_NEARLY_EQUAL("DeltaBeat is 0 after stop", Setup.Clock->GetDeltaBeat(), 0.f, KINDA_SMALL_NUMBER);
+
+		return true;
+	}
+
+	// =========================================================================
+	// Test 6: OffsetMusicSource.PositiveOffset
+	// =========================================================================
+	IMPLEMENT_SIMPLE_AUTOMATION_TEST(
+		FOffsetClockPositiveTest,
+		"Harmonix.Metasound.MusicClock.OffsetMusicSource.PositiveOffset",
+		EAutomationTestFlags::EditorContext | EAutomationTestFlags::EngineFilter)
+	bool FOffsetClockPositiveTest::RunTest(const FString&)
+	{
+		FTestClockSetup Setup;
+		Setup.Create(GetTransientPackage());
+
+		// Offset source reads directly from the source, not from a clock
+		UOffsetMusicSource* OffsetSource = NewObject<UOffsetMusicSource>(GetTransientPackage());
+		OffsetSource->SetParentSource(TScriptInterface<IMusicSource>(Setup.Source));
+		OffsetSource->SetOffsetMs(500.f);
+		UMusicClock* OffsetClock = NewObject<UMusicClock>(GetTransientPackage());
+		OffsetClock->SetSource(TScriptInterface<IMusicSource>(OffsetSource));
+
+		Setup.Start();
+
+		for (int32 i = 0; i < 10; ++i)
+		{
+			Setup.Tick(0.2);
+			OffsetSource->Update();
+			++GFrameCounter;
+			OffsetClock->UpdateForGameFrame();
+		}
+
+		float BaseSeconds = Setup.Clock->GetSecondsIncludingCountIn();
+		float OffsetSeconds = OffsetClock->GetSecondsIncludingCountIn();
+
+		UTEST_GREATER("Offset clock is ahead of base", OffsetSeconds, BaseSeconds);
+		float DiffMs = (OffsetSeconds - BaseSeconds) * 1000.f;
+		UTEST_NEARLY_EQUAL("Offset is approximately 500ms", DiffMs, 500.f, 50.f);
+
+		return true;
+	}
+
+	// =========================================================================
+	// Test 7: OffsetMusicSource.NegativeOffset
+	// =========================================================================
+	IMPLEMENT_SIMPLE_AUTOMATION_TEST(
+		FOffsetClockNegativeTest,
+		"Harmonix.Metasound.MusicClock.OffsetMusicSource.NegativeOffset",
+		EAutomationTestFlags::EditorContext | EAutomationTestFlags::EngineFilter)
+	bool FOffsetClockNegativeTest::RunTest(const FString&)
+	{
+		FTestClockSetup Setup;
+		Setup.Create(GetTransientPackage());
+
+		UOffsetMusicSource* OffsetSource = NewObject<UOffsetMusicSource>(GetTransientPackage());
+		OffsetSource->SetParentSource(TScriptInterface<IMusicSource>(Setup.Source));
+		OffsetSource->SetOffsetMs(-200.f);
+		UMusicClock* OffsetClock = NewObject<UMusicClock>(GetTransientPackage());
+		OffsetClock->SetSource(TScriptInterface<IMusicSource>(OffsetSource));
+
+		Setup.Start();
+
+		for (int32 i = 0; i < 10; ++i)
+		{
+			Setup.Tick(0.2);
+			OffsetSource->Update();
+			++GFrameCounter;
+			OffsetClock->UpdateForGameFrame();
+		}
+
+		float BaseSeconds = Setup.Clock->GetSecondsIncludingCountIn();
+		float OffsetSeconds = OffsetClock->GetSecondsIncludingCountIn();
+
+		UTEST_LESS("Offset clock is behind base", OffsetSeconds, BaseSeconds);
+		float DiffMs = (BaseSeconds - OffsetSeconds) * 1000.f;
+		UTEST_NEARLY_EQUAL("Offset is approximately 200ms", DiffMs, 200.f, 50.f);
+
+		return true;
+	}
+
+	// =========================================================================
+	// Test 8: OffsetMusicSource.Chain
+	// Offsets chain directly: Source -> OffsetA (+100ms) -> OffsetB (+50ms)
+	// No intermediate clocks needed in the chain — offset sources read from
+	// IMusicSource, and each offset source is itself an IMusicSource.
+	// =========================================================================
+	IMPLEMENT_SIMPLE_AUTOMATION_TEST(
+		FOffsetClockChainTest,
+		"Harmonix.Metasound.MusicClock.OffsetMusicSource.Chain",
+		EAutomationTestFlags::EditorContext | EAutomationTestFlags::EngineFilter)
+	bool FOffsetClockChainTest::RunTest(const FString&)
+	{
+		FTestClockSetup Setup;
+		Setup.Create(GetTransientPackage());
+
+		// Chain: Source -> OffsetA (+100ms) -> OffsetB (+50ms) -> ClockB
+		UOffsetMusicSource* OffsetA = NewObject<UOffsetMusicSource>(GetTransientPackage());
+		OffsetA->SetParentSource(TScriptInterface<IMusicSource>(Setup.Source));
+		OffsetA->SetOffsetMs(100.f);
+
+		UOffsetMusicSource* OffsetB = NewObject<UOffsetMusicSource>(GetTransientPackage());
+		OffsetB->SetParentSource(TScriptInterface<IMusicSource>(OffsetA));
+		OffsetB->SetOffsetMs(50.f);
+
+		UMusicClock* ChainClock = NewObject<UMusicClock>(GetTransientPackage());
+		ChainClock->SetSource(TScriptInterface<IMusicSource>(OffsetB));
+
+		Setup.Start();
+
+		for (int32 i = 0; i < 10; ++i)
+		{
+			Setup.Tick(0.2);
+			OffsetA->Update();
+			OffsetB->Update();
+			++GFrameCounter;
+			ChainClock->UpdateForGameFrame();
+		}
+
+		float BaseSeconds = Setup.Clock->GetSecondsIncludingCountIn();
+		float ChainSeconds = ChainClock->GetSecondsIncludingCountIn();
+
+		float TotalDiffMs = (ChainSeconds - BaseSeconds) * 1000.f;
+		UTEST_NEARLY_EQUAL("Chained offset is approximately 150ms", TotalDiffMs, 150.f, 30.f);
+
+		return true;
+	}
+
+	// =========================================================================
+	// Test 9: OffsetMusicSource.LoopBoundaryWrap
+	// Validates that a negative offset does NOT wrap on the first pass through
+	// a loop, but DOES wrap correctly after the parent source loops.
+	//
+	// Setup: 2-bar loop at 120bpm 4/4 = 4000ms (0 to 4000ms), offset -100ms
+	//   First pass, parent at 50ms  -> Offset = -50ms  -> stays negative (no wrap)
+	//   First pass, parent at 200ms -> Offset = 100ms  -> normal (in bounds)
+	//   After loop, parent at 50ms  -> Offset = -50ms  -> wraps to 3950ms
+	// =========================================================================
+	IMPLEMENT_SIMPLE_AUTOMATION_TEST(
+		FOffsetClockLoopBoundaryTest,
+		"Harmonix.Metasound.MusicClock.OffsetMusicSource.LoopBoundaryWrap",
+		EAutomationTestFlags::EditorContext | EAutomationTestFlags::EngineFilter)
+	bool FOffsetClockLoopBoundaryTest::RunTest(const FString&)
+	{
+		// At 120bpm 4/4: 1 bar = 2 seconds. 2 bars = 4 seconds = 4000ms.
+		TSharedRef<FManualTimeSource> TimeSource = MakeShared<FManualTimeSource>();
+		URuntimeMusicSource* Source = NewObject<URuntimeMusicSource>(GetTransientPackage());
+		Source->Initialize(TimeSource, nullptr); // Default 120bpm 4/4
+		Source->SetLoopRegionByBars(1, 2); // Loop from bar 1 for 2 bars
+
+		UTEST_EQUAL("Source is looping", Source->IsLooping(), true);
+		float LoopLenMs = Source->GetLoopLengthMs();
+		UTEST_NEARLY_EQUAL("Loop length is ~4000ms", LoopLenMs, 4000.f, 100.f);
+
+		// Base clock for reading the source's position in assertions
+		UMusicClock* BaseClock = NewObject<UMusicClock>(GetTransientPackage());
+		BaseClock->SetSource(TScriptInterface<IMusicSource>(Source));
+
+		// Offset source reads directly from the source
+		UOffsetMusicSource* OffsetSource = NewObject<UOffsetMusicSource>(GetTransientPackage());
+		OffsetSource->SetParentSource(TScriptInterface<IMusicSource>(Source));
+		OffsetSource->SetOffsetMs(-100.f); // 100ms behind
+		UMusicClock* OffsetClock = NewObject<UMusicClock>(GetTransientPackage());
+		OffsetClock->SetSource(TScriptInterface<IMusicSource>(OffsetSource));
+
+		TFunction<void(double)> TickAll = [&](double DeltaSec)
+		{
+			TimeSource->AdvanceBy(DeltaSec);
+			Source->Update();
+			++GFrameCounter;
+			BaseClock->UpdateForGameFrame();
+			OffsetSource->Update();
+			++GFrameCounter;
+			OffsetClock->UpdateForGameFrame();
+		};
+
+		// Start
+		TimeSource->Start();
+		TickAll(0.0); // Process start
+
+		// --- First pass: parent at ~50ms, offset = -50ms. Should NOT wrap. ---
+		TickAll(0.05);
+		float OffsetMs1 = OffsetClock->GetSecondsIncludingCountIn() * 1000.f;
+		UTEST_LESS("First pass: negative offset stays before loop start", OffsetMs1, 100.f);
+
+		// --- Advance past the offset threshold so offset time goes positive ---
+		TickAll(0.15); // Parent at ~200ms, offset at ~100ms
+		float OffsetMs1b = OffsetClock->GetSecondsIncludingCountIn() * 1000.f;
+		UTEST_GREATER("First pass: offset advances into loop region", OffsetMs1b, 0.f);
+
+		// --- Advance to mid-loop ---
+		TickAll(1.8); // Parent at ~2000ms
+		float BaseMs2 = BaseClock->GetSecondsIncludingCountIn() * 1000.f;
+		float OffsetMs2 = OffsetClock->GetSecondsIncludingCountIn() * 1000.f;
+		float Diff2 = BaseMs2 - OffsetMs2;
+		UTEST_NEARLY_EQUAL("Mid-loop: offset is ~100ms behind", Diff2, 100.f, 30.f);
+
+		// --- Cross the loop boundary ---
+		TickAll(2.1); // Parent past 4000ms, loops to ~100ms
+		float OffsetMs3 = OffsetClock->GetSecondsIncludingCountIn() * 1000.f;
+		// After the parent loops, the offset clock should wrap (parent near start
+		// with negative offset → wraps to near loop end)
+		UTEST_GREATER("After loop: offset clock wrapped to near end of loop", OffsetMs3, 3800.f);
+		// Loop should have been detected
+		UTEST_EQUAL("Offset clock detected loop", OffsetClock->LoopedThisFrame(), true);
+
+		return true;
+	}
+
+	// =========================================================================
+	// Test 10: UMidiClockUpdateSubsystem.TrackNewTypes
+	// =========================================================================
+	IMPLEMENT_SIMPLE_AUTOMATION_TEST(
+		FSubsystemTrackNewTypesTest,
+		"Harmonix.Metasound.MusicClock.Subsystem.TrackNewTypes",
+		EAutomationTestFlags::EditorContext | EAutomationTestFlags::EngineFilter)
+	bool FSubsystemTrackNewTypesTest::RunTest(const FString&)
+	{
+		UTEST_NOT_NULL("GEngine exists", GEngine);
+		UMidiClockUpdateSubsystem* Subsystem = GEngine->GetEngineSubsystem<UMidiClockUpdateSubsystem>();
+		UTEST_NOT_NULL("Subsystem exists", Subsystem);
+
+		TSharedRef<FManualTimeSource> TimeSource = MakeShared<FManualTimeSource>();
+		URuntimeMusicSource* Source = NewObject<URuntimeMusicSource>();
+		Source->Initialize(TimeSource, nullptr);
+
+		UMusicClock* Clock = NewObject<UMusicClock>();
+		Clock->SetSource(TScriptInterface<IMusicSource>(Source));
+
+		UMidiClockUpdateSubsystem::TrackMusicSource(Source);
+		UMidiClockUpdateSubsystem::TrackMusicClock(Clock);
+
+		TimeSource->Start();
+		TimeSource->AdvanceBy(1.0);
+
+		// Tick subsystem — should update source and clock
+		Subsystem->TickForTesting();
+
+		UTEST_GREATER("Source position advanced", Source->GetCurrentSongPos().SecondsIncludingCountIn, 0.f);
+		UTEST_GREATER("Clock position advanced", Clock->GetSecondsIncludingCountIn(), 0.f);
+
+		// Unregister
+		UMidiClockUpdateSubsystem::StopTrackingMusicClock(Clock);
+		UMidiClockUpdateSubsystem::StopTrackingMusicSource(Source);
+
+		float PosBeforeSecondTick = Clock->GetSecondsIncludingCountIn();
+		TimeSource->AdvanceBy(1.0);
+		Subsystem->TickForTesting();
+
+		UTEST_NEARLY_EQUAL("Clock not updated after unregister", Clock->GetSecondsIncludingCountIn(), PosBeforeSecondTick, 0.01f);
+
+		return true;
+	}
+}
+
+#endif
